UMG - More cleanup, introducing a new UContentWidget base class for border, buttons...etc.  Now able to embed text element in buttons correctly.

diff --git a/Engine/Source/Editor/UMGEditor/Private/SUMGEditorTree.cpp b/Engine/Source/Editor/UMGEditor/Private/SUMGEditorTree.cpp
--- a/Engine/Source/Editor/UMGEditor/Private/SUMGEditorTree.cpp
+++ b/Engine/Source/Editor/UMGEditor/Private/SUMGEditorTree.cpp
@@ -163,12 +163,19 @@ FReply SUMGEditorTree::CreateTestUI()
 	Button2->ButtonText = FText::FromString("Button 2");
 	UButtonComponent* Button3 = ConstructObject<UButtonComponent>(UButtonComponent::StaticClass(), BP);
 	Button3->ButtonText = FText::FromString("Button 3");
+	UButtonComponent* Button4 = ConstructObject<UButtonComponent>(UButtonComponent::StaticClass(), BP);
+
+	UTextBlockComponent* Text1 = ConstructObject<UTextBlockComponent>(UTextBlockComponent::StaticClass(), BP);
+	Text1->Text = FText::FromString(TEXT("Button!"));
+	Button4->SetContent(Text1);
 
 	BP->WidgetTemplates.Add(Canvas);
 	BP->WidgetTemplates.Add(Vertical);
 	BP->WidgetTemplates.Add(Button1);
 	BP->WidgetTemplates.Add(Button2);
 	BP->WidgetTemplates.Add(Button3);
+	BP->WidgetTemplates.Add(Button4);
+	BP->WidgetTemplates.Add(Text1);
 
 	UCanvasPanelSlot* Slot = Canvas->AddSlot(Vertical);
 	Slot->Size.X = 100;
@@ -179,6 +186,7 @@ FReply SUMGEditorTree::CreateTestUI()
 	Vertical->AddSlot(Button1);
 	Vertical->AddSlot(Button2);
 	Vertical->AddSlot(Button3);
+	Vertical->AddSlot(Button4);
 
 	FBlueprintEditorUtils::MarkBlueprintAsStructurallyModified(BP);
 
diff --git a/Engine/Source/Runtime/UMG/Classes/BorderComponent.h b/Engine/Source/Runtime/UMG/Classes/BorderComponent.h
--- a/Engine/Source/Runtime/UMG/Classes/BorderComponent.h
+++ b/Engine/Source/Runtime/UMG/Classes/BorderComponent.h
@@ -7,7 +7,7 @@
 /** A border is a container widget that can contain one child widget, providing an opportunity to surround it with a border image and adjustable padding. */
 
 UCLASS(meta=(BlueprintSpawnableComponent), ClassGroup=UserInterface)
-class UMG_API UBorderComponent : public USlateNonLeafWidgetComponent
+class UMG_API UBorderComponent : public UContentWidget
 {
 	GENERATED_UCLASS_BODY()
 
@@ -57,6 +57,10 @@ class UMG_API UBorderComponent : public USlateNonLeafWidgetComponent
 	UPROPERTY(EditAnywhere, BlueprintReadOnly, Category=Appearance)
 	bool bShowEffectWhenDisabled;
 
+	// UContentWidget interface
+	virtual void SetContent(USlateWrapperComponent* Content) OVERRIDE;
+	// End UContentWidget interface
+
 protected:
 	TWeakPtr<class SBorder> MyBorder;
 
@@ -65,10 +69,6 @@ protected:
 	virtual TSharedRef<SWidget> RebuildWidget() OVERRIDE;
 	// End of USlateWrapperComponent interface
 
-	// USlateNonLeafWidgetComponent interface
-	virtual void OnKnownChildrenChanged() OVERRIDE;
-	// End of USlateNonLeafWidgetComponent
-
 	FMargin GetContentPadding() const;
 	FLinearColor GetContentColor() const;
 	FSlateColor GetBorderColor() const;
diff --git a/Engine/Source/Runtime/UMG/Classes/BoxPanelComponent.h b/Engine/Source/Runtime/UMG/Classes/BoxPanelComponent.h
--- a/Engine/Source/Runtime/UMG/Classes/BoxPanelComponent.h
+++ b/Engine/Source/Runtime/UMG/Classes/BoxPanelComponent.h
@@ -61,12 +61,6 @@ class UMG_API UBoxPanelComponent : public USlateNonLeafWidgetComponent
 	// End of UObject interface
 #endif
 
-	// USceneComponent interface
-	virtual bool HasAnySockets() const;
-	virtual FTransform GetSocketTransform(FName InSocketName, ERelativeTransformSpace TransformSpace = RTS_World) const;
-	virtual void QuerySupportedSockets(TArray<FComponentSocketDescription>& OutSockets) const;
-	// End of USceneComponent interface
-
 protected:
 	UPROPERTY(EditAnywhere, Category=Slots)
 	TArray<FBoxPanelSlot> Slots;
@@ -80,8 +74,4 @@ protected:
 	// USlateWrapperComponent interface
 	virtual TSharedRef<SWidget> RebuildWidget() OVERRIDE;
 	// End of USlateWrapperComponent interface
-
-	// USlateNonLeafWidgetComponent interface
-	virtual void OnKnownChildrenChanged() OVERRIDE;
-	// End of USlateNonLeafWidgetComponent
 };
diff --git a/Engine/Source/Runtime/UMG/Classes/ButtonComponent.h b/Engine/Source/Runtime/UMG/Classes/ButtonComponent.h
--- a/Engine/Source/Runtime/UMG/Classes/ButtonComponent.h
+++ b/Engine/Source/Runtime/UMG/Classes/ButtonComponent.h
@@ -10,7 +10,7 @@ DECLARE_DYNAMIC_DELEGATE_RetVal( FText, FOnGetText );
 
 /** Buttons are clickable widgets */
 UCLASS(meta=(BlueprintSpawnableComponent), ClassGroup=UserInterface)
-class UMG_API UButtonComponent : public USlateNonLeafWidgetComponent
+class UMG_API UButtonComponent : public UContentWidget
 {
 	GENERATED_UCLASS_BODY()
 
@@ -71,6 +71,10 @@ public:
 	UFUNCTION(BlueprintNativeEvent, Category=Appearance)
 	FText GetText() const;
 
+	// UContentWidget interface
+	virtual void SetContent(USlateWrapperComponent* Content) OVERRIDE;
+	// End UContentWidget interface
+
 protected:
 	TAttribute<FText> TextAttribute;
 	TWeakPtr<class SButton> MyButton;
@@ -80,10 +84,6 @@ protected:
 	virtual TSharedRef<SWidget> RebuildWidget() OVERRIDE;
 	// End of USlateWrapperComponent interface
 
-	// USlateNonLeafWidgetComponent interface
-	virtual void OnKnownChildrenChanged() OVERRIDE;
-	// End of USlateNonLeafWidgetComponent
-
 	FMargin GetContentPadding() const;
 	FText GetButtonText() const;
 	FSlateColor GetButtonColor() const;
diff --git a/Engine/Source/Runtime/UMG/Classes/CanvasPanelComponent.h b/Engine/Source/Runtime/UMG/Classes/CanvasPanelComponent.h
--- a/Engine/Source/Runtime/UMG/Classes/CanvasPanelComponent.h
+++ b/Engine/Source/Runtime/UMG/Classes/CanvasPanelComponent.h
@@ -42,8 +42,4 @@ protected:
 	// USlateWrapperComponent interface
 	virtual TSharedRef<SWidget> RebuildWidget() OVERRIDE;
 	// End of USlateWrapperComponent interface
-
-	// USlateNonLeafWidgetComponent interface
-	virtual void OnKnownChildrenChanged() OVERRIDE;
-	// End of USlateNonLeafWidgetComponent
 };
diff --git a/Engine/Source/Runtime/UMG/Classes/ContentWidget.h b/Engine/Source/Runtime/UMG/Classes/ContentWidget.h
new file mode 100644
--- /dev/null
+++ b/Engine/Source/Runtime/UMG/Classes/ContentWidget.h
@@ -0,0 +1,25 @@
+// Copyright 1998-2014 Epic Games, Inc. All Rights Reserved.
+
+#pragma once
+
+#include "ContentWidget.generated.h"
+
+UCLASS(Abstract)
+class UMG_API UContentWidget : public USlateNonLeafWidgetComponent
+{
+	GENERATED_UCLASS_BODY()
+
+	UFUNCTION()
+	virtual USlateWrapperComponent* GetContent();
+	
+	UFUNCTION()
+	virtual void SetContent(USlateWrapperComponent* InContent);
+
+	virtual int32 GetChildrenCount() const OVERRIDE;
+	virtual USlateWrapperComponent* GetChildAt(int32 Index) const OVERRIDE;
+
+protected:
+
+	UPROPERTY()
+	USlateWrapperComponent* Content;
+};
diff --git a/Engine/Source/Runtime/UMG/Classes/GridPanelComponent.h b/Engine/Source/Runtime/UMG/Classes/GridPanelComponent.h
--- a/Engine/Source/Runtime/UMG/Classes/GridPanelComponent.h
+++ b/Engine/Source/Runtime/UMG/Classes/GridPanelComponent.h
@@ -72,12 +72,6 @@ class UMG_API UGridPanelComponent : public USlateNonLeafWidgetComponent
 	// End of UObject interface
 #endif
 
-	// USceneComponent interface
-	virtual bool HasAnySockets() const;
-	virtual FTransform GetSocketTransform(FName InSocketName, ERelativeTransformSpace TransformSpace = RTS_World) const;
-	virtual void QuerySupportedSockets(TArray<FComponentSocketDescription>& OutSockets) const;
-	// End of USceneComponent interface
-
 protected:
 	UPROPERTY(EditAnywhere, Category=Slots)
 	TArray<FGridPanelSlot> Slots;
@@ -91,10 +85,6 @@ protected:
 	virtual TSharedRef<SWidget> RebuildWidget() OVERRIDE;
 	// End of USlateWrapperComponent interface
 
-	// USlateNonLeafWidgetComponent interface
-	virtual void OnKnownChildrenChanged() OVERRIDE;
-	// End of USlateNonLeafWidgetComponent
-
 	void RebuildUniqueSlotTable();
 	void EnsureSlotIsUnique(FGridPanelSlot& SlotConfig);
 };
diff --git a/Engine/Source/Runtime/UMG/Classes/HorizontalBoxComponent.h b/Engine/Source/Runtime/UMG/Classes/HorizontalBoxComponent.h
--- a/Engine/Source/Runtime/UMG/Classes/HorizontalBoxComponent.h
+++ b/Engine/Source/Runtime/UMG/Classes/HorizontalBoxComponent.h
@@ -38,8 +38,4 @@ protected:
 	// USlateWrapperComponent interface
 	virtual TSharedRef<SWidget> RebuildWidget() OVERRIDE;
 	// End of USlateWrapperComponent interface
-
-	// USlateNonLeafWidgetComponent interface
-	virtual void OnKnownChildrenChanged() OVERRIDE;
-	// End of USlateNonLeafWidgetComponent
 };
diff --git a/Engine/Source/Runtime/UMG/Classes/SlateNonLeafWidgetComponent.h b/Engine/Source/Runtime/UMG/Classes/SlateNonLeafWidgetComponent.h
--- a/Engine/Source/Runtime/UMG/Classes/SlateNonLeafWidgetComponent.h
+++ b/Engine/Source/Runtime/UMG/Classes/SlateNonLeafWidgetComponent.h
@@ -9,33 +9,6 @@ class UMG_API USlateNonLeafWidgetComponent : public USlateWrapperComponent
 {
 	GENERATED_UCLASS_BODY()
 
-	// UActorComponent interface
-	virtual void TickComponent(float DeltaTime, enum ELevelTick TickType, FActorComponentTickFunction* ThisTickFunction) OVERRIDE;
-	// End of UActorComponent interface
-
-#if WITH_EDITOR
-	// UActorComponent interface
-	virtual void CheckForErrors() OVERRIDE;
-	// End of UActorComponent interface
-#endif
-
 	virtual int32 GetChildrenCount() const { return 0; }
 	virtual USlateWrapperComponent* GetChildAt(int32 Index) const { return NULL; }
-
-protected:
-	// This function is called when the known children list has changed; use this to repopulate child slots in derived classes
-	virtual void OnKnownChildrenChanged() {}
-
-	// Helper to get the first child (for bCanHaveMultipleChildren = false subclasses)
-	USlateWrapperComponent* GetFirstWrappedChild() const;
-
-	// Always returns something valid, even if there are no children
-	TSharedRef<SWidget> GetFirstChildWidget() const;
-
-protected:
-	// The list of known children
-	TSet<TWeakObjectPtr<USlateWrapperComponent>> KnownChildren;
-
-	// Can we have more than one child?
-	bool bCanHaveMultipleChildren;
 };
diff --git a/Engine/Source/Runtime/UMG/Classes/SpacerComponent.h b/Engine/Source/Runtime/UMG/Classes/SpacerComponent.h
--- a/Engine/Source/Runtime/UMG/Classes/SpacerComponent.h
+++ b/Engine/Source/Runtime/UMG/Classes/SpacerComponent.h
@@ -6,7 +6,7 @@
 
 /** A spacer widget; it does not have a visual representation, and just provides padding between other widgets. */
 UCLASS(meta=(BlueprintSpawnableComponent), ClassGroup=UserInterface)
-class UMG_API USpacerComponent : public USlateWrapperComponent
+class UMG_API USpacerComponent : public USlateLeafWidgetComponent
 {
 	GENERATED_UCLASS_BODY()
 
diff --git a/Engine/Source/Runtime/UMG/Classes/TextBlockComponent.h b/Engine/Source/Runtime/UMG/Classes/TextBlockComponent.h
--- a/Engine/Source/Runtime/UMG/Classes/TextBlockComponent.h
+++ b/Engine/Source/Runtime/UMG/Classes/TextBlockComponent.h
@@ -9,9 +9,15 @@ class UMG_API UTextBlockComponent : public USlateLeafWidgetComponent
 {
 	GENERATED_UCLASS_BODY()
 
-protected:
+public:
+	UPROPERTY(EditAnywhere, Category=TextBlock)
+	FText Text;
+
+	UPROPERTY(EditAnywhere, Category=TextBlock)
+	FSlateFontInfo Font;
+
 	UPROPERTY(EditAnywhere, Category=TextBlock)
-	FText DisplayText;
+	FVector2D ShadowOffset;
 
 	//@TODO: Expose FSlateColor
 	UPROPERTY(EditAnywhere, Category=TextBlock)
@@ -25,7 +31,7 @@ protected:
 	virtual TSharedRef<SWidget> RebuildWidget() OVERRIDE;
 	// End of USlateWrapperComponent interface
 
-	FString GetDisplayText() const;
+	FString GetText() const;
 	FSlateColor GetColorAndOpacity() const;
 	FLinearColor GetShadowColorAndOpacity() const;
 };
diff --git a/Engine/Source/Runtime/UMG/Classes/VerticalBoxComponent.h b/Engine/Source/Runtime/UMG/Classes/VerticalBoxComponent.h
--- a/Engine/Source/Runtime/UMG/Classes/VerticalBoxComponent.h
+++ b/Engine/Source/Runtime/UMG/Classes/VerticalBoxComponent.h
@@ -38,8 +38,4 @@ protected:
 	// USlateWrapperComponent interface
 	virtual TSharedRef<SWidget> RebuildWidget() OVERRIDE;
 	// End of USlateWrapperComponent interface
-
-	// USlateNonLeafWidgetComponent interface
-	virtual void OnKnownChildrenChanged() OVERRIDE;
-	// End of USlateNonLeafWidgetComponent
 };
diff --git a/Engine/Source/Runtime/UMG/Private/BorderComponent.cpp b/Engine/Source/Runtime/UMG/Private/BorderComponent.cpp
--- a/Engine/Source/Runtime/UMG/Private/BorderComponent.cpp
+++ b/Engine/Source/Runtime/UMG/Private/BorderComponent.cpp
@@ -10,8 +10,6 @@
 UBorderComponent::UBorderComponent(const FPostConstructInitializeProperties& PCIP)
 	: Super(PCIP)
 {
-	bCanHaveMultipleChildren = false;
-
 	SBorder::FArguments BorderDefaults;
 
 	HorizontalAlignment = BorderDefaults._HAlign;
@@ -48,12 +46,14 @@ TSharedRef<SWidget> UBorderComponent::RebuildWidget()
 	return NewBorder;
 }
 
-void UBorderComponent::OnKnownChildrenChanged()
+void UBorderComponent::SetContent(USlateWrapperComponent* Content)
 {
+	Super::SetContent(Content);
+
 	TSharedPtr<SBorder> Border = MyBorder.Pin();
-	if (Border.IsValid())
+	if ( Border.IsValid() )
 	{
-		Border->SetContent(GetFirstChildWidget());
+		Border->SetContent(Content->GetWidget());
 	}
 }
 
diff --git a/Engine/Source/Runtime/UMG/Private/BoxPanelComponent.cpp b/Engine/Source/Runtime/UMG/Private/BoxPanelComponent.cpp
--- a/Engine/Source/Runtime/UMG/Private/BoxPanelComponent.cpp
+++ b/Engine/Source/Runtime/UMG/Private/BoxPanelComponent.cpp
@@ -117,79 +117,6 @@ TSharedRef<SWidget> UBoxPanelComponent::RebuildWidget()
 	return Box;
 }
 
-void UBoxPanelComponent::OnKnownChildrenChanged()
-{
-	//TSharedPtr<SConcreteBoxPanel> Box = MyBox.Pin();
-	//if (Box.IsValid())
-	//{
-	//	Box->ClearChildren();
-
-	//	// Add slots
-	//	TMap<FName, int32> SlotNameToIndex;
-
-	//	for (int32 SlotIndex = 0; SlotIndex < Slots.Num(); ++SlotIndex)
-	//	{
-	//		const FBoxPanelSlot& SlotConfig = Slots[SlotIndex];
-
-	//		SlotNameToIndex.Add(SlotConfig.SlotName, SlotIndex);
-
-	//		FSizeParam SlotSize = USlateWrapperComponent::ConvertSerializedSizeParamToRuntime(SlotConfig.SlotSize);
-
-	//		auto NewSlot = Box->AddSlot()
-	//			.Padding(SlotConfig.Padding)
-	//			.HAlign(SlotConfig.HorizontalAlignment)
-	//			.VAlign(SlotConfig.VerticalAlignment)
-	//			.SetSizeParam(SlotSize)
-	//			.MaxWidth(SlotConfig.bHasMaximumSize ? SlotConfig.MaxSize : 0.0f)
-	//		[
-	//			SNullWidget::NullWidget
-	//		];
-	//	}
-
-	//	// Place widgets in their slots, and add anything left over to the end
-	//	for (int32 ComponentIndex = 0; ComponentIndex < AttachChildren.Num(); ++ComponentIndex)
-	//	{
-	//		if (USlateWrapperComponent* Wrapper = Cast<USlateWrapperComponent>(AttachChildren[ComponentIndex]))
-	//		{
-	//			if (Wrapper->IsRegistered())
-	//			{
-	//				if (int32* pSlotIndex = SlotNameToIndex.Find(Wrapper->AttachSocketName))
-	//				{
-	//					auto Foo = (TPanelChildren<SHorizontalBox::FSlot>*)(Box->GetChildren());
-	//					(*Foo)[*pSlotIndex].Widget = Wrapper->GetWidget();
-	//				}
-	//				else
-	//				{
-	//					Box->AddSlot()
-	//					[
-	//						Wrapper->GetWidget()
-	//					];
-	//				}
-	//			}
-	//		}
-	//	}
-	//}
-}
-
-bool UBoxPanelComponent::HasAnySockets() const
-{
-	return true;
-}
-
-FTransform UBoxPanelComponent::GetSocketTransform(FName InSocketName, ERelativeTransformSpace TransformSpace) const
-{
-	return FTransform::Identity;
-}
-
-void UBoxPanelComponent::QuerySupportedSockets(TArray<FComponentSocketDescription>& OutSockets) const
-{
-	for (int32 SlotIndex = 0; SlotIndex < Slots.Num(); ++SlotIndex)
-	{
-		FName SocketName = Slots[SlotIndex].SlotName;
-		new (OutSockets) FComponentSocketDescription(SocketName, EComponentSocketType::Socket);
-	}
-}
-
 #if WITH_EDITOR
 void UBoxPanelComponent::PostEditChangeProperty(struct FPropertyChangedEvent& PropertyChangedEvent)
 {
diff --git a/Engine/Source/Runtime/UMG/Private/ButtonComponent.cpp b/Engine/Source/Runtime/UMG/Private/ButtonComponent.cpp
--- a/Engine/Source/Runtime/UMG/Private/ButtonComponent.cpp
+++ b/Engine/Source/Runtime/UMG/Private/ButtonComponent.cpp
@@ -66,20 +66,26 @@ TSharedRef<SWidget> UButtonComponent::RebuildWidget()
 		.ContentScale(ContentScale)
 		.ButtonColorAndOpacity(BIND_UOBJECT_ATTRIBUTE(FSlateColor, GetButtonColor))
 		.ForegroundColor(BIND_UOBJECT_ATTRIBUTE(FSlateColor, GetForegroundColor));
+
+	if ( Content )
+	{
+		NewButton->SetContent(Content->GetWidget());
+	}
 	
 	MyButton = NewButton;
 	return NewButton;
 }
 
-void UButtonComponent::OnKnownChildrenChanged()
+void UButtonComponent::SetContent(USlateWrapperComponent* InContent)
 {
+	Super::SetContent(InContent);
+
 	TSharedPtr<SButton> Button = MyButton.Pin();
 	if (Button.IsValid())
 	{
-		USlateWrapperComponent* ChildComponent = GetFirstWrappedChild();
-		if (ChildComponent)
+		if ( InContent )
 		{
-			Button->SetContent(ChildComponent->GetWidget());
+			Button->SetContent(InContent->GetWidget());
 		}
 		else
 		{
diff --git a/Engine/Source/Runtime/UMG/Private/CanvasPanelComponent.cpp b/Engine/Source/Runtime/UMG/Private/CanvasPanelComponent.cpp
--- a/Engine/Source/Runtime/UMG/Private/CanvasPanelComponent.cpp
+++ b/Engine/Source/Runtime/UMG/Private/CanvasPanelComponent.cpp
@@ -55,19 +55,12 @@ TSharedRef<SWidget> UCanvasPanelComponent::RebuildWidget()
 	TSharedRef<SFixedSizeCanvas> NewCanvas = SNew(SFixedSizeCanvas, DesiredCanvasSize);
 	MyCanvas = NewCanvas;
 
-	OnKnownChildrenChanged();
-
-	return NewCanvas;
-}
-
-void UCanvasPanelComponent::OnKnownChildrenChanged()
-{
 	TSharedPtr<SFixedSizeCanvas> Canvas = MyCanvas.Pin();
-	if (Canvas.IsValid())
+	if ( Canvas.IsValid() )
 	{
 		Canvas->ClearChildren();
 
-		for (int32 SlotIndex = 0; SlotIndex < Slots.Num(); ++SlotIndex)
+		for ( int32 SlotIndex = 0; SlotIndex < Slots.Num(); ++SlotIndex )
 		{
 			UCanvasPanelSlot* Slot = Slots[SlotIndex];
 			if ( Slot == NULL )
@@ -80,11 +73,13 @@ void UCanvasPanelComponent::OnKnownChildrenChanged()
 				.Size(Slot->Size)
 				.HAlign(Slot->HorizontalAlignment)
 				.VAlign(Slot->VerticalAlignment)
-			[
-				Slot->Content == NULL ? SNullWidget::NullWidget : Slot->Content->GetWidget()
-			];
+				[
+					Slot->Content == NULL ? SNullWidget::NullWidget : Slot->Content->GetWidget()
+				];
 		}
 	}
+
+	return NewCanvas;
 }
 
 UCanvasPanelSlot* UCanvasPanelComponent::AddSlot(USlateWrapperComponent* Content)
diff --git a/Engine/Source/Runtime/UMG/Private/ContentWidget.cpp b/Engine/Source/Runtime/UMG/Private/ContentWidget.cpp
new file mode 100644
--- /dev/null
+++ b/Engine/Source/Runtime/UMG/Private/ContentWidget.cpp
@@ -0,0 +1,46 @@
+// Copyright 1998-2014 Epic Games, Inc. All Rights Reserved.
+
+#include "UMGPrivatePCH.h"
+
+#if WITH_EDITOR
+#include "MessageLog.h"
+#include "UObjectToken.h"
+#endif
+
+#define LOCTEXT_NAMESPACE "UMG"
+
+/////////////////////////////////////////////////////
+// UContentWidget
+
+UContentWidget::UContentWidget(const FPostConstructInitializeProperties& PCIP)
+	: Super(PCIP)
+{
+	PrimaryComponentTick.bCanEverTick = true;
+
+	bAutoActivate = true;
+	bTickInEditor = true;
+}
+
+USlateWrapperComponent* UContentWidget::GetContent()
+{
+	return Content;
+}
+
+void UContentWidget::SetContent(USlateWrapperComponent* InContent)
+{
+	Content = InContent;
+}
+
+int32 UContentWidget::GetChildrenCount() const
+{
+	return Content != NULL ? 1 : 0;
+}
+
+USlateWrapperComponent* UContentWidget::GetChildAt(int32 Index) const
+{
+	return Content;
+}
+
+/////////////////////////////////////////////////////
+
+#undef LOCTEXT_NAMESPACE
diff --git a/Engine/Source/Runtime/UMG/Private/GridPanelComponent.cpp b/Engine/Source/Runtime/UMG/Private/GridPanelComponent.cpp
--- a/Engine/Source/Runtime/UMG/Private/GridPanelComponent.cpp
+++ b/Engine/Source/Runtime/UMG/Private/GridPanelComponent.cpp
@@ -17,13 +17,6 @@ TSharedRef<SWidget> UGridPanelComponent::RebuildWidget()
 	TSharedRef<SGridPanel> NewGrid = SNew(SGridPanel);
 	MyGrid = NewGrid;
 
-	//@TODO: Expose FillColumn and FillRow
-
-	return NewGrid;
-}
-
-void UGridPanelComponent::OnKnownChildrenChanged()
-{
 	//TSharedPtr<SGridPanel> Grid = MyGrid.Pin();
 	//if (Grid.IsValid())
 	//{
@@ -86,25 +79,8 @@ void UGridPanelComponent::OnKnownChildrenChanged()
 	//		];
 	//	}
 	//}
-}
-
-bool UGridPanelComponent::HasAnySockets() const
-{
-	return true;
-}
 
-FTransform UGridPanelComponent::GetSocketTransform(FName InSocketName, ERelativeTransformSpace TransformSpace) const
-{
-	return FTransform::Identity;
-}
-
-void UGridPanelComponent::QuerySupportedSockets(TArray<FComponentSocketDescription>& OutSockets) const
-{
-	for (int32 SlotIndex = 0; SlotIndex < Slots.Num(); ++SlotIndex)
-	{
-		FName SocketName = Slots[SlotIndex].SlotName;
-		new (OutSockets) FComponentSocketDescription(SocketName, EComponentSocketType::Socket);
-	}
+	return NewGrid;
 }
 
 #if WITH_EDITOR
diff --git a/Engine/Source/Runtime/UMG/Private/HorizontalBoxComponent.cpp b/Engine/Source/Runtime/UMG/Private/HorizontalBoxComponent.cpp
--- a/Engine/Source/Runtime/UMG/Private/HorizontalBoxComponent.cpp
+++ b/Engine/Source/Runtime/UMG/Private/HorizontalBoxComponent.cpp
@@ -25,19 +25,12 @@ TSharedRef<SWidget> UHorizontalBoxComponent::RebuildWidget()
 	TSharedRef<SHorizontalBox> NewCanvas = SNew(SHorizontalBox);
 	MyHorizontalBox = NewCanvas;
 
-	OnKnownChildrenChanged();
-
-	return NewCanvas;
-}
-
-void UHorizontalBoxComponent::OnKnownChildrenChanged()
-{
 	TSharedPtr<SHorizontalBox> Canvas = MyHorizontalBox.Pin();
-	if (Canvas.IsValid())
+	if ( Canvas.IsValid() )
 	{
 		Canvas->ClearChildren();
 
-		for (int32 SlotIndex = 0; SlotIndex < Slots.Num(); ++SlotIndex)
+		for ( int32 SlotIndex = 0; SlotIndex < Slots.Num(); ++SlotIndex )
 		{
 			UHorizontalBoxSlot* Slot = Slots[SlotIndex];
 			if ( Slot == NULL )
@@ -49,13 +42,15 @@ void UHorizontalBoxComponent::OnKnownChildrenChanged()
 				.Padding(Slot->Padding)
 				.HAlign(Slot->HorizontalAlignment)
 				.VAlign(Slot->VerticalAlignment)
-			[
-				Slot->Content == NULL ? SNullWidget::NullWidget : Slot->Content->GetWidget()
-			];
+				[
+					Slot->Content == NULL ? SNullWidget::NullWidget : Slot->Content->GetWidget()
+				];
 
 			NewSlot.SizeParam = USlateWrapperComponent::ConvertSerializedSizeParamToRuntime(Slot->Size);
 		}
 	}
+
+	return NewCanvas;
 }
 
 UHorizontalBoxSlot* UHorizontalBoxComponent::AddSlot(USlateWrapperComponent* Content)
diff --git a/Engine/Source/Runtime/UMG/Private/SlateNonLeafWidgetComponent.cpp b/Engine/Source/Runtime/UMG/Private/SlateNonLeafWidgetComponent.cpp
--- a/Engine/Source/Runtime/UMG/Private/SlateNonLeafWidgetComponent.cpp
+++ b/Engine/Source/Runtime/UMG/Private/SlateNonLeafWidgetComponent.cpp
@@ -19,110 +19,8 @@ USlateNonLeafWidgetComponent::USlateNonLeafWidgetComponent(const FPostConstructI
 
 	bAutoActivate = true;
 	bTickInEditor = true;
-
-	bCanHaveMultipleChildren = true;
-}
-
-void USlateNonLeafWidgetComponent::TickComponent(float DeltaTime, enum ELevelTick TickType, FActorComponentTickFunction* ThisTickFunction)
-{
-	//bool bNeedRebuild = false;
-	//TSet<USlateWrapperComponent*> RemainingChildren;
-	//for (auto SetIt = KnownChildren.CreateIterator(); SetIt; ++SetIt)
-	//{
-	//	if (USlateWrapperComponent* Wrapper = (*SetIt).Get())
-	//	{
-	//		if (Wrapper->IsRegistered())
-	//		{
-	//			RemainingChildren.Add(Wrapper);
-	//		}
-	//	}
-	//}
-
-	//for (int32 ComponentIndex = 0; (ComponentIndex < AttachChildren.Num()) && !bNeedRebuild; ++ComponentIndex)
-	//{
-	//	if (USlateWrapperComponent* Wrapper = Cast<USlateWrapperComponent>(AttachChildren[ComponentIndex]))
-	//	{
-	//		if (Wrapper->IsRegistered())
-	//		{
-	//			if (KnownChildren.Contains(Wrapper))
-	//			{
-	//				RemainingChildren.Remove(Wrapper);
-	//			}
-	//			else
-	//			{
-	//				bNeedRebuild = true;
-	//			}
-	//		}
-	//	}
-	//}
-
-	//bNeedRebuild = bNeedRebuild || (RemainingChildren.Num() > 0);
-
-	//if (bNeedRebuild)
-	//{
-	//	// Recreate the list of known children
-	//	KnownChildren.Empty();
-	//	for (int32 ComponentIndex = 0; ComponentIndex < AttachChildren.Num(); ++ComponentIndex)
-	//	{
-	//		if (USlateWrapperComponent* Wrapper = Cast<USlateWrapperComponent>(AttachChildren[ComponentIndex]))
-	//		{
-	//			if (Wrapper->IsRegistered())
-	//			{
-	//				KnownChildren.Add(Wrapper);
-	//			}
-	//		}
-	//	}
-
-	//	// Let derived classes do something with the list
-	//	OnKnownChildrenChanged();
-	//}
-}
-
-USlateWrapperComponent* USlateNonLeafWidgetComponent::GetFirstWrappedChild() const
-{
-	if (auto It = KnownChildren.CreateConstIterator())
-	{
-		return It->Get();
-	}
-	else
-	{
-		return NULL;
-	}
 }
 
-TSharedRef<SWidget> USlateNonLeafWidgetComponent::GetFirstChildWidget() const
-{
-	if (USlateWrapperComponent* WrappedChild = GetFirstWrappedChild())
-	{
-		return WrappedChild->GetWidget();
-	}
-	else
-	{
-		return SNullWidget::NullWidget;
-	}
-}
-
-#if WITH_EDITOR
-
-void USlateNonLeafWidgetComponent::CheckForErrors()
-{
-	Super::CheckForErrors();
-
-	if ((KnownChildren.Num() > 1) && !bCanHaveMultipleChildren)
-	{
-		AActor* Owner = GetOwner();
-		FFormatNamedArguments Arguments;
-		Arguments.Add(TEXT("ComponentName"), FText::FromString(GetName()));
-		Arguments.Add(TEXT("OwnerName"), FText::FromString(Owner->GetName()));
-		FText Message = FText::Format( LOCTEXT( "MapCheck_Message_TooManyChildren", "{ComponentName}::{OwnerName} cannot have more than one child attached" ), Arguments );
-		FMessageLog("MapCheck").Warning()
-			->AddToken(FUObjectToken::Create(Owner))
-			->AddToken(FTextToken::Create(Message) );
-	}
-}
-
-#endif
-
 /////////////////////////////////////////////////////
 
 #undef LOCTEXT_NAMESPACE
diff --git a/Engine/Source/Runtime/UMG/Private/TextBlockComponent.cpp b/Engine/Source/Runtime/UMG/Private/TextBlockComponent.cpp
--- a/Engine/Source/Runtime/UMG/Private/TextBlockComponent.cpp
+++ b/Engine/Source/Runtime/UMG/Private/TextBlockComponent.cpp
@@ -10,24 +10,28 @@
 UTextBlockComponent::UTextBlockComponent(const FPostConstructInitializeProperties& PCIP)
 	: Super(PCIP)
 {
-	DisplayText = LOCTEXT("TextBlockDefaultValue", "Hello World");
+	Text = LOCTEXT("TextBlockDefaultValue", "Text Block");
+	Font = FSlateFontInfo(TEXT("Slate/Fonts/Roboto-Bold.ttf"), 24);
+	ShadowOffset = FVector2D(1.0f, 1.0f);
 	ColorAndOpacity = FLinearColor::White;
 	ShadowColorAndOpacity = FLinearColor::Transparent;
 }
 
 TSharedRef<SWidget> UTextBlockComponent::RebuildWidget()
 {
+	FString FontPath = FPaths::EngineContentDir() / Font.FontName.ToString();
+
 	return SNew(STextBlock)
-		.Text( BIND_UOBJECT_ATTRIBUTE(FString, GetDisplayText) )
-		.Font( FSlateFontInfo(FPaths::EngineContentDir() / TEXT("Slate/Fonts/Roboto-Bold.ttf"), 24) )
+		.Text( BIND_UOBJECT_ATTRIBUTE(FString, GetText) )
+		.Font(FSlateFontInfo(FontPath, Font.Size))
 		.ColorAndOpacity( BIND_UOBJECT_ATTRIBUTE(FSlateColor, GetColorAndOpacity) )
-		.ShadowOffset(FVector2D(1.0f, 1.0f))
+		.ShadowOffset(ShadowOffset)
 		.ShadowColorAndOpacity( BIND_UOBJECT_ATTRIBUTE(FLinearColor, GetShadowColorAndOpacity) );
 }
 
-FString UTextBlockComponent::GetDisplayText() const
+FString UTextBlockComponent::GetText() const
 {
-	return DisplayText.ToString();
+	return Text.ToString();
 }
 
 FSlateColor UTextBlockComponent::GetColorAndOpacity() const
diff --git a/Engine/Source/Runtime/UMG/Private/VerticalBoxComponent.cpp b/Engine/Source/Runtime/UMG/Private/VerticalBoxComponent.cpp
--- a/Engine/Source/Runtime/UMG/Private/VerticalBoxComponent.cpp
+++ b/Engine/Source/Runtime/UMG/Private/VerticalBoxComponent.cpp
@@ -25,19 +25,12 @@ TSharedRef<SWidget> UVerticalBoxComponent::RebuildWidget()
 	TSharedRef<SVerticalBox> NewCanvas = SNew(SVerticalBox);
 	MyVerticalBox = NewCanvas;
 
-	OnKnownChildrenChanged();
-
-	return NewCanvas;
-}
-
-void UVerticalBoxComponent::OnKnownChildrenChanged()
-{
 	TSharedPtr<SVerticalBox> Canvas = MyVerticalBox.Pin();
-	if (Canvas.IsValid())
+	if ( Canvas.IsValid() )
 	{
 		Canvas->ClearChildren();
 
-		for (int32 SlotIndex = 0; SlotIndex < Slots.Num(); ++SlotIndex)
+		for ( int32 SlotIndex = 0; SlotIndex < Slots.Num(); ++SlotIndex )
 		{
 			UVerticalBoxSlot* Slot = Slots[SlotIndex];
 			if ( Slot == NULL )
@@ -49,13 +42,15 @@ void UVerticalBoxComponent::OnKnownChildrenChanged()
 				.Padding(Slot->Padding)
 				.HAlign(Slot->HorizontalAlignment)
 				.VAlign(Slot->VerticalAlignment)
-			[
-				Slot->Content == NULL ? SNullWidget::NullWidget : Slot->Content->GetWidget()
-			];
+				[
+					Slot->Content == NULL ? SNullWidget::NullWidget : Slot->Content->GetWidget()
+				];
 
 			NewSlot.SizeParam = USlateWrapperComponent::ConvertSerializedSizeParamToRuntime(Slot->Size);
 		}
 	}
+
+	return NewCanvas;
 }
 
 UVerticalBoxSlot* UVerticalBoxComponent::AddSlot(USlateWrapperComponent* Content)
@@ -95,4 +90,4 @@ void UVerticalBoxComponent::PostEditChangeProperty(struct FPropertyChangedEvent&
 		}
 	}
 }
-#endif
\ No newline at end of file
+#endif
