UMG - Reorganizing UMG runtime, creating a module (include all) header, trimming down requirements and explictly importing all public modules.

diff --git a/Engine/Source/Editor/UMGEditor/Private/UMGEditorPrivatePCH.h b/Engine/Source/Editor/UMGEditor/Private/UMGEditorPrivatePCH.h
--- a/Engine/Source/Editor/UMGEditor/Private/UMGEditorPrivatePCH.h
+++ b/Engine/Source/Editor/UMGEditor/Private/UMGEditorPrivatePCH.h
@@ -3,8 +3,10 @@
 #pragma once
 
 #include "UnrealEd.h"
-#include "UMGClasses.h"
-#include "IUMGModule.h"
+
+#include "UMG.h"
+
 #include "UMGEditorClasses.h"
+
 #include "AssetTypeActions_Base.h"
 #include "AssetTypeActions_Blueprint.h"
\ No newline at end of file
diff --git a/Engine/Source/Runtime/UMG/Classes/BorderComponent.h b/Engine/Source/Runtime/UMG/Classes/BorderComponent.h
deleted file mode 100644
diff --git a/Engine/Source/Runtime/UMG/Classes/BoxPanelComponent.h b/Engine/Source/Runtime/UMG/Classes/BoxPanelComponent.h
deleted file mode 100644
diff --git a/Engine/Source/Runtime/UMG/Classes/ButtonComponent.h b/Engine/Source/Runtime/UMG/Classes/ButtonComponent.h
deleted file mode 100644
diff --git a/Engine/Source/Runtime/UMG/Classes/CanvasPanelComponent.h b/Engine/Source/Runtime/UMG/Classes/CanvasPanelComponent.h
deleted file mode 100644
diff --git a/Engine/Source/Runtime/UMG/Classes/CanvasPanelSlot.h b/Engine/Source/Runtime/UMG/Classes/CanvasPanelSlot.h
deleted file mode 100644
diff --git a/Engine/Source/Runtime/UMG/Classes/CheckBoxComponent.h b/Engine/Source/Runtime/UMG/Classes/CheckBoxComponent.h
deleted file mode 100644
diff --git a/Engine/Source/Runtime/UMG/Classes/CircularThrobberComponent.h b/Engine/Source/Runtime/UMG/Classes/CircularThrobberComponent.h
deleted file mode 100644
diff --git a/Engine/Source/Runtime/UMG/Classes/ContentWidget.h b/Engine/Source/Runtime/UMG/Classes/ContentWidget.h
deleted file mode 100644
diff --git a/Engine/Source/Runtime/UMG/Classes/EditableTextBlockComponent.h b/Engine/Source/Runtime/UMG/Classes/EditableTextBlockComponent.h
deleted file mode 100644
diff --git a/Engine/Source/Runtime/UMG/Classes/GridPanelComponent.h b/Engine/Source/Runtime/UMG/Classes/GridPanelComponent.h
deleted file mode 100644
diff --git a/Engine/Source/Runtime/UMG/Classes/HorizontalBoxComponent.h b/Engine/Source/Runtime/UMG/Classes/HorizontalBoxComponent.h
deleted file mode 100644
diff --git a/Engine/Source/Runtime/UMG/Classes/HorizontalBoxSlot.h b/Engine/Source/Runtime/UMG/Classes/HorizontalBoxSlot.h
deleted file mode 100644
diff --git a/Engine/Source/Runtime/UMG/Classes/ImageComponent.h b/Engine/Source/Runtime/UMG/Classes/ImageComponent.h
deleted file mode 100644
diff --git a/Engine/Source/Runtime/UMG/Classes/PanelSlot.h b/Engine/Source/Runtime/UMG/Classes/PanelSlot.h
deleted file mode 100644
diff --git a/Engine/Source/Runtime/UMG/Classes/ScrollBoxComponent.h b/Engine/Source/Runtime/UMG/Classes/ScrollBoxComponent.h
deleted file mode 100644
diff --git a/Engine/Source/Runtime/UMG/Classes/SlateLeafWidgetComponent.h b/Engine/Source/Runtime/UMG/Classes/SlateLeafWidgetComponent.h
deleted file mode 100644
diff --git a/Engine/Source/Runtime/UMG/Classes/SlateNonLeafWidgetComponent.h b/Engine/Source/Runtime/UMG/Classes/SlateNonLeafWidgetComponent.h
deleted file mode 100644
diff --git a/Engine/Source/Runtime/UMG/Classes/SlateWrapperComponent.h b/Engine/Source/Runtime/UMG/Classes/SlateWrapperComponent.h
deleted file mode 100644
diff --git a/Engine/Source/Runtime/UMG/Classes/SlateWrapperTypes.h b/Engine/Source/Runtime/UMG/Classes/SlateWrapperTypes.h
deleted file mode 100644
diff --git a/Engine/Source/Runtime/UMG/Classes/SpacerComponent.h b/Engine/Source/Runtime/UMG/Classes/SpacerComponent.h
deleted file mode 100644
diff --git a/Engine/Source/Runtime/UMG/Classes/SpinningImageComponent.h b/Engine/Source/Runtime/UMG/Classes/SpinningImageComponent.h
deleted file mode 100644
diff --git a/Engine/Source/Runtime/UMG/Classes/TextBlockComponent.h b/Engine/Source/Runtime/UMG/Classes/TextBlockComponent.h
deleted file mode 100644
diff --git a/Engine/Source/Runtime/UMG/Classes/ThrobberComponent.h b/Engine/Source/Runtime/UMG/Classes/ThrobberComponent.h
deleted file mode 100644
diff --git a/Engine/Source/Runtime/UMG/Classes/UserWidget.h b/Engine/Source/Runtime/UMG/Classes/UserWidget.h
deleted file mode 100644
diff --git a/Engine/Source/Runtime/UMG/Classes/VerticalBoxComponent.h b/Engine/Source/Runtime/UMG/Classes/VerticalBoxComponent.h
deleted file mode 100644
diff --git a/Engine/Source/Runtime/UMG/Classes/VerticalBoxSlot.h b/Engine/Source/Runtime/UMG/Classes/VerticalBoxSlot.h
deleted file mode 100644
diff --git a/Engine/Source/Runtime/UMG/Classes/WidgetBlueprint.h b/Engine/Source/Runtime/UMG/Classes/WidgetBlueprint.h
deleted file mode 100644
diff --git a/Engine/Source/Runtime/UMG/Classes/WidgetBlueprintGeneratedClass.h b/Engine/Source/Runtime/UMG/Classes/WidgetBlueprintGeneratedClass.h
deleted file mode 100644
diff --git a/Engine/Source/Runtime/UMG/Classes/WidgetTree.h b/Engine/Source/Runtime/UMG/Classes/WidgetTree.h
deleted file mode 100644
diff --git a/Engine/Source/Runtime/UMG/Private/BorderComponent.cpp b/Engine/Source/Runtime/UMG/Private/BorderComponent.cpp
deleted file mode 100644
diff --git a/Engine/Source/Runtime/UMG/Private/BoxPanelComponent.cpp b/Engine/Source/Runtime/UMG/Private/BoxPanelComponent.cpp
deleted file mode 100644
diff --git a/Engine/Source/Runtime/UMG/Private/ButtonComponent.cpp b/Engine/Source/Runtime/UMG/Private/ButtonComponent.cpp
deleted file mode 100644
diff --git a/Engine/Source/Runtime/UMG/Private/CanvasPanelComponent.cpp b/Engine/Source/Runtime/UMG/Private/CanvasPanelComponent.cpp
deleted file mode 100644
diff --git a/Engine/Source/Runtime/UMG/Private/CanvasPanelSlot.cpp b/Engine/Source/Runtime/UMG/Private/CanvasPanelSlot.cpp
deleted file mode 100644
diff --git a/Engine/Source/Runtime/UMG/Private/CheckBoxComponent.cpp b/Engine/Source/Runtime/UMG/Private/CheckBoxComponent.cpp
deleted file mode 100644
diff --git a/Engine/Source/Runtime/UMG/Private/CircularThrobberComponent.cpp b/Engine/Source/Runtime/UMG/Private/CircularThrobberComponent.cpp
deleted file mode 100644
diff --git a/Engine/Source/Runtime/UMG/Private/Components/BorderComponent.cpp b/Engine/Source/Runtime/UMG/Private/Components/BorderComponent.cpp
new file mode 100644
--- /dev/null
+++ b/Engine/Source/Runtime/UMG/Private/Components/BorderComponent.cpp
@@ -0,0 +1,95 @@
+// Copyright 1998-2014 Epic Games, Inc. All Rights Reserved.
+
+#include "UMGPrivatePCH.h"
+
+#define LOCTEXT_NAMESPACE "UMG"
+
+/////////////////////////////////////////////////////
+// UBorderComponent
+
+UBorderComponent::UBorderComponent(const FPostConstructInitializeProperties& PCIP)
+	: Super(PCIP)
+{
+	SBorder::FArguments BorderDefaults;
+
+	HorizontalAlignment = BorderDefaults._HAlign;
+	VerticalAlignment = BorderDefaults._VAlign;
+	ContentPadding = BorderDefaults._Padding.Get();
+	//BorderImage = *BorderDefaults._BorderImage.Get();
+	DesiredSizeScale = BorderDefaults._DesiredSizeScale.Get();
+	ContentScale = BorderDefaults._ContentScale.Get();
+	ContentColorAndOpacity = BorderDefaults._ColorAndOpacity.Get();
+	BorderColorAndOpacity = BorderDefaults._BorderBackgroundColor.Get();
+	ForegroundColor = BorderDefaults._ForegroundColor.Get();
+	bShowEffectWhenDisabled = BorderDefaults._ShowEffectWhenDisabled.Get();
+}
+
+TSharedRef<SWidget> UBorderComponent::RebuildWidget()
+{
+	TSharedRef<SBorder> NewBorder = SNew(SBorder)
+		.HAlign(HorizontalAlignment)
+		.VAlign(VerticalAlignment)
+		.Padding(BIND_UOBJECT_ATTRIBUTE(FMargin, GetContentPadding))
+		// SLATE_EVENT( FPointerEventHandler, OnMouseButtonDown )
+		// SLATE_EVENT( FPointerEventHandler, OnMouseButtonUp )
+		// SLATE_EVENT( FPointerEventHandler, OnMouseMove )
+		// SLATE_EVENT( FPointerEventHandler, OnMouseDoubleClick )
+		.BorderImage(BIND_UOBJECT_ATTRIBUTE(const FSlateBrush*, GetBorderBrush))
+		.ContentScale(ContentScale)
+		.DesiredSizeScale(DesiredSizeScale)
+		.ColorAndOpacity(BIND_UOBJECT_ATTRIBUTE(FLinearColor, GetContentColor))
+		.BorderBackgroundColor(BIND_UOBJECT_ATTRIBUTE(FSlateColor, GetBorderColor))
+		.ForegroundColor(BIND_UOBJECT_ATTRIBUTE(FSlateColor, GetForegroundColor))
+		.ShowEffectWhenDisabled(bShowEffectWhenDisabled);
+
+	NewBorder->SetContent(Content ? Content->GetWidget() : SNullWidget::NullWidget);
+
+	MyBorder = NewBorder;
+	return NewBorder;
+}
+
+void UBorderComponent::SetContent(USlateWrapperComponent* Content)
+{
+	Super::SetContent(Content);
+
+	TSharedPtr<SBorder> Border = MyBorder.Pin();
+	if ( Border.IsValid() )
+	{
+		Border->SetContent(Content ? Content->GetWidget() : SNullWidget::NullWidget);
+	}
+}
+
+FMargin UBorderComponent::GetContentPadding() const
+{
+	return ContentPadding;
+}
+
+FLinearColor UBorderComponent::GetContentColor() const
+{
+	return ContentColorAndOpacity;
+}
+
+const FSlateBrush* UBorderComponent::GetBorderBrush() const
+{
+	if ( BorderBrush == NULL )
+	{
+		SBorder::FArguments BorderDefaults;
+		return BorderDefaults._BorderImage.Get();
+	}
+
+	return &BorderBrush->Brush;
+}
+
+FSlateColor UBorderComponent::GetBorderColor() const
+{
+	return BorderColorAndOpacity;
+}
+
+FSlateColor UBorderComponent::GetForegroundColor() const
+{
+	return ForegroundColor;
+}
+
+/////////////////////////////////////////////////////
+
+#undef LOCTEXT_NAMESPACE
diff --git a/Engine/Source/Runtime/UMG/Private/Components/BoxPanelComponent.cpp b/Engine/Source/Runtime/UMG/Private/Components/BoxPanelComponent.cpp
new file mode 100644
--- /dev/null
+++ b/Engine/Source/Runtime/UMG/Private/Components/BoxPanelComponent.cpp
@@ -0,0 +1,142 @@
+// Copyright 1998-2014 Epic Games, Inc. All Rights Reserved.
+
+#include "UMGPrivatePCH.h"
+
+/////////////////////////////////////////////////////
+// SConcreteBoxPanel
+
+class SConcreteBoxPanel : public SBoxPanel
+{
+public:
+	class FSlot : public SBoxPanel::FSlot
+	{
+	public:
+		FSlot& SetSizeParam(const FSizeParam& InParam)
+		{
+			SizeParam = InParam;
+			return *this;
+		}
+
+		FSlot& MaxWidth( const TAttribute< float >& InMaxWidth )
+		{
+			MaxSize = InMaxWidth;
+			return *this;
+		}
+
+		FSlot& Padding( TAttribute<FMargin> InPadding )
+		{
+			SlotPadding = InPadding;
+			return *this;
+		}
+
+		FSlot& HAlign( EHorizontalAlignment InHAlignment )
+		{
+			HAlignment = InHAlignment;
+			return *this;
+		}
+
+		FSlot& VAlign( EVerticalAlignment InVAlignment )
+		{
+			VAlignment = InVAlignment;
+			return *this;
+		}
+
+		FSlot& operator[]( TSharedRef<SWidget> InWidget )
+		{
+			Widget = InWidget;
+			return *this;
+		}
+
+		FSlot& Expose( FSlot*& OutVarToInit )
+		{
+			OutVarToInit = this;
+			return *this;
+		}
+	};
+
+	static FSlot& Slot()
+	{
+		return *(new FSlot());
+	}
+
+	SLATE_BEGIN_ARGS( SConcreteBoxPanel )
+		{
+			_Visibility = EVisibility::SelfHitTestInvisible;
+		}
+		SLATE_SUPPORTS_SLOT(SConcreteBoxPanel::FSlot)
+	SLATE_END_ARGS()
+
+	FSlot& AddSlot()
+	{
+		SConcreteBoxPanel::FSlot& NewSlot = SConcreteBoxPanel::Slot();
+		this->Children.Add( &NewSlot );
+		return NewSlot;
+	}
+
+	void Construct(const FArguments& InArgs, EOrientation InOrientation)
+	{
+		DirtyHackSetOrientation(InOrientation);
+
+		const int32 NumSlots = InArgs.Slots.Num();
+		for (int32 SlotIndex = 0; SlotIndex < NumSlots; ++SlotIndex)
+		{
+			Children.Add(InArgs.Slots[SlotIndex]);
+		}
+	}
+
+	SConcreteBoxPanel()
+		: SBoxPanel(Orient_Horizontal)
+	{
+	}
+
+protected:
+	//@TODO: Make Orientation not const in the base class; it's already protected and AFAIK there are no
+	// negative consequences from altering it before the first tick/computedesiredsize call
+	void DirtyHackSetOrientation(EOrientation InOrientation)
+	{
+		EOrientation& NonConstOrientation = *(const_cast<EOrientation*>(&Orientation));
+		NonConstOrientation = InOrientation;
+	}
+};
+
+/////////////////////////////////////////////////////
+// UBoxPanelComponent
+
+UBoxPanelComponent::UBoxPanelComponent(const FPostConstructInitializeProperties& PCIP)
+	: Super(PCIP)
+{
+}
+
+TSharedRef<SWidget> UBoxPanelComponent::RebuildWidget()
+{
+	SConcreteBoxPanel::FArguments Args;
+
+	TSharedRef<SConcreteBoxPanel> Box = SNew(SConcreteBoxPanel, Orientation);
+	MyBox = Box;
+
+	return Box;
+}
+
+#if WITH_EDITOR
+void UBoxPanelComponent::PostEditChangeProperty(struct FPropertyChangedEvent& PropertyChangedEvent)
+{
+	// Ensure the slots have unique names
+	int32 SlotNumbering = 1;
+
+	TSet<FName> UniqueSlotNames;
+	for (int32 SlotIndex = 0; SlotIndex < Slots.Num(); ++SlotIndex)
+	{
+		FName& SlotName = Slots[SlotIndex].SlotName;
+
+		if ((SlotName == NAME_None) || UniqueSlotNames.Contains(SlotName))
+		{
+			do 
+			{
+				SlotName = FName(TEXT("Slot"), SlotNumbering++);
+			} while (UniqueSlotNames.Contains(SlotName));
+		}
+
+		UniqueSlotNames.Add(SlotName);
+	}
+}
+#endif
\ No newline at end of file
diff --git a/Engine/Source/Runtime/UMG/Private/Components/ButtonComponent.cpp b/Engine/Source/Runtime/UMG/Private/Components/ButtonComponent.cpp
new file mode 100644
--- /dev/null
+++ b/Engine/Source/Runtime/UMG/Private/Components/ButtonComponent.cpp
@@ -0,0 +1,140 @@
+// Copyright 1998-2014 Epic Games, Inc. All Rights Reserved.
+
+#include "UMGPrivatePCH.h"
+
+#define LOCTEXT_NAMESPACE "UMG"
+
+/////////////////////////////////////////////////////
+// UButtonComponent
+
+UButtonComponent::UButtonComponent(const FPostConstructInitializeProperties& PCIP)
+	: Super(PCIP)
+{
+	SButton::FArguments ButtonDefaults;
+
+	ButtonStyle = NULL;
+	HorizontalAlignment = ButtonDefaults._HAlign;
+	VerticalAlignment = ButtonDefaults._VAlign;
+	ContentPadding = ButtonDefaults._ContentPadding.Get();
+
+	bDisplayText = true;
+	ButtonText = LOCTEXT("ButtonDefaultValue", "This is a button");
+	ButtonTextStyle = NULL;
+
+	DesiredSizeScale = ButtonDefaults._DesiredSizeScale.Get();
+	ContentScale = ButtonDefaults._ContentScale.Get();
+
+	ButtonColorAndOpacity = ButtonDefaults._ButtonColorAndOpacity.Get();
+	ForegroundColor = ButtonDefaults._ForegroundColor.Get();
+}
+
+TSharedRef<SWidget> UButtonComponent::RebuildWidget()
+{
+	SButton::FArguments ButtonDefaults;
+
+	const FButtonStyle* StylePtr = (ButtonStyle != NULL) ? ButtonStyle->GetStyle<FButtonStyle>() : NULL;
+	if (StylePtr == NULL)
+	{
+		StylePtr = ButtonDefaults._ButtonStyle;
+	}
+
+	const FTextBlockStyle* TextStylePtr = ( ButtonTextStyle != NULL ) ? ButtonTextStyle->GetStyle<FTextBlockStyle>() : NULL;
+	if ( TextStylePtr == NULL )
+	{
+		TextStylePtr = ButtonDefaults._TextStyle;
+	}
+
+	if (bDisplayText)
+	{
+		TextAttribute = BIND_UOBJECT_ATTRIBUTE(FText, GetButtonText);
+	}
+	else
+	{
+		TAttribute<FText> NotSet;
+		TextAttribute = NotSet;
+	}
+
+	TSharedRef<SButton> NewButton = SNew(SButton)
+		.ButtonStyle(StylePtr)
+		.TextStyle(TextStylePtr)
+		.HAlign(HorizontalAlignment)
+		.VAlign(VerticalAlignment)
+		.ContentPadding(BIND_UOBJECT_ATTRIBUTE(FMargin, GetContentPadding))
+		.Text(TextAttribute)
+		.OnClicked(BIND_UOBJECT_DELEGATE(FOnClicked, SlateOnClickedCallback))
+		.DesiredSizeScale(DesiredSizeScale)
+		.ContentScale(ContentScale)
+		.ButtonColorAndOpacity(BIND_UOBJECT_ATTRIBUTE(FSlateColor, GetButtonColor))
+		.ForegroundColor(BIND_UOBJECT_ATTRIBUTE(FSlateColor, GetForegroundColor));
+
+	if ( Content )
+	{
+		NewButton->SetContent(Content->GetWidget());
+	}
+	
+	MyButton = NewButton;
+	return NewButton;
+}
+
+void UButtonComponent::SetContent(USlateWrapperComponent* InContent)
+{
+	Super::SetContent(InContent);
+
+	TSharedPtr<SButton> Button = MyButton.Pin();
+	if (Button.IsValid())
+	{
+		if ( InContent )
+		{
+			Button->SetContent(InContent->GetWidget());
+		}
+		else
+		{
+			// Removed a child; replace the text if desired
+			if (bDisplayText)
+			{
+				Button->SetContent(SNew(STextBlock)
+					.Text(TextAttribute)
+					.TextStyle( ButtonTextStyle ) );
+			}
+			else
+			{
+				Button->SetContent(SNullWidget::NullWidget);
+			}
+		}
+	}
+}
+
+FMargin UButtonComponent::GetContentPadding() const
+{
+	return ContentPadding;
+}
+
+FText UButtonComponent::GetButtonText() const
+{
+	return ButtonText;
+}
+
+FText UButtonComponent::GetText_Implementation() const
+{
+	return ButtonText;
+}
+
+FReply UButtonComponent::SlateOnClickedCallback()
+{
+	OnClicked.Broadcast();
+	return FReply::Handled();
+}
+
+FSlateColor UButtonComponent::GetButtonColor() const
+{
+	return ButtonColorAndOpacity;
+}
+
+FSlateColor UButtonComponent::GetForegroundColor() const
+{
+	return ForegroundColor;
+}
+
+/////////////////////////////////////////////////////
+
+#undef LOCTEXT_NAMESPACE
diff --git a/Engine/Source/Runtime/UMG/Private/Components/CanvasPanelComponent.cpp b/Engine/Source/Runtime/UMG/Private/Components/CanvasPanelComponent.cpp
new file mode 100644
--- /dev/null
+++ b/Engine/Source/Runtime/UMG/Private/Components/CanvasPanelComponent.cpp
@@ -0,0 +1,135 @@
+// Copyright 1998-2014 Epic Games, Inc. All Rights Reserved.
+
+#include "UMGPrivatePCH.h"
+
+/////////////////////////////////////////////////////
+// SFixedSizeCanvas
+
+class SFixedSizeCanvas : public SCanvas
+{
+	SLATE_BEGIN_ARGS( SFixedSizeCanvas )
+		{
+			_Visibility = EVisibility::SelfHitTestInvisible;
+		}
+	SLATE_END_ARGS()
+
+	FVector2D CanvasSize;
+
+	void Construct(const FArguments& InArgs, FVector2D InSize)
+	{
+		SCanvas::FArguments ParentArgs;
+		SCanvas::Construct(ParentArgs);
+
+		CanvasSize = InSize;
+	}
+
+	// SWidget interface
+	virtual FVector2D ComputeDesiredSize() const OVERRIDE
+	{
+		return CanvasSize;
+	}
+	// End of SWidget interface
+};
+
+/////////////////////////////////////////////////////
+// UCanvasPanelComponent
+
+UCanvasPanelComponent::UCanvasPanelComponent(const FPostConstructInitializeProperties& PCIP)
+	: Super(PCIP)
+{
+	DesiredCanvasSize = FVector2D(128.0f, 128.0f);
+}
+
+int32 UCanvasPanelComponent::GetChildrenCount() const
+{
+	return Slots.Num();
+}
+
+USlateWrapperComponent* UCanvasPanelComponent::GetChildAt(int32 Index) const
+{
+	return Slots[Index]->Content;
+}
+
+TSharedRef<SWidget> UCanvasPanelComponent::RebuildWidget()
+{
+	TSharedRef<SFixedSizeCanvas> NewCanvas = SNew(SFixedSizeCanvas, DesiredCanvasSize);
+	MyCanvas = NewCanvas;
+
+	TSharedPtr<SFixedSizeCanvas> Canvas = MyCanvas.Pin();
+	if ( Canvas.IsValid() )
+	{
+		Canvas->ClearChildren();
+
+		for ( int32 SlotIndex = 0; SlotIndex < Slots.Num(); ++SlotIndex )
+		{
+			UCanvasPanelSlot* Slot = Slots[SlotIndex];
+			if ( Slot == NULL )
+			{
+				Slots[SlotIndex] = Slot = ConstructObject<UCanvasPanelSlot>(UCanvasPanelSlot::StaticClass(), this);
+			}
+
+			auto NewSlot = Canvas->AddSlot()
+				.Position(Slot->Position)
+				.Size(Slot->Size)
+				.HAlign(Slot->HorizontalAlignment)
+				.VAlign(Slot->VerticalAlignment)
+				[
+					Slot->Content == NULL ? SNullWidget::NullWidget : Slot->Content->GetWidget()
+				];
+		}
+	}
+
+	return NewCanvas;
+}
+
+UCanvasPanelSlot* UCanvasPanelComponent::AddSlot(USlateWrapperComponent* Content)
+{
+	UCanvasPanelSlot* Slot = ConstructObject<UCanvasPanelSlot>(UCanvasPanelSlot::StaticClass(), this);
+	Slot->Content = Content;
+
+#if WITH_EDITOR
+	Content->Slot = Slot;
+#endif
+	
+	Slots.Add(Slot);
+
+	return Slot;
+}
+
+#if WITH_EDITOR
+
+void UCanvasPanelComponent::ConnectEditorData()
+{
+	for ( UCanvasPanelSlot* Slot : Slots )
+	{
+		Slot->Content->Slot = Slot;
+	}
+}
+
+void UCanvasPanelComponent::PostEditChangeProperty(struct FPropertyChangedEvent& PropertyChangedEvent)
+{
+	// Ensure the slots have unique names
+	int32 SlotNumbering = 1;
+
+	TSet<FName> UniqueSlotNames;
+	for (int32 SlotIndex = 0; SlotIndex < Slots.Num(); ++SlotIndex)
+	{
+		if ( Slots[SlotIndex] == NULL )
+		{
+			Slots[SlotIndex] = ConstructObject<UCanvasPanelSlot>(UCanvasPanelSlot::StaticClass(), this);
+		}
+
+		//FName& SlotName = Slots[SlotIndex].SlotName;
+
+		//if ((SlotName == NAME_None) || UniqueSlotNames.Contains(SlotName))
+		//{
+		//	do 
+		//	{
+		//		SlotName = FName(TEXT("Slot"), SlotNumbering++);
+		//	} while (UniqueSlotNames.Contains(SlotName));
+		//}
+
+		//UniqueSlotNames.Add(SlotName);
+	}
+}
+#endif
\ No newline at end of file
diff --git a/Engine/Source/Runtime/UMG/Private/Components/CanvasPanelSlot.cpp b/Engine/Source/Runtime/UMG/Private/Components/CanvasPanelSlot.cpp
new file mode 100644
--- /dev/null
+++ b/Engine/Source/Runtime/UMG/Private/Components/CanvasPanelSlot.cpp
@@ -0,0 +1,15 @@
+// Copyright 1998-2014 Epic Games, Inc. All Rights Reserved.
+
+#include "UMGPrivatePCH.h"
+
+/////////////////////////////////////////////////////
+// UCanvasPanelSlot
+
+UCanvasPanelSlot::UCanvasPanelSlot(const FPostConstructInitializeProperties& PCIP)
+	: Super(PCIP)
+{
+	Position = FVector2D::ZeroVector;
+	Size = FVector2D(1.0f, 1.0f);
+	HorizontalAlignment = HAlign_Left;
+	VerticalAlignment = VAlign_Top;
+}
diff --git a/Engine/Source/Runtime/UMG/Private/Components/CheckBoxComponent.cpp b/Engine/Source/Runtime/UMG/Private/Components/CheckBoxComponent.cpp
new file mode 100644
--- /dev/null
+++ b/Engine/Source/Runtime/UMG/Private/Components/CheckBoxComponent.cpp
@@ -0,0 +1,65 @@
+// Copyright 1998-2014 Epic Games, Inc. All Rights Reserved.
+
+#include "UMGPrivatePCH.h"
+
+#define LOCTEXT_NAMESPACE "UMG"
+
+/////////////////////////////////////////////////////
+// UCheckBoxComponent
+
+UCheckBoxComponent::UCheckBoxComponent(const FPostConstructInitializeProperties& PCIP)
+	: Super(PCIP)
+{
+	SCheckBox::FArguments CheckBoxDefaults;
+
+	CheckBoxStyleName = NAME_None;
+	bIsChecked = false;
+
+	HorizontalAlignment = CheckBoxDefaults._HAlign;
+	Padding = CheckBoxDefaults._Padding.Get();
+
+	ForegroundColor = CheckBoxDefaults._ForegroundColor.Get();
+	BorderBackgroundColor = CheckBoxDefaults._BorderBackgroundColor.Get();
+}
+
+TSharedRef<SWidget> UCheckBoxComponent::RebuildWidget()
+{
+	return SNew(SCheckBox)
+// 		SLATE_ATTRIBUTE( ESlateCheckBoxType::Type, Type )
+// 		SLATE_DEFAULT_SLOT( FArguments, Content )
+		//.Style( FEditorStyle::Get(), CheckBoxStyleName )
+// 		SLATE_ARGUMENT( const FCheckBoxStyle *, CheckBoxStyle )
+		.OnCheckStateChanged( BIND_UOBJECT_DELEGATE(FOnCheckStateChanged, SlateOnCheckStateChangedCallback) )
+		.IsChecked( BIND_UOBJECT_ATTRIBUTE(ESlateCheckBoxState::Type, GetCheckState) )
+		.HAlign( HorizontalAlignment )
+		.Padding( Padding )
+		.ForegroundColor( ForegroundColor )
+		.BorderBackgroundColor( BorderBackgroundColor );
+// 		SLATE_ATTRIBUTE( bool, ReadOnly )
+// 		SLATE_ARGUMENT( bool, IsFocusable )
+// 		SLATE_EVENT( FOnGetContent, OnGetMenuContent )
+}
+
+ESlateCheckBoxState::Type UCheckBoxComponent::GetCheckState() const
+{
+	return bIsChecked ? ESlateCheckBoxState::Checked : ESlateCheckBoxState::Unchecked;
+}
+
+void UCheckBoxComponent::SlateOnCheckStateChangedCallback(ESlateCheckBoxState::Type NewState)
+{
+	//@TODO: Choosing to treat Undetermined as Checked
+	const bool bWantsToBeChecked = NewState != ESlateCheckBoxState::Unchecked;
+
+	if (OnCheckStateChanged.IsBound())
+	{
+		OnCheckStateChanged.Broadcast(bWantsToBeChecked);
+	}
+	else
+	{
+		bIsChecked = bWantsToBeChecked;
+	}
+}
+
+/////////////////////////////////////////////////////
+
+#undef LOCTEXT_NAMESPACE
diff --git a/Engine/Source/Runtime/UMG/Private/Components/CircularThrobberComponent.cpp b/Engine/Source/Runtime/UMG/Private/Components/CircularThrobberComponent.cpp
new file mode 100644
--- /dev/null
+++ b/Engine/Source/Runtime/UMG/Private/Components/CircularThrobberComponent.cpp
@@ -0,0 +1,32 @@
+// Copyright 1998-2014 Epic Games, Inc. All Rights Reserved.
+
+#include "UMGPrivatePCH.h"
+
+#define LOCTEXT_NAMESPACE "UMG"
+
+/////////////////////////////////////////////////////
+// UCircularThrobberComponent
+
+UCircularThrobberComponent::UCircularThrobberComponent(const FPostConstructInitializeProperties& PCIP)
+	: Super(PCIP)
+{
+	SCircularThrobber::FArguments DefaultArgs;
+
+	NumberOfPieces = DefaultArgs._NumPieces;
+	Period = DefaultArgs._Period;
+	Radius = DefaultArgs._Radius;
+	PieceImage = *DefaultArgs._PieceImage;
+}
+
+TSharedRef<SWidget> UCircularThrobberComponent::RebuildWidget()
+{
+	return SNew(SCircularThrobber)
+		.PieceImage(&PieceImage)
+		.NumPieces(NumberOfPieces)
+		.Period(Period)
+		.Radius(Radius);
+}
+
+/////////////////////////////////////////////////////
+
+#undef LOCTEXT_NAMESPACE
diff --git a/Engine/Source/Runtime/UMG/Private/Components/ContentWidget.cpp b/Engine/Source/Runtime/UMG/Private/Components/ContentWidget.cpp
new file mode 100644
--- /dev/null
+++ b/Engine/Source/Runtime/UMG/Private/Components/ContentWidget.cpp
@@ -0,0 +1,46 @@
+// Copyright 1998-2014 Epic Games, Inc. All Rights Reserved.
+
+#include "UMGPrivatePCH.h"
+
+#if WITH_EDITOR
+#include "MessageLog.h"
+#include "UObjectToken.h"
+#endif
+
+#define LOCTEXT_NAMESPACE "UMG"
+
+/////////////////////////////////////////////////////
+// UContentWidget
+
+UContentWidget::UContentWidget(const FPostConstructInitializeProperties& PCIP)
+	: Super(PCIP)
+{
+	PrimaryComponentTick.bCanEverTick = true;
+
+	bAutoActivate = true;
+	bTickInEditor = true;
+}
+
+USlateWrapperComponent* UContentWidget::GetContent()
+{
+	return Content;
+}
+
+void UContentWidget::SetContent(USlateWrapperComponent* InContent)
+{
+	Content = InContent;
+}
+
+int32 UContentWidget::GetChildrenCount() const
+{
+	return Content != NULL ? 1 : 0;
+}
+
+USlateWrapperComponent* UContentWidget::GetChildAt(int32 Index) const
+{
+	return Content;
+}
+
+/////////////////////////////////////////////////////
+
+#undef LOCTEXT_NAMESPACE
diff --git a/Engine/Source/Runtime/UMG/Private/Components/EditableTextBlockComponent.cpp b/Engine/Source/Runtime/UMG/Private/Components/EditableTextBlockComponent.cpp
new file mode 100644
--- /dev/null
+++ b/Engine/Source/Runtime/UMG/Private/Components/EditableTextBlockComponent.cpp
@@ -0,0 +1,22 @@
+// Copyright 1998-2014 Epic Games, Inc. All Rights Reserved.
+
+#include "UMGPrivatePCH.h"
+
+#define LOCTEXT_NAMESPACE "UMG"
+
+/////////////////////////////////////////////////////
+// UEditableTextBlockComponent
+
+UEditableTextBlockComponent::UEditableTextBlockComponent(const FPostConstructInitializeProperties& PCIP)
+	: Super(PCIP)
+{
+}
+
+TSharedRef<SWidget> UEditableTextBlockComponent::RebuildWidget()
+{
+	return SNew(SEditableTextBox);
+}
+
+/////////////////////////////////////////////////////
+
+#undef LOCTEXT_NAMESPACE
diff --git a/Engine/Source/Runtime/UMG/Private/Components/GridPanelComponent.cpp b/Engine/Source/Runtime/UMG/Private/Components/GridPanelComponent.cpp
new file mode 100644
--- /dev/null
+++ b/Engine/Source/Runtime/UMG/Private/Components/GridPanelComponent.cpp
@@ -0,0 +1,115 @@
+// Copyright 1998-2014 Epic Games, Inc. All Rights Reserved.
+
+#include "UMGPrivatePCH.h"
+
+/////////////////////////////////////////////////////
+// UGridPanelComponent
+
+UGridPanelComponent::UGridPanelComponent(const FPostConstructInitializeProperties& PCIP)
+	: Super(PCIP)
+{
+}
+
+TSharedRef<SWidget> UGridPanelComponent::RebuildWidget()
+{
+	RebuildUniqueSlotTable();
+
+	TSharedRef<SGridPanel> NewGrid = SNew(SGridPanel);
+	MyGrid = NewGrid;
+
+	//TSharedPtr<SGridPanel> Grid = MyGrid.Pin();
+	//if (Grid.IsValid())
+	//{
+	//	// Create a mapping of slot names to slot indices
+	//	TMap<FName, int32> SlotNameToIndex;
+	//	for (int32 SlotIndex = 0; SlotIndex < Slots.Num(); ++SlotIndex)
+	//	{
+	//		const FGridPanelSlot& SlotConfig = Slots[SlotIndex];
+	//		SlotNameToIndex.Add(SlotConfig.SlotName, SlotIndex);
+	//	}
+	//	
+	//	// Make sure all children are assigned to a slot
+	//	TArray<USlateWrapperComponent*> WrappedComponentsInSlotOrder;
+	//	WrappedComponentsInSlotOrder.AddZeroed(Slots.Num());
+
+	//	for (int32 ComponentIndex = 0; ComponentIndex < AttachChildren.Num(); ++ComponentIndex)
+	//	{
+	//		if (USlateWrapperComponent* Wrapper = Cast<USlateWrapperComponent>(AttachChildren[ComponentIndex]))
+	//		{
+	//			if (Wrapper->IsRegistered())
+	//			{
+	//				WrappedComponentsInSlotOrder.Add(Wrapper);
+
+	//				if (int32* pSlotIndex = SlotNameToIndex.Find(Wrapper->AttachSocketName))
+	//				{
+	//					WrappedComponentsInSlotOrder[*pSlotIndex] = Wrapper;
+	//				}
+	//				else
+	//				{
+	//					FGridPanelSlot& NewSlot = *new (Slots) FGridPanelSlot();
+	//					EnsureSlotIsUnique(NewSlot);
+	//					WrappedComponentsInSlotOrder.Add(Wrapper);
+
+	//					Wrapper->AttachSocketName = NewSlot.SlotName;
+	//				}
+	//			}						
+	//		}
+	//	}
+
+	//	// Add slots to the grid panel
+	//	Grid->ClearChildren();
+	//	for (int32 SlotIndex = 0; SlotIndex < Slots.Num(); ++SlotIndex)
+	//	{
+	//		const FGridPanelSlot& SlotConfig = Slots[SlotIndex];
+
+	//		TSharedRef<SWidget> ContentWidget = SNullWidget::NullWidget;
+	//		if (WrappedComponentsInSlotOrder[SlotIndex] != NULL)
+	//		{
+	//			ContentWidget = WrappedComponentsInSlotOrder[SlotIndex]->GetWidget();
+	//		}
+
+	//		Grid->AddSlot(SlotConfig.Column, SlotConfig.Row)
+	//			.RowSpan(SlotConfig.RowSpan)
+	//			.ColumnSpan(SlotConfig.ColumnSpan)
+	//			.Padding(SlotConfig.Padding)
+	//			.HAlign(SlotConfig.HorizontalAlignment)
+	//			.VAlign(SlotConfig.VerticalAlignment)
+	//		[
+	//			ContentWidget
+	//		];
+	//	}
+	//}
+
+	return NewGrid;
+}
+
+#if WITH_EDITOR
+void UGridPanelComponent::PostEditChangeProperty(struct FPropertyChangedEvent& PropertyChangedEvent)
+{
+	RebuildUniqueSlotTable();
+}
+#endif
+
+void UGridPanelComponent::RebuildUniqueSlotTable()
+{
+	// Ensure the slots have unique names and coordinates
+	UniqueSlotCoordinates.Empty();
+
+	for (int32 SlotIndex = 0; SlotIndex < Slots.Num(); ++SlotIndex)
+	{
+		EnsureSlotIsUnique(Slots[SlotIndex]);
+	}
+}
+
+void UGridPanelComponent::EnsureSlotIsUnique(FGridPanelSlot& SlotConfig)
+{
+	FName& SlotName = SlotConfig.SlotName;
+
+	while (UniqueSlotCoordinates.Contains(SlotConfig.AsPoint()))
+	{
+		SlotConfig.Row++;
+	}
+
+	UniqueSlotCoordinates.Add(SlotConfig.AsPoint());
+	SlotName = FName(*FString::Printf(TEXT("X%d_Y%d"), SlotConfig.Column, SlotConfig.Row));
+}
diff --git a/Engine/Source/Runtime/UMG/Private/Components/HorizontalBoxComponent.cpp b/Engine/Source/Runtime/UMG/Private/Components/HorizontalBoxComponent.cpp
new file mode 100644
--- /dev/null
+++ b/Engine/Source/Runtime/UMG/Private/Components/HorizontalBoxComponent.cpp
@@ -0,0 +1,93 @@
+// Copyright 1998-2014 Epic Games, Inc. All Rights Reserved.
+
+#include "UMGPrivatePCH.h"
+
+/////////////////////////////////////////////////////
+// UHorizontalBoxComponent
+
+UHorizontalBoxComponent::UHorizontalBoxComponent(const FPostConstructInitializeProperties& PCIP)
+	: Super(PCIP)
+{
+}
+
+int32 UHorizontalBoxComponent::GetChildrenCount() const
+{
+	return Slots.Num();
+}
+
+USlateWrapperComponent* UHorizontalBoxComponent::GetChildAt(int32 Index) const
+{
+	return Slots[Index]->Content;
+}
+
+TSharedRef<SWidget> UHorizontalBoxComponent::RebuildWidget()
+{
+	TSharedRef<SHorizontalBox> NewCanvas = SNew(SHorizontalBox);
+	MyHorizontalBox = NewCanvas;
+
+	TSharedPtr<SHorizontalBox> Canvas = MyHorizontalBox.Pin();
+	if ( Canvas.IsValid() )
+	{
+		Canvas->ClearChildren();
+
+		for ( int32 SlotIndex = 0; SlotIndex < Slots.Num(); ++SlotIndex )
+		{
+			UHorizontalBoxSlot* Slot = Slots[SlotIndex];
+			if ( Slot == NULL )
+			{
+				Slots[SlotIndex] = Slot = ConstructObject<UHorizontalBoxSlot>(UHorizontalBoxSlot::StaticClass(), this);
+			}
+
+			auto& NewSlot = Canvas->AddSlot()
+				.Padding(Slot->Padding)
+				.HAlign(Slot->HorizontalAlignment)
+				.VAlign(Slot->VerticalAlignment)
+				[
+					Slot->Content == NULL ? SNullWidget::NullWidget : Slot->Content->GetWidget()
+				];
+
+			NewSlot.SizeParam = USlateWrapperComponent::ConvertSerializedSizeParamToRuntime(Slot->Size);
+		}
+	}
+
+	return NewCanvas;
+}
+
+UHorizontalBoxSlot* UHorizontalBoxComponent::AddSlot(USlateWrapperComponent* Content)
+{
+	UHorizontalBoxSlot* Slot = ConstructObject<UHorizontalBoxSlot>(UHorizontalBoxSlot::StaticClass(), this);
+	Slot->Content = Content;
+
+#if WITH_EDITOR
+	Content->Slot = Slot;
+#endif
+	
+	Slots.Add(Slot);
+
+	return Slot;
+}
+
+#if WITH_EDITOR
+void UHorizontalBoxComponent::ConnectEditorData()
+{
+	for ( UHorizontalBoxSlot* Slot : Slots )
+	{
+		Slot->Content->Slot = Slot;
+	}
+}
+
+void UHorizontalBoxComponent::PostEditChangeProperty(struct FPropertyChangedEvent& PropertyChangedEvent)
+{
+	// Ensure the slots have unique names
+	int32 SlotNumbering = 1;
+
+	TSet<FName> UniqueSlotNames;
+	for (int32 SlotIndex = 0; SlotIndex < Slots.Num(); ++SlotIndex)
+	{
+		if ( Slots[SlotIndex] == NULL )
+		{
+			Slots[SlotIndex] = ConstructObject<UHorizontalBoxSlot>(UHorizontalBoxSlot::StaticClass(), this);
+		}
+	}
+}
+#endif
diff --git a/Engine/Source/Runtime/UMG/Private/Components/HorizontalBoxSlot.cpp b/Engine/Source/Runtime/UMG/Private/Components/HorizontalBoxSlot.cpp
new file mode 100644
--- /dev/null
+++ b/Engine/Source/Runtime/UMG/Private/Components/HorizontalBoxSlot.cpp
@@ -0,0 +1,13 @@
+// Copyright 1998-2014 Epic Games, Inc. All Rights Reserved.
+
+#include "UMGPrivatePCH.h"
+
+/////////////////////////////////////////////////////
+// UHorizontalBoxSlot
+
+UHorizontalBoxSlot::UHorizontalBoxSlot(const FPostConstructInitializeProperties& PCIP)
+	: Super(PCIP)
+{
+	HorizontalAlignment = HAlign_Left;
+	VerticalAlignment = VAlign_Top;
+}
diff --git a/Engine/Source/Runtime/UMG/Private/Components/ImageComponent.cpp b/Engine/Source/Runtime/UMG/Private/Components/ImageComponent.cpp
new file mode 100644
--- /dev/null
+++ b/Engine/Source/Runtime/UMG/Private/Components/ImageComponent.cpp
@@ -0,0 +1,44 @@
+// Copyright 1998-2014 Epic Games, Inc. All Rights Reserved.
+
+#include "UMGPrivatePCH.h"
+
+#define LOCTEXT_NAMESPACE "UMG"
+
+/////////////////////////////////////////////////////
+// UImageComponent
+
+UImageComponent::UImageComponent(const FPostConstructInitializeProperties& PCIP)
+	: Super(PCIP)
+	, Image()
+	, ColorAndOpacity(FLinearColor::White)
+{
+	//@TODO: Would like to default to the checkerboard, but TextureName always wins over TextureObject :(
+//	Image = *FEditorStyle::GetDefaultBrush();
+}
+
+TSharedRef<SWidget> UImageComponent::RebuildWidget()
+{
+	return SNew(SImage)
+		.Image( BIND_UOBJECT_ATTRIBUTE(const FSlateBrush*, GetImageBrush) )
+		.ColorAndOpacity( BIND_UOBJECT_ATTRIBUTE(FSlateColor, GetColorAndOpacity) );
+}
+
+const FSlateBrush* UImageComponent::GetImageBrush() const
+{
+	if ( Image == NULL )
+	{
+		SImage::FArguments ImageDefaults;
+		return ImageDefaults._Image.Get();
+	}
+
+	return &Image->Brush;
+}
+
+FSlateColor UImageComponent::GetColorAndOpacity() const
+{
+	return ColorAndOpacity;
+}
+
+/////////////////////////////////////////////////////
+
+#undef LOCTEXT_NAMESPACE
diff --git a/Engine/Source/Runtime/UMG/Private/Components/PanelSlot.cpp b/Engine/Source/Runtime/UMG/Private/Components/PanelSlot.cpp
new file mode 100644
--- /dev/null
+++ b/Engine/Source/Runtime/UMG/Private/Components/PanelSlot.cpp
@@ -0,0 +1,11 @@
+// Copyright 1998-2014 Epic Games, Inc. All Rights Reserved.
+
+#include "UMGPrivatePCH.h"
+
+/////////////////////////////////////////////////////
+// UPanelSlot
+
+UPanelSlot::UPanelSlot(const FPostConstructInitializeProperties& PCIP)
+	: Super(PCIP)
+{
+}
\ No newline at end of file
diff --git a/Engine/Source/Runtime/UMG/Private/Components/ScrollBoxComponent.cpp b/Engine/Source/Runtime/UMG/Private/Components/ScrollBoxComponent.cpp
new file mode 100644
--- /dev/null
+++ b/Engine/Source/Runtime/UMG/Private/Components/ScrollBoxComponent.cpp
@@ -0,0 +1,22 @@
+// Copyright 1998-2014 Epic Games, Inc. All Rights Reserved.
+
+#include "UMGPrivatePCH.h"
+
+#define LOCTEXT_NAMESPACE "UMG"
+
+/////////////////////////////////////////////////////
+// UScrollBoxComponent
+
+UScrollBoxComponent::UScrollBoxComponent(const FPostConstructInitializeProperties& PCIP)
+	: Super(PCIP)
+{
+}
+
+TSharedRef<SWidget> UScrollBoxComponent::RebuildWidget()
+{
+	return SNew(SScrollBox);
+}
+
+/////////////////////////////////////////////////////
+
+#undef LOCTEXT_NAMESPACE
diff --git a/Engine/Source/Runtime/UMG/Private/Components/SlateLeafWidgetComponent.cpp b/Engine/Source/Runtime/UMG/Private/Components/SlateLeafWidgetComponent.cpp
new file mode 100644
--- /dev/null
+++ b/Engine/Source/Runtime/UMG/Private/Components/SlateLeafWidgetComponent.cpp
@@ -0,0 +1,16 @@
+// Copyright 1998-2014 Epic Games, Inc. All Rights Reserved.
+
+#include "UMGPrivatePCH.h"
+
+/////////////////////////////////////////////////////
+// USlateLeafWidgetComponent
+
+USlateLeafWidgetComponent::USlateLeafWidgetComponent(const FPostConstructInitializeProperties& PCIP)
+	: Super(PCIP)
+{
+}
+
+bool USlateLeafWidgetComponent::CanAttachAsChild(USceneComponent* ChildComponent, FName SocketName) const
+{
+	return false;
+}
diff --git a/Engine/Source/Runtime/UMG/Private/Components/SlateNonLeafWidgetComponent.cpp b/Engine/Source/Runtime/UMG/Private/Components/SlateNonLeafWidgetComponent.cpp
new file mode 100644
--- /dev/null
+++ b/Engine/Source/Runtime/UMG/Private/Components/SlateNonLeafWidgetComponent.cpp
@@ -0,0 +1,26 @@
+// Copyright 1998-2014 Epic Games, Inc. All Rights Reserved.
+
+#include "UMGPrivatePCH.h"
+
+#if WITH_EDITOR
+#include "MessageLog.h"
+#include "UObjectToken.h"
+#endif
+
+#define LOCTEXT_NAMESPACE "UMG"
+
+/////////////////////////////////////////////////////
+// USlateNonLeafWidgetComponent
+
+USlateNonLeafWidgetComponent::USlateNonLeafWidgetComponent(const FPostConstructInitializeProperties& PCIP)
+	: Super(PCIP)
+{
+	PrimaryComponentTick.bCanEverTick = true;
+
+	bAutoActivate = true;
+	bTickInEditor = true;
+}
+
+/////////////////////////////////////////////////////
+
+#undef LOCTEXT_NAMESPACE
diff --git a/Engine/Source/Runtime/UMG/Private/Components/SlateWrapperComponent.cpp b/Engine/Source/Runtime/UMG/Private/Components/SlateWrapperComponent.cpp
new file mode 100644
--- /dev/null
+++ b/Engine/Source/Runtime/UMG/Private/Components/SlateWrapperComponent.cpp
@@ -0,0 +1,70 @@
+// Copyright 1998-2014 Epic Games, Inc. All Rights Reserved.
+
+#include "UMGPrivatePCH.h"
+
+/////////////////////////////////////////////////////
+// USlateWrapperComponent
+
+USlateWrapperComponent::USlateWrapperComponent(const FPostConstructInitializeProperties& PCIP)
+	: Super(PCIP)
+{
+}
+
+void USlateWrapperComponent::OnRegister()
+{
+	Super::OnRegister();
+}
+
+void USlateWrapperComponent::OnUnregister()
+{
+	MyWidget = NULL;
+	Super::OnUnregister();
+}
+
+TSharedRef<SWidget> USlateWrapperComponent::GetWidget()
+{
+	if ( !MyWidget.IsValid() )
+	{
+		MyWidget = RebuildWidget();
+	}
+
+	return MyWidget.ToSharedRef();
+}
+
+TSharedRef<SWidget> USlateWrapperComponent::RebuildWidget()
+{
+	ensureMsg(false, TEXT("You must implement RebuildWidget() in your child class"));
+	return SNullWidget::NullWidget;
+}
+
+EVisibility USlateWrapperComponent::ConvertSerializedVisibilityToRuntime(TEnumAsByte<ESlateVisibility::Type> Input)
+{
+	switch ( Input )
+	{
+	case ESlateVisibility::Visible:
+		return EVisibility::Visible;
+	case ESlateVisibility::Collapsed:
+		return EVisibility::Collapsed;
+	case ESlateVisibility::Hidden:
+		return EVisibility::Hidden;
+	case ESlateVisibility::HitTestInvisible:
+		return EVisibility::HitTestInvisible;
+	case ESlateVisibility::SelfHitTestInvisible:
+		return EVisibility::SelfHitTestInvisible;
+	default:
+		check(false);
+		return EVisibility::Visible;
+	}
+}
+
+FSizeParam USlateWrapperComponent::ConvertSerializedSizeParamToRuntime(const FSlateChildSize& Input)
+{
+	switch ( Input.SizeRule )
+	{
+	default:
+	case ESlateSizeRule::Automatic:
+		return FAuto();
+	case ESlateSizeRule::Fill:
+		return FStretch(Input.Value);
+	}
+}
\ No newline at end of file
diff --git a/Engine/Source/Runtime/UMG/Private/Components/SpacerComponent.cpp b/Engine/Source/Runtime/UMG/Private/Components/SpacerComponent.cpp
new file mode 100644
--- /dev/null
+++ b/Engine/Source/Runtime/UMG/Private/Components/SpacerComponent.cpp
@@ -0,0 +1,29 @@
+// Copyright 1998-2014 Epic Games, Inc. All Rights Reserved.
+
+#include "UMGPrivatePCH.h"
+
+#define LOCTEXT_NAMESPACE "UMG"
+
+/////////////////////////////////////////////////////
+// USpacerComponent
+
+USpacerComponent::USpacerComponent(const FPostConstructInitializeProperties& PCIP)
+	: Super(PCIP)
+	, Size(1.0f, 1.0f)
+{
+}
+
+TSharedRef<SWidget> USpacerComponent::RebuildWidget()
+{
+	return SNew(SSpacer)
+		.Size( BIND_UOBJECT_ATTRIBUTE(FVector2D, GetSpacerSize) );
+}
+
+FVector2D USpacerComponent::GetSpacerSize() const
+{
+	return Size;
+}
+
+/////////////////////////////////////////////////////
+
+#undef LOCTEXT_NAMESPACE
diff --git a/Engine/Source/Runtime/UMG/Private/Components/SpinningImageComponent.cpp b/Engine/Source/Runtime/UMG/Private/Components/SpinningImageComponent.cpp
new file mode 100644
--- /dev/null
+++ b/Engine/Source/Runtime/UMG/Private/Components/SpinningImageComponent.cpp
@@ -0,0 +1,27 @@
+// Copyright 1998-2014 Epic Games, Inc. All Rights Reserved.
+
+#include "UMGPrivatePCH.h"
+
+#define LOCTEXT_NAMESPACE "UMG"
+
+/////////////////////////////////////////////////////
+// UImageComponent
+
+USpinningImageComponent::USpinningImageComponent(const FPostConstructInitializeProperties& PCIP)
+	: Super(PCIP)
+{
+	SSpinningImage::FArguments DefaultArgs;
+	Period = DefaultArgs._Period;
+}
+
+TSharedRef<SWidget> USpinningImageComponent::RebuildWidget()
+{
+	return SNew(SSpinningImage)
+		.Image( BIND_UOBJECT_ATTRIBUTE(const FSlateBrush*, GetImageBrush) )
+		.ColorAndOpacity( BIND_UOBJECT_ATTRIBUTE(FSlateColor, GetColorAndOpacity) )
+		.Period( Period );
+}
+
+/////////////////////////////////////////////////////
+
+#undef LOCTEXT_NAMESPACE
diff --git a/Engine/Source/Runtime/UMG/Private/Components/TextBlockComponent.cpp b/Engine/Source/Runtime/UMG/Private/Components/TextBlockComponent.cpp
new file mode 100644
--- /dev/null
+++ b/Engine/Source/Runtime/UMG/Private/Components/TextBlockComponent.cpp
@@ -0,0 +1,49 @@
+// Copyright 1998-2014 Epic Games, Inc. All Rights Reserved.
+
+#include "UMGPrivatePCH.h"
+
+#define LOCTEXT_NAMESPACE "UMG"
+
+/////////////////////////////////////////////////////
+// UTextBlockComponent
+
+UTextBlockComponent::UTextBlockComponent(const FPostConstructInitializeProperties& PCIP)
+	: Super(PCIP)
+{
+	Text = LOCTEXT("TextBlockDefaultValue", "Text Block");
+	Font = FSlateFontInfo(TEXT("Slate/Fonts/Roboto-Bold.ttf"), 24);
+	ShadowOffset = FVector2D(1.0f, 1.0f);
+	ColorAndOpacity = FLinearColor::White;
+	ShadowColorAndOpacity = FLinearColor::Transparent;
+}
+
+TSharedRef<SWidget> UTextBlockComponent::RebuildWidget()
+{
+	FString FontPath = FPaths::EngineContentDir() / Font.FontName.ToString();
+
+	return SNew(STextBlock)
+		.Text( BIND_UOBJECT_ATTRIBUTE(FString, GetText) )
+		.Font(FSlateFontInfo(FontPath, Font.Size))
+		.ColorAndOpacity( BIND_UOBJECT_ATTRIBUTE(FSlateColor, GetColorAndOpacity) )
+		.ShadowOffset(ShadowOffset)
+		.ShadowColorAndOpacity( BIND_UOBJECT_ATTRIBUTE(FLinearColor, GetShadowColorAndOpacity) );
+}
+
+FString UTextBlockComponent::GetText() const
+{
+	return Text.ToString();
+}
+
+FSlateColor UTextBlockComponent::GetColorAndOpacity() const
+{
+	return ColorAndOpacity;
+}
+
+FLinearColor UTextBlockComponent::GetShadowColorAndOpacity() const
+{
+	return ShadowColorAndOpacity;
+}
+
+/////////////////////////////////////////////////////
+
+#undef LOCTEXT_NAMESPACE
diff --git a/Engine/Source/Runtime/UMG/Private/Components/ThrobberComponent.cpp b/Engine/Source/Runtime/UMG/Private/Components/ThrobberComponent.cpp
new file mode 100644
--- /dev/null
+++ b/Engine/Source/Runtime/UMG/Private/Components/ThrobberComponent.cpp
@@ -0,0 +1,37 @@
+// Copyright 1998-2014 Epic Games, Inc. All Rights Reserved.
+
+#include "UMGPrivatePCH.h"
+
+#define LOCTEXT_NAMESPACE "UMG"
+
+/////////////////////////////////////////////////////
+// UThrobberComponent
+
+UThrobberComponent::UThrobberComponent(const FPostConstructInitializeProperties& PCIP)
+	: Super(PCIP)
+{
+	SThrobber::FArguments DefaultArgs;
+
+	PieceImage = *DefaultArgs._PieceImage;
+	NumberOfPieces = DefaultArgs._NumPieces;
+
+	bAnimateVertically = (DefaultArgs._Animate & SThrobber::Vertical) != 0;
+	bAnimateHorizontally = (DefaultArgs._Animate & SThrobber::Horizontal) != 0;
+	bAnimateOpacity = (DefaultArgs._Animate & SThrobber::Opacity) != 0;
+}
+
+TSharedRef<SWidget> UThrobberComponent::RebuildWidget()
+{
+	const int32 AnimationParams = (bAnimateVertically ? SThrobber::Vertical : 0) |
+		(bAnimateHorizontally ? SThrobber::Horizontal : 0) |
+		(bAnimateOpacity ? SThrobber::Opacity : 0);
+
+	return SNew(SThrobber)
+		.PieceImage(&PieceImage)
+		.NumPieces(NumberOfPieces)
+		.Animate(static_cast<SThrobber::EAnimation>(AnimationParams));
+}
+
+/////////////////////////////////////////////////////
+
+#undef LOCTEXT_NAMESPACE
diff --git a/Engine/Source/Runtime/UMG/Private/Components/VerticalBoxComponent.cpp b/Engine/Source/Runtime/UMG/Private/Components/VerticalBoxComponent.cpp
new file mode 100644
--- /dev/null
+++ b/Engine/Source/Runtime/UMG/Private/Components/VerticalBoxComponent.cpp
@@ -0,0 +1,93 @@
+// Copyright 1998-2014 Epic Games, Inc. All Rights Reserved.
+
+#include "UMGPrivatePCH.h"
+
+/////////////////////////////////////////////////////
+// UVerticalBoxComponent
+
+UVerticalBoxComponent::UVerticalBoxComponent(const FPostConstructInitializeProperties& PCIP)
+	: Super(PCIP)
+{
+}
+
+int32 UVerticalBoxComponent::GetChildrenCount() const
+{
+	return Slots.Num();
+}
+
+USlateWrapperComponent* UVerticalBoxComponent::GetChildAt(int32 Index) const
+{
+	return Slots[Index]->Content;
+}
+
+TSharedRef<SWidget> UVerticalBoxComponent::RebuildWidget()
+{
+	TSharedRef<SVerticalBox> NewCanvas = SNew(SVerticalBox);
+	MyVerticalBox = NewCanvas;
+
+	TSharedPtr<SVerticalBox> Canvas = MyVerticalBox.Pin();
+	if ( Canvas.IsValid() )
+	{
+		Canvas->ClearChildren();
+
+		for ( int32 SlotIndex = 0; SlotIndex < Slots.Num(); ++SlotIndex )
+		{
+			UVerticalBoxSlot* Slot = Slots[SlotIndex];
+			if ( Slot == NULL )
+			{
+				Slots[SlotIndex] = Slot = ConstructObject<UVerticalBoxSlot>(UVerticalBoxSlot::StaticClass(), this);
+			}
+
+			auto& NewSlot = Canvas->AddSlot()
+				.Padding(Slot->Padding)
+				.HAlign(Slot->HorizontalAlignment)
+				.VAlign(Slot->VerticalAlignment)
+				[
+					Slot->Content == NULL ? SNullWidget::NullWidget : Slot->Content->GetWidget()
+				];
+
+			NewSlot.SizeParam = USlateWrapperComponent::ConvertSerializedSizeParamToRuntime(Slot->Size);
+		}
+	}
+
+	return NewCanvas;
+}
+
+UVerticalBoxSlot* UVerticalBoxComponent::AddSlot(USlateWrapperComponent* Content)
+{
+	UVerticalBoxSlot* Slot = ConstructObject<UVerticalBoxSlot>(UVerticalBoxSlot::StaticClass(), this);
+	Slot->Content = Content;
+
+#if WITH_EDITOR
+	Content->Slot = Slot;
+#endif
+	
+	Slots.Add(Slot);
+
+	return Slot;
+}
+
+#if WITH_EDITOR
+void UVerticalBoxComponent::ConnectEditorData()
+{
+	for ( UVerticalBoxSlot* Slot : Slots )
+	{
+		Slot->Content->Slot = Slot;
+	}
+}
+
+void UVerticalBoxComponent::PostEditChangeProperty(struct FPropertyChangedEvent& PropertyChangedEvent)
+{
+	// Ensure the slots have unique names
+	int32 SlotNumbering = 1;
+
+	TSet<FName> UniqueSlotNames;
+	for (int32 SlotIndex = 0; SlotIndex < Slots.Num(); ++SlotIndex)
+	{
+		if ( Slots[SlotIndex] == NULL )
+		{
+			Slots[SlotIndex] = ConstructObject<UVerticalBoxSlot>(UVerticalBoxSlot::StaticClass(), this);
+		}
+	}
+}
+#endif
diff --git a/Engine/Source/Runtime/UMG/Private/Components/VerticalBoxSlot.cpp b/Engine/Source/Runtime/UMG/Private/Components/VerticalBoxSlot.cpp
new file mode 100644
--- /dev/null
+++ b/Engine/Source/Runtime/UMG/Private/Components/VerticalBoxSlot.cpp
@@ -0,0 +1,13 @@
+// Copyright 1998-2014 Epic Games, Inc. All Rights Reserved.
+
+#include "UMGPrivatePCH.h"
+
+/////////////////////////////////////////////////////
+// UVerticalBoxSlot
+
+UVerticalBoxSlot::UVerticalBoxSlot(const FPostConstructInitializeProperties& PCIP)
+	: Super(PCIP)
+{
+	HorizontalAlignment = HAlign_Left;
+	VerticalAlignment = VAlign_Top;
+}
diff --git a/Engine/Source/Runtime/UMG/Private/ContentWidget.cpp b/Engine/Source/Runtime/UMG/Private/ContentWidget.cpp
deleted file mode 100644
diff --git a/Engine/Source/Runtime/UMG/Private/EditableTextBlockComponent.cpp b/Engine/Source/Runtime/UMG/Private/EditableTextBlockComponent.cpp
deleted file mode 100644
diff --git a/Engine/Source/Runtime/UMG/Private/GridPanelComponent.cpp b/Engine/Source/Runtime/UMG/Private/GridPanelComponent.cpp
deleted file mode 100644
diff --git a/Engine/Source/Runtime/UMG/Private/HorizontalBoxComponent.cpp b/Engine/Source/Runtime/UMG/Private/HorizontalBoxComponent.cpp
deleted file mode 100644
diff --git a/Engine/Source/Runtime/UMG/Private/HorizontalBoxSlot.cpp b/Engine/Source/Runtime/UMG/Private/HorizontalBoxSlot.cpp
deleted file mode 100644
diff --git a/Engine/Source/Runtime/UMG/Private/ImageComponent.cpp b/Engine/Source/Runtime/UMG/Private/ImageComponent.cpp
deleted file mode 100644
diff --git a/Engine/Source/Runtime/UMG/Private/PanelSlot.cpp b/Engine/Source/Runtime/UMG/Private/PanelSlot.cpp
deleted file mode 100644
diff --git a/Engine/Source/Runtime/UMG/Private/ScrollBoxComponent.cpp b/Engine/Source/Runtime/UMG/Private/ScrollBoxComponent.cpp
deleted file mode 100644
diff --git a/Engine/Source/Runtime/UMG/Private/SlateLeafWidgetComponent.cpp b/Engine/Source/Runtime/UMG/Private/SlateLeafWidgetComponent.cpp
deleted file mode 100644
diff --git a/Engine/Source/Runtime/UMG/Private/SlateNonLeafWidgetComponent.cpp b/Engine/Source/Runtime/UMG/Private/SlateNonLeafWidgetComponent.cpp
deleted file mode 100644
diff --git a/Engine/Source/Runtime/UMG/Private/SlateWrapperComponent.cpp b/Engine/Source/Runtime/UMG/Private/SlateWrapperComponent.cpp
deleted file mode 100644
diff --git a/Engine/Source/Runtime/UMG/Private/SpacerComponent.cpp b/Engine/Source/Runtime/UMG/Private/SpacerComponent.cpp
deleted file mode 100644
diff --git a/Engine/Source/Runtime/UMG/Private/SpinningImageComponent.cpp b/Engine/Source/Runtime/UMG/Private/SpinningImageComponent.cpp
deleted file mode 100644
diff --git a/Engine/Source/Runtime/UMG/Private/TextBlockComponent.cpp b/Engine/Source/Runtime/UMG/Private/TextBlockComponent.cpp
deleted file mode 100644
diff --git a/Engine/Source/Runtime/UMG/Private/ThrobberComponent.cpp b/Engine/Source/Runtime/UMG/Private/ThrobberComponent.cpp
deleted file mode 100644
diff --git a/Engine/Source/Runtime/UMG/Private/UMGPrivatePCH.h b/Engine/Source/Runtime/UMG/Private/UMGPrivatePCH.h
--- a/Engine/Source/Runtime/UMG/Private/UMGPrivatePCH.h
+++ b/Engine/Source/Runtime/UMG/Private/UMGPrivatePCH.h
@@ -2,12 +2,4 @@
 
 #pragma once
 
-// Uncommenting this like causes the FTransform linkage issue
-//#include "Core.h"
-
-#include "Engine.h"
-
-DECLARE_LOG_CATEGORY_EXTERN(LogUMG, Verbose, All);
-
-#include "UMGClasses.h"
-#include "IUMGModule.h"
\ No newline at end of file
+#include "UMG.h"
\ No newline at end of file
diff --git a/Engine/Source/Runtime/UMG/Private/VerticalBoxComponent.cpp b/Engine/Source/Runtime/UMG/Private/VerticalBoxComponent.cpp
deleted file mode 100644
diff --git a/Engine/Source/Runtime/UMG/Private/VerticalBoxSlot.cpp b/Engine/Source/Runtime/UMG/Private/VerticalBoxSlot.cpp
deleted file mode 100644
diff --git a/Engine/Source/Runtime/UMG/Public/Blueprint/UserWidget.h b/Engine/Source/Runtime/UMG/Public/Blueprint/UserWidget.h
new file mode 100644
--- /dev/null
+++ b/Engine/Source/Runtime/UMG/Public/Blueprint/UserWidget.h
@@ -0,0 +1,28 @@
+// Copyright 1998-2014 Epic Games, Inc. All Rights Reserved.
+
+#pragma once
+
+#include "UserWidget.generated.h"
+
+UCLASS(Abstract, hideCategories=(Object, Actor))
+class UMG_API AUserWidget : public AActor
+{
+	GENERATED_UCLASS_BODY()
+
+	// AActor interface
+	virtual void PostActorCreated() OVERRIDE;
+	virtual void Destroyed() OVERRIDE;
+	virtual void RerunConstructionScripts() OVERRIDE;
+	// End of AActor interface
+
+	USlateWrapperComponent* GetWidgetHandle(TSharedRef<SWidget> InWidget);
+
+	TSharedRef<SWidget> GetRootWidget();
+
+private:
+	TSharedPtr<SWidget> RootWidget;
+	TMap< SWidget*, USlateWrapperComponent* > WidgetToComponent;
+
+private:
+	void RebuildWrapperWidget();
+};
diff --git a/Engine/Source/Runtime/UMG/Public/Blueprint/WidgetBlueprint.h b/Engine/Source/Runtime/UMG/Public/Blueprint/WidgetBlueprint.h
new file mode 100644
--- /dev/null
+++ b/Engine/Source/Runtime/UMG/Public/Blueprint/WidgetBlueprint.h
@@ -0,0 +1,30 @@
+// Copyright 1998-2014 Epic Games, Inc. All Rights Reserved.
+
+#pragma once
+
+#include "WidgetBlueprint.generated.h"
+
+UCLASS(dependson=(UBlueprint, AUserWidget), BlueprintType)
+class UMG_API UWidgetBlueprint : public UBlueprint
+{
+	GENERATED_UCLASS_BODY()
+
+public:
+	/** A tree of the widget templates to be created */
+	UPROPERTY()
+	class UWidgetTree* WidgetTree;
+
+	virtual void PostLoad() OVERRIDE;
+	virtual void PostLoadSubobjects(FObjectInstancingGraph* OuterInstanceGraph) OVERRIDE;
+	
+#if WITH_EDITOR
+	// UBlueprint interface
+	virtual bool SupportedByDefaultBlueprintFactory() const OVERRIDE
+	{
+		return false;
+	}
+	// End of UBlueprint interface
+
+	static bool ValidateGeneratedClass(const UClass* InClass);
+#endif
+};
\ No newline at end of file
diff --git a/Engine/Source/Runtime/UMG/Public/Blueprint/WidgetBlueprintGeneratedClass.h b/Engine/Source/Runtime/UMG/Public/Blueprint/WidgetBlueprintGeneratedClass.h
new file mode 100644
--- /dev/null
+++ b/Engine/Source/Runtime/UMG/Public/Blueprint/WidgetBlueprintGeneratedClass.h
@@ -0,0 +1,17 @@
+// Copyright 1998-2014 Epic Games, Inc. All Rights Reserved.
+
+#pragma once
+
+#include "WidgetBlueprintGeneratedClass.generated.h"
+
+UCLASS(dependson=(UWidgetBlueprint))
+class UMG_API UWidgetBlueprintGeneratedClass : public UBlueprintGeneratedClass
+{
+	GENERATED_UCLASS_BODY()
+
+	/** A tree of the widget templates to be created */
+	UPROPERTY()
+	class UWidgetTree* WidgetTree;
+
+	virtual void CreateComponentsForActor(AActor* Actor) const OVERRIDE;
+};
\ No newline at end of file
diff --git a/Engine/Source/Runtime/UMG/Public/Blueprint/WidgetTree.h b/Engine/Source/Runtime/UMG/Public/Blueprint/WidgetTree.h
new file mode 100644
--- /dev/null
+++ b/Engine/Source/Runtime/UMG/Public/Blueprint/WidgetTree.h
@@ -0,0 +1,15 @@
+// Copyright 1998-2014 Epic Games, Inc. All Rights Reserved.
+
+#pragma once
+
+#include "WidgetTree.generated.h"
+
+UCLASS()
+class UMG_API UWidgetTree : public UObject
+{
+	GENERATED_UCLASS_BODY()
+
+	/** Array of templates for widgets */
+	UPROPERTY()
+	TArray<class USlateWrapperComponent*> WidgetTemplates;
+};
diff --git a/Engine/Source/Runtime/UMG/Public/Components/BorderComponent.h b/Engine/Source/Runtime/UMG/Public/Components/BorderComponent.h
new file mode 100644
--- /dev/null
+++ b/Engine/Source/Runtime/UMG/Public/Components/BorderComponent.h
@@ -0,0 +1,77 @@
+// Copyright 1998-2014 Epic Games, Inc. All Rights Reserved.
+
+#pragma once
+
+#include "BorderComponent.generated.h"
+
+/** A border is a container widget that can contain one child widget, providing an opportunity to surround it with a border image and adjustable padding. */
+
+UCLASS(meta=(BlueprintSpawnableComponent), ClassGroup=UserInterface)
+class UMG_API UBorderComponent : public UContentWidget
+{
+	GENERATED_UCLASS_BODY()
+
+	/** Horizontal positioning of the content within the border */
+	UPROPERTY(EditAnywhere, BlueprintReadOnly, Category=Appearance)
+	TEnumAsByte<EHorizontalAlignment> HorizontalAlignment;
+
+	/** Vertical positioning of the content within the border */
+	UPROPERTY(EditAnywhere, BlueprintReadOnly, Category=Appearance)
+	TEnumAsByte<EVerticalAlignment> VerticalAlignment;
+
+	/** The padding to add around the inner content. */
+	UPROPERTY(EditAnywhere, BlueprintReadWrite, Category=Appearance)
+	FMargin ContentPadding;
+
+	//@TODO: Should these be exposed?
+	// SLATE_EVENT( FPointerEventHandler, OnMouseButtonDown )
+	// SLATE_EVENT( FPointerEventHandler, OnMouseButtonUp )
+	// SLATE_EVENT( FPointerEventHandler, OnMouseMove )
+	// SLATE_EVENT( FPointerEventHandler, OnMouseDoubleClick )
+
+	/** Image to use for the border */
+	UPROPERTY(EditAnywhere, BlueprintReadWrite, Category=Appearance, meta=( DisplayThumbnail = "true" ))
+	USlateBrushAsset* BorderBrush;
+
+	/** The scaling factor for the border */
+	UPROPERTY(EditAnywhere, BlueprintReadOnly, Category=Appearance)
+	FVector2D DesiredSizeScale;
+
+	/** The scaling factor for the border content */
+	UPROPERTY(EditAnywhere, BlueprintReadOnly, Category=Appearance)
+	FVector2D ContentScale;
+
+	/** Color and opacity multiplier of content in the border */
+	UPROPERTY(EditAnywhere, BlueprintReadWrite, Category=Appearance)
+	FLinearColor ContentColorAndOpacity;
+
+	/** Color and opacity of the actual border image */
+	UPROPERTY(EditAnywhere, BlueprintReadWrite, Category=Appearance)
+	FSlateColor BorderColorAndOpacity;
+
+	/** The foreground color of text and some glyphs that appear as the border's content. */
+	UPROPERTY(EditAnywhere, BlueprintReadWrite, Category=Appearance)
+	FSlateColor ForegroundColor;
+
+	/** Whether or not to show the disabled effect when this border is disabled */
+	UPROPERTY(EditAnywhere, BlueprintReadOnly, Category=Appearance)
+	bool bShowEffectWhenDisabled;
+
+	// UContentWidget interface
+	virtual void SetContent(USlateWrapperComponent* Content) OVERRIDE;
+	// End UContentWidget interface
+
+protected:
+	TWeakPtr<class SBorder> MyBorder;
+
+protected:
+	// USlateWrapperComponent interface
+	virtual TSharedRef<SWidget> RebuildWidget() OVERRIDE;
+	// End of USlateWrapperComponent interface
+
+	FMargin GetContentPadding() const;
+	FLinearColor GetContentColor() const;
+	const FSlateBrush* GetBorderBrush() const;
+	FSlateColor GetBorderColor() const;
+	FSlateColor GetForegroundColor() const;
+};
diff --git a/Engine/Source/Runtime/UMG/Public/Components/BoxPanelComponent.h b/Engine/Source/Runtime/UMG/Public/Components/BoxPanelComponent.h
new file mode 100644
--- /dev/null
+++ b/Engine/Source/Runtime/UMG/Public/Components/BoxPanelComponent.h
@@ -0,0 +1,77 @@
+// Copyright 1998-2014 Epic Games, Inc. All Rights Reserved.
+
+#pragma once
+
+#include "BoxPanelComponent.generated.h"
+
+USTRUCT()
+struct UMG_API FBoxPanelSlot
+{
+	GENERATED_USTRUCT_BODY()
+
+	/** Name of the slot (use this to attach children to the slot) */
+	UPROPERTY(EditAnywhere, Category=Appearance)
+	FName SlotName;
+
+	/** How much space this slot should occupy in the direction of the panel. */
+	UPROPERTY(EditAnywhere, Category=Appearance)
+	FSlateChildSize SlotSize;
+
+	/** Horizontal positioning of child within the allocated slot */
+	UPROPERTY(EditAnywhere, Category=Appearance)
+	TEnumAsByte<EHorizontalAlignment> HorizontalAlignment;
+
+	/** Vertical positioning of child within the allocated slot */
+	UPROPERTY(EditAnywhere, Category=Appearance)
+	TEnumAsByte<EVerticalAlignment> VerticalAlignment;
+
+	/** The padding to add around the child. */
+	UPROPERTY(EditAnywhere, Category=Appearance)
+	FMargin Padding;
+
+	/** Does this slot have a maximum size? */
+	UPROPERTY(EditAnywhere, Category=Appearance)
+	bool bHasMaximumSize;
+
+	/** The maximum size that this slot can be */
+	UPROPERTY(EditAnywhere, Category=Appearance, meta=(EditCondition="bHasMaximumSize"))
+	float MaxSize;
+
+	FBoxPanelSlot()
+		: SlotName(NAME_None) 
+		, SlotSize( FSlateChildSize() )
+		, HorizontalAlignment( HAlign_Fill )
+		, VerticalAlignment( VAlign_Fill )
+		, Padding( FMargin(0.0f) )
+		, bHasMaximumSize(false)
+		, MaxSize( 50.0f )
+	{
+	}
+};
+
+
+UCLASS(meta=(BlueprintSpawnableComponent), ClassGroup=UserInterface)
+class UMG_API UBoxPanelComponent : public USlateNonLeafWidgetComponent
+{
+	GENERATED_UCLASS_BODY()
+
+#if WITH_EDITOR
+	// UObject interface
+	virtual void PostEditChangeProperty(struct FPropertyChangedEvent& PropertyChangedEvent) OVERRIDE;
+	// End of UObject interface
+#endif
+
+protected:
+	UPROPERTY(EditAnywhere, Category=Slots)
+	TArray<FBoxPanelSlot> Slots;
+
+	UPROPERTY(EditAnywhere, Category=Appearance)
+	TEnumAsByte<EOrientation> Orientation;
+
+	TWeakPtr<class SConcreteBoxPanel> MyBox;
+
+protected:
+	// USlateWrapperComponent interface
+	virtual TSharedRef<SWidget> RebuildWidget() OVERRIDE;
+	// End of USlateWrapperComponent interface
+};
diff --git a/Engine/Source/Runtime/UMG/Public/Components/ButtonComponent.h b/Engine/Source/Runtime/UMG/Public/Components/ButtonComponent.h
new file mode 100644
--- /dev/null
+++ b/Engine/Source/Runtime/UMG/Public/Components/ButtonComponent.h
@@ -0,0 +1,93 @@
+// Copyright 1998-2014 Epic Games, Inc. All Rights Reserved.
+
+#pragma once
+
+#include "ButtonComponent.generated.h"
+
+DECLARE_DYNAMIC_MULTICAST_DELEGATE( FOnButtonComponentClicked );
+
+DECLARE_DYNAMIC_DELEGATE_RetVal( FText, FOnGetText );
+
+/** Buttons are clickable widgets */
+UCLASS(meta=(BlueprintSpawnableComponent), ClassGroup=UserInterface)
+class UMG_API UButtonComponent : public UContentWidget
+{
+	GENERATED_UCLASS_BODY()
+
+public:
+	/** Style of the button */
+	UPROPERTY(EditAnywhere, BlueprintReadOnly, Category=Style)
+	USlateWidgetStyleAsset* ButtonStyle;
+
+	/** Horizontal positioning of the content within the button */
+	UPROPERTY(EditAnywhere, BlueprintReadOnly, Category=Appearance)
+	TEnumAsByte<EHorizontalAlignment> HorizontalAlignment;
+
+	/** Vertical positioning of the content within the button */
+	UPROPERTY(EditAnywhere, BlueprintReadOnly, Category=Appearance)
+	TEnumAsByte<EVerticalAlignment> VerticalAlignment;
+
+	/** The padding to add around the button content. */
+	UPROPERTY(EditAnywhere, BlueprintReadWrite, Category=Appearance)
+	FMargin ContentPadding;
+
+	/** Should the button have text content inside? */
+	UPROPERTY(EditAnywhere, BlueprintReadOnly, Category=Appearance)
+	bool bDisplayText;
+
+	/** Text to display as the content of the button */
+	UPROPERTY(EditAnywhere, BlueprintReadWrite, Category=Appearance, meta=(EditCondition="bDisplayText"))
+	FText ButtonText;
+
+	/** Text to display as the content of the button */
+	UPROPERTY(EditAnywhere, BlueprintReadOnly, Category=Style, meta=( EditCondition="bDisplayText" ))
+	USlateWidgetStyleAsset* ButtonTextStyle;
+
+	/** The scaling factor for the button border */
+	UPROPERTY(EditAnywhere, BlueprintReadOnly, Category=Appearance)
+	FVector2D DesiredSizeScale;
+
+	/** The scaling factor for the button content */
+	UPROPERTY(EditAnywhere, BlueprintReadOnly, Category=Appearance)
+	FVector2D ContentScale;
+
+	/** The color multiplier for the button images */
+	UPROPERTY(EditAnywhere, BlueprintReadWrite, Category=Appearance)
+	FSlateColor ButtonColorAndOpacity;
+
+	/** The foreground color of the button */
+	UPROPERTY(EditAnywhere, BlueprintReadWrite, Category=Appearance)
+	FSlateColor ForegroundColor;
+
+	/** Called when the button is clicked */
+	UPROPERTY(BlueprintAssignable)
+	FOnButtonComponentClicked OnClicked;
+
+	/** Called when the controller has instigated damage in any way */
+	UPROPERTY(BlueprintReadWrite, Category=Appearance)
+	FOnGetText OnGetText;
+
+	/** Called when the controller has instigated damage in any way */
+	UFUNCTION(BlueprintNativeEvent, Category=Appearance)
+	FText GetText() const;
+
+	// UContentWidget interface
+	virtual void SetContent(USlateWrapperComponent* Content) OVERRIDE;
+	// End UContentWidget interface
+
+protected:
+	TAttribute<FText> TextAttribute;
+	TWeakPtr<class SButton> MyButton;
+
+protected:
+	// USlateWrapperComponent interface
+	virtual TSharedRef<SWidget> RebuildWidget() OVERRIDE;
+	// End of USlateWrapperComponent interface
+
+	FMargin GetContentPadding() const;
+	FText GetButtonText() const;
+	FSlateColor GetButtonColor() const;
+	FSlateColor GetForegroundColor() const;
+
+	virtual FReply SlateOnClickedCallback();
+};
diff --git a/Engine/Source/Runtime/UMG/Public/Components/CanvasPanelComponent.h b/Engine/Source/Runtime/UMG/Public/Components/CanvasPanelComponent.h
new file mode 100644
--- /dev/null
+++ b/Engine/Source/Runtime/UMG/Public/Components/CanvasPanelComponent.h
@@ -0,0 +1,45 @@
+// Copyright 1998-2014 Epic Games, Inc. All Rights Reserved.
+
+#pragma once
+
+#include "CanvasPanelComponent.generated.h"
+
+UCLASS(meta=(BlueprintSpawnableComponent), ClassGroup=UserInterface)
+class UMG_API UCanvasPanelComponent : public USlateNonLeafWidgetComponent
+{
+	GENERATED_UCLASS_BODY()
+
+	/** The items placed on the canvas */
+	UPROPERTY(EditAnywhere, EditInline, Category=Slots)
+	TArray<UCanvasPanelSlot*> Slots;
+
+	/** The desired size of the canvas */
+	UPROPERTY(EditAnywhere, Category=Appearance)
+	FVector2D DesiredCanvasSize;
+
+	UCanvasPanelSlot* AddSlot(USlateWrapperComponent* Content);
+
+	// USlateNonLeafWidgetComponent
+	virtual int32 GetChildrenCount() const OVERRIDE;
+	virtual USlateWrapperComponent* GetChildAt(int32 Index) const OVERRIDE;
+	// End USlateNonLeafWidgetComponent
+
+#if WITH_EDITOR
+	// UObject interface
+	virtual void PostEditChangeProperty(struct FPropertyChangedEvent& PropertyChangedEvent) OVERRIDE;
+	// End of UObject interface
+
+	// USlateWrapperComponent interface
+	virtual void ConnectEditorData() OVERRIDE;
+	// End USlateWrapperComponent interface
+#endif
+
+protected:
+
+	TWeakPtr<class SFixedSizeCanvas> MyCanvas;
+
+protected:
+	// USlateWrapperComponent interface
+	virtual TSharedRef<SWidget> RebuildWidget() OVERRIDE;
+	// End of USlateWrapperComponent interface
+};
diff --git a/Engine/Source/Runtime/UMG/Public/Components/CanvasPanelSlot.h b/Engine/Source/Runtime/UMG/Public/Components/CanvasPanelSlot.h
new file mode 100644
--- /dev/null
+++ b/Engine/Source/Runtime/UMG/Public/Components/CanvasPanelSlot.h
@@ -0,0 +1,50 @@
+// Copyright 1998-2014 Epic Games, Inc. All Rights Reserved.
+
+#pragma once
+
+#include "CanvasPanelSlot.generated.h"
+
+UCLASS()
+class UMG_API UCanvasPanelSlot : public UPanelSlot
+{
+	GENERATED_UCLASS_BODY()
+
+	/** Position. */
+	UPROPERTY(EditAnywhere, Category=Appearance)
+	FVector2D Position;
+
+	/** Size. */
+	UPROPERTY(EditAnywhere, Category=Appearance)
+	FVector2D Size;
+
+	/**
+	* Horizontal pivot position
+	*  Given a top aligned slot, where '+' represents the
+	*  anchor point defined by PositionAttr.
+	*
+	*   Left				Center				Right
+	+ _ _ _ _            _ _ + _ _          _ _ _ _ +
+	|		  |		   | 		   |	  |		    |
+	| _ _ _ _ |        | _ _ _ _ _ |	  | _ _ _ _ |
+	*
+	*  Note: FILL is NOT supported.
+	*/
+	UPROPERTY(EditAnywhere, Category=Appearance)
+	TEnumAsByte<EHorizontalAlignment> HorizontalAlignment;
+
+	/**
+	* Vertical pivot position
+	*   Given a left aligned slot, where '+' represents the
+	*   anchor point defined by PositionAttr.
+	*
+	*   Top					Center			  Bottom
+	*	+_ _ _ _ _		 _ _ _ _ _		 _ _ _ _ _
+	*	|         |		| 		  |		|		  |
+	*	|         |     +		  |		|		  |
+	*	| _ _ _ _ |		| _ _ _ _ |		+ _ _ _ _ |
+	*
+	*  Note: FILL is NOT supported.
+	*/
+	UPROPERTY(EditAnywhere, Category=Appearance)
+	TEnumAsByte<EVerticalAlignment> VerticalAlignment;
+};
\ No newline at end of file
diff --git a/Engine/Source/Runtime/UMG/Public/Components/CheckBoxComponent.h b/Engine/Source/Runtime/UMG/Public/Components/CheckBoxComponent.h
new file mode 100644
--- /dev/null
+++ b/Engine/Source/Runtime/UMG/Public/Components/CheckBoxComponent.h
@@ -0,0 +1,70 @@
+// Copyright 1998-2014 Epic Games, Inc. All Rights Reserved.
+
+#pragma once
+
+#include "CheckBoxComponent.generated.h"
+
+DECLARE_DYNAMIC_MULTICAST_DELEGATE_OneParam( FOnCheckBoxComponentStateChanged, bool, bIsChecked );
+
+/** Check box widget */
+UCLASS(meta=(BlueprintSpawnableComponent), ClassGroup=UserInterface)
+class UMG_API UCheckBoxComponent : public USlateWrapperComponent
+{
+	GENERATED_UCLASS_BODY()
+
+protected:
+	/** Style of the check box */
+	UPROPERTY(EditAnywhere, BlueprintReadOnly, Category=Appearance)
+	USlateWidgetStyleAsset* CheckBoxStyle;
+
+
+	/** Type of check box */
+	//SLATE_ATTRIBUTE( ESlateCheckBoxType::Type, Type )
+
+	/** Content to be placed next to the check box, or for a toggle button, the content to be placed inside the button */
+	//SLATE_DEFAULT_SLOT( FArguments, Content )
+
+	/** The name of the visual style for this checkbox */
+	UPROPERTY(EditAnywhere, BlueprintReadOnly, Category=Appearance)
+	FName CheckBoxStyleName;
+
+	/** The style structure (image references) for this checkbox' various states */
+	//SLATE_ARGUMENT( const FCheckBoxStyle *, CheckBoxStyle )
+
+	/** Called when the checked state has changed */
+	UPROPERTY(BlueprintAssignable)
+	FOnCheckBoxComponentStateChanged OnCheckStateChanged;
+
+	/** Whether the check box is currently in a checked state */
+	UPROPERTY(EditAnywhere, BlueprintReadWrite, Category=Appearance)
+	bool bIsChecked;
+
+	/** How the content of the toggle button should align within the given space */
+	UPROPERTY(EditAnywhere, BlueprintReadOnly, Category=Appearance)
+	TEnumAsByte<EHorizontalAlignment> HorizontalAlignment;
+
+	/** Spacing between the check box image and its content */
+	UPROPERTY(EditAnywhere, BlueprintReadOnly, Category=Appearance)
+	FMargin Padding;
+
+	/** Foreground color for the checkbox's content and parts */
+	UPROPERTY(EditAnywhere, BlueprintReadOnly, Category=Appearance)
+	FSlateColor ForegroundColor;
+
+	/** The color of the background border */
+	UPROPERTY(EditAnywhere, BlueprintReadOnly, Category=Appearance)
+	FSlateColor BorderBackgroundColor;
+
+	//SLATE_ATTRIBUTE( bool, ReadOnly )
+	//SLATE_ARGUMENT( bool, IsFocusable )
+
+	//SLATE_EVENT( FOnGetContent, OnGetMenuContent )
+
+protected:
+	// USlateWrapperComponent interface
+	virtual TSharedRef<SWidget> RebuildWidget() OVERRIDE;
+	// End of USlateWrapperComponent interface
+
+	ESlateCheckBoxState::Type GetCheckState() const;
+	void SlateOnCheckStateChangedCallback(ESlateCheckBoxState::Type NewState);
+};
diff --git a/Engine/Source/Runtime/UMG/Public/Components/CircularThrobberComponent.h b/Engine/Source/Runtime/UMG/Public/Components/CircularThrobberComponent.h
new file mode 100644
--- /dev/null
+++ b/Engine/Source/Runtime/UMG/Public/Components/CircularThrobberComponent.h
@@ -0,0 +1,35 @@
+// Copyright 1998-2014 Epic Games, Inc. All Rights Reserved.
+
+#pragma once
+
+#include "CircularThrobberComponent.generated.h"
+
+/**  A throbber widget that orients images in a spinning circle. */
+UCLASS(meta=(BlueprintSpawnableComponent), ClassGroup=UserInterface)
+class UMG_API UCircularThrobberComponent : public USlateLeafWidgetComponent
+{
+	GENERATED_UCLASS_BODY()
+
+protected:
+
+	/** How many pieces there are */
+	UPROPERTY(EditAnywhere, BlueprintReadOnly, Category=Appearance, meta=(ClampMin = "1", UIMin = "1", UIMax = "8"))
+	int32 NumberOfPieces;
+
+	/** The amount of time for a full circle (in seconds) */
+	UPROPERTY(EditAnywhere, BlueprintReadOnly, Category=Appearance)
+	float Period;
+
+	/** The radius of the circle */
+	UPROPERTY(EditAnywhere, BlueprintReadOnly, Category=Appearance)
+	float Radius;
+
+	/** Image to use for each segment of the throbber */
+	UPROPERTY(EditAnywhere, BlueprintReadOnly, Category=Appearance)
+	FSlateBrush PieceImage;
+
+protected:
+	// USlateWrapperComponent interface
+	virtual TSharedRef<SWidget> RebuildWidget() OVERRIDE;
+	// End of USlateWrapperComponent interface
+};
diff --git a/Engine/Source/Runtime/UMG/Public/Components/ContentWidget.h b/Engine/Source/Runtime/UMG/Public/Components/ContentWidget.h
new file mode 100644
--- /dev/null
+++ b/Engine/Source/Runtime/UMG/Public/Components/ContentWidget.h
@@ -0,0 +1,25 @@
+// Copyright 1998-2014 Epic Games, Inc. All Rights Reserved.
+
+#pragma once
+
+#include "ContentWidget.generated.h"
+
+UCLASS(Abstract)
+class UMG_API UContentWidget : public USlateNonLeafWidgetComponent
+{
+	GENERATED_UCLASS_BODY()
+
+	UFUNCTION()
+	virtual USlateWrapperComponent* GetContent();
+	
+	UFUNCTION()
+	virtual void SetContent(USlateWrapperComponent* InContent);
+
+	virtual int32 GetChildrenCount() const OVERRIDE;
+	virtual USlateWrapperComponent* GetChildAt(int32 Index) const OVERRIDE;
+
+protected:
+
+	UPROPERTY()
+	USlateWrapperComponent* Content;
+};
diff --git a/Engine/Source/Runtime/UMG/Public/Components/EditableTextBlockComponent.h b/Engine/Source/Runtime/UMG/Public/Components/EditableTextBlockComponent.h
new file mode 100644
--- /dev/null
+++ b/Engine/Source/Runtime/UMG/Public/Components/EditableTextBlockComponent.h
@@ -0,0 +1,20 @@
+// Copyright 1998-2014 Epic Games, Inc. All Rights Reserved.
+
+#pragma once
+
+#include "EditableTextBlockComponent.generated.h"
+
+/** Editable text box widget */
+UCLASS(meta=(BlueprintSpawnableComponent), ClassGroup=UserInterface)
+class UMG_API UEditableTextBlockComponent : public USlateWrapperComponent
+{
+	GENERATED_UCLASS_BODY()
+
+protected:
+
+
+protected:
+	// USlateWrapperComponent interface
+	virtual TSharedRef<SWidget> RebuildWidget() OVERRIDE;
+	// End of USlateWrapperComponent
+};
diff --git a/Engine/Source/Runtime/UMG/Public/Components/GridPanelComponent.h b/Engine/Source/Runtime/UMG/Public/Components/GridPanelComponent.h
new file mode 100644
--- /dev/null
+++ b/Engine/Source/Runtime/UMG/Public/Components/GridPanelComponent.h
@@ -0,0 +1,90 @@
+// Copyright 1998-2014 Epic Games, Inc. All Rights Reserved.
+
+#pragma once
+
+#include "GridPanelComponent.generated.h"
+
+
+USTRUCT()
+struct UMG_API FGridPanelSlot
+{
+	GENERATED_USTRUCT_BODY()
+
+	/** Name of the slot (use this to attach children to the slot) */
+	UPROPERTY(VisibleAnywhere, Category=Appearance)
+	FName SlotName;
+
+	/** The row position within the grid. */
+	UPROPERTY(EditAnywhere, Category=Appearance, meta=(ClampMin = "0"))
+	int32 Row;
+
+	/** The column position within the grid. */
+	UPROPERTY(EditAnywhere, Category=Appearance, meta=(ClampMin = "0"))
+	int32 Column;
+
+	/** How many rows this this slot spans over */
+	UPROPERTY(EditAnywhere, Category=Appearance, meta=(ClampMin = "1"))
+	int32 RowSpan;
+
+	/** How many columns this slot spans over */
+	UPROPERTY(EditAnywhere, Category=Appearance, meta=(ClampMin = "1"))
+	int32 ColumnSpan;
+
+	/** Horizontal positioning of child within the allocated slot */
+	UPROPERTY(EditAnywhere, Category=Appearance)
+	TEnumAsByte<EHorizontalAlignment> HorizontalAlignment;
+
+	/** Vertical positioning of child within the allocated slot */
+	UPROPERTY(EditAnywhere, Category=Appearance)
+	TEnumAsByte<EVerticalAlignment> VerticalAlignment;
+
+	/** The padding to add around the child. */
+	UPROPERTY(EditAnywhere, Category=Appearance)
+	FMargin Padding;
+
+	FIntPoint AsPoint() const
+	{
+		return FIntPoint(Column, Row);
+	}
+
+	FGridPanelSlot()
+		: SlotName(NAME_None) 
+		, Row(0)
+		, Column(0)
+		, RowSpan(1)
+		, ColumnSpan(1)
+		, HorizontalAlignment(HAlign_Fill)
+		, VerticalAlignment(VAlign_Fill)
+		, Padding(FMargin(0.0f))
+	{
+	}
+};
+
+
+UCLASS(meta=(BlueprintSpawnableComponent), ClassGroup=UserInterface)
+class UMG_API UGridPanelComponent : public USlateNonLeafWidgetComponent
+{
+	GENERATED_UCLASS_BODY()
+
+#if WITH_EDITOR
+	// UObject interface
+	virtual void PostEditChangeProperty(struct FPropertyChangedEvent& PropertyChangedEvent) OVERRIDE;
+	// End of UObject interface
+#endif
+
+protected:
+	UPROPERTY(EditAnywhere, Category=Slots)
+	TArray<FGridPanelSlot> Slots;
+
+	TWeakPtr<class SGridPanel> MyGrid;
+
+	TSet<FIntPoint> UniqueSlotCoordinates;
+
+protected:
+	// USlateWrapperComponent interface
+	virtual TSharedRef<SWidget> RebuildWidget() OVERRIDE;
+	// End of USlateWrapperComponent interface
+
+	void RebuildUniqueSlotTable();
+	void EnsureSlotIsUnique(FGridPanelSlot& SlotConfig);
+};
diff --git a/Engine/Source/Runtime/UMG/Public/Components/HorizontalBoxComponent.h b/Engine/Source/Runtime/UMG/Public/Components/HorizontalBoxComponent.h
new file mode 100644
--- /dev/null
+++ b/Engine/Source/Runtime/UMG/Public/Components/HorizontalBoxComponent.h
@@ -0,0 +1,41 @@
+// Copyright 1998-2014 Epic Games, Inc. All Rights Reserved.
+
+#pragma once
+
+#include "HorizontalBoxComponent.generated.h"
+
+UCLASS(meta=(BlueprintSpawnableComponent), ClassGroup=UserInterface)
+class UMG_API UHorizontalBoxComponent : public USlateNonLeafWidgetComponent
+{
+	GENERATED_UCLASS_BODY()
+
+	/** The items placed on the canvas */
+	UPROPERTY(EditAnywhere, EditInline, Category=Slots)
+	TArray<UHorizontalBoxSlot*> Slots;
+
+	UHorizontalBoxSlot* AddSlot(USlateWrapperComponent* Content);
+
+	// USlateNonLeafWidgetComponent
+	virtual int32 GetChildrenCount() const OVERRIDE;
+	virtual USlateWrapperComponent* GetChildAt(int32 Index) const OVERRIDE;
+	// End USlateNonLeafWidgetComponent
+
+#if WITH_EDITOR
+	// UObject interface
+	virtual void PostEditChangeProperty(struct FPropertyChangedEvent& PropertyChangedEvent) OVERRIDE;
+	// End of UObject interface
+
+	// USlateWrapperComponent interface
+	virtual void ConnectEditorData() OVERRIDE;
+	// End USlateWrapperComponent interface
+#endif
+
+protected:
+
+	TWeakPtr<class SHorizontalBox> MyHorizontalBox;
+
+protected:
+	// USlateWrapperComponent interface
+	virtual TSharedRef<SWidget> RebuildWidget() OVERRIDE;
+	// End of USlateWrapperComponent interface
+};
diff --git a/Engine/Source/Runtime/UMG/Public/Components/HorizontalBoxSlot.h b/Engine/Source/Runtime/UMG/Public/Components/HorizontalBoxSlot.h
new file mode 100644
--- /dev/null
+++ b/Engine/Source/Runtime/UMG/Public/Components/HorizontalBoxSlot.h
@@ -0,0 +1,26 @@
+// Copyright 1998-2014 Epic Games, Inc. All Rights Reserved.
+
+#pragma once
+
+#include "SlateWrapperTypes.h"
+
+#include "HorizontalBoxSlot.generated.h"
+
+UCLASS()
+class UMG_API UHorizontalBoxSlot : public UPanelSlot
+{
+	GENERATED_UCLASS_BODY()
+	
+	UPROPERTY(EditAnywhere, Category=Appearance)
+	FMargin Padding;
+
+	/** How much space this slot should occupy in the direction of the panel. */
+	UPROPERTY(EditAnywhere, Category=Appearance)
+	FSlateChildSize Size;
+
+	UPROPERTY(EditAnywhere, Category=Appearance)
+	TEnumAsByte<EHorizontalAlignment> HorizontalAlignment;
+
+	UPROPERTY(EditAnywhere, Category=Appearance)
+	TEnumAsByte<EVerticalAlignment> VerticalAlignment;
+};
\ No newline at end of file
diff --git a/Engine/Source/Runtime/UMG/Public/Components/ImageComponent.h b/Engine/Source/Runtime/UMG/Public/Components/ImageComponent.h
new file mode 100644
--- /dev/null
+++ b/Engine/Source/Runtime/UMG/Public/Components/ImageComponent.h
@@ -0,0 +1,28 @@
+// Copyright 1998-2014 Epic Games, Inc. All Rights Reserved.
+
+#pragma once
+
+#include "ImageComponent.generated.h"
+
+/** Image widget */
+UCLASS(meta=(BlueprintSpawnableComponent), ClassGroup=UserInterface)
+class UMG_API UImageComponent : public USlateLeafWidgetComponent
+{
+	GENERATED_UCLASS_BODY()
+
+	/** Image to draw */
+	UPROPERTY(EditAnywhere, BlueprintReadWrite, Category=Appearance, meta=( DisplayThumbnail = "true" ))
+	USlateBrushAsset* Image;
+
+	/** Color and opacity */
+	UPROPERTY(EditAnywhere, BlueprintReadWrite, Category=Appearance)
+	FSlateColor ColorAndOpacity;
+
+protected:
+	// USlateWrapperComponent interface
+	virtual TSharedRef<SWidget> RebuildWidget() OVERRIDE;
+	// End of USlateWrapperComponent interface
+
+	const FSlateBrush* GetImageBrush() const;
+	FSlateColor GetColorAndOpacity() const;
+};
diff --git a/Engine/Source/Runtime/UMG/Public/Components/PanelSlot.h b/Engine/Source/Runtime/UMG/Public/Components/PanelSlot.h
new file mode 100644
--- /dev/null
+++ b/Engine/Source/Runtime/UMG/Public/Components/PanelSlot.h
@@ -0,0 +1,14 @@
+// Copyright 1998-2014 Epic Games, Inc. All Rights Reserved.
+
+#pragma once
+
+#include "PanelSlot.generated.h"
+
+UCLASS()
+class UMG_API UPanelSlot : public UObject
+{
+	GENERATED_UCLASS_BODY()
+
+	UPROPERTY()
+	class USlateWrapperComponent* Content;
+};
diff --git a/Engine/Source/Runtime/UMG/Public/Components/ScrollBoxComponent.h b/Engine/Source/Runtime/UMG/Public/Components/ScrollBoxComponent.h
new file mode 100644
--- /dev/null
+++ b/Engine/Source/Runtime/UMG/Public/Components/ScrollBoxComponent.h
@@ -0,0 +1,18 @@
+// Copyright 1998-2014 Epic Games, Inc. All Rights Reserved.
+
+#pragma once
+
+#include "ScrollBoxComponent.generated.h"
+
+UCLASS(meta=(BlueprintSpawnableComponent), ClassGroup=UserInterface)
+class UMG_API UScrollBoxComponent : public USlateWrapperComponent
+{
+	GENERATED_UCLASS_BODY()
+
+protected:
+
+protected:
+	// USlateWrapperComponent interface
+	virtual TSharedRef<SWidget> RebuildWidget() OVERRIDE;
+	// End of USlateWrapperComponent interface
+};
diff --git a/Engine/Source/Runtime/UMG/Public/Components/SlateLeafWidgetComponent.h b/Engine/Source/Runtime/UMG/Public/Components/SlateLeafWidgetComponent.h
new file mode 100644
--- /dev/null
+++ b/Engine/Source/Runtime/UMG/Public/Components/SlateLeafWidgetComponent.h
@@ -0,0 +1,16 @@
+// Copyright 1998-2014 Epic Games, Inc. All Rights Reserved.
+
+#pragma once
+
+#include "SlateLeafWidgetComponent.generated.h"
+
+/** Leaf widget that cannot have any children */
+UCLASS(Abstract)
+class UMG_API USlateLeafWidgetComponent: public USlateWrapperComponent
+{
+	GENERATED_UCLASS_BODY()
+
+	// USceneComponent interface
+	virtual bool CanAttachAsChild(USceneComponent* ChildComponent, FName SocketName) const;
+	// End of USceneComponent interface
+};
diff --git a/Engine/Source/Runtime/UMG/Public/Components/SlateNonLeafWidgetComponent.h b/Engine/Source/Runtime/UMG/Public/Components/SlateNonLeafWidgetComponent.h
new file mode 100644
--- /dev/null
+++ b/Engine/Source/Runtime/UMG/Public/Components/SlateNonLeafWidgetComponent.h
@@ -0,0 +1,16 @@
+// Copyright 1998-2014 Epic Games, Inc. All Rights Reserved.
+
+#pragma once
+
+#include "SlateNonLeafWidgetComponent.generated.h"
+
+UCLASS(Abstract)
+class UMG_API USlateNonLeafWidgetComponent : public USlateWrapperComponent
+{
+	GENERATED_UCLASS_BODY()
+
+	virtual int32 GetChildrenCount() const { return 0; }
+	virtual USlateWrapperComponent* GetChildAt(int32 Index) const { return NULL; }
+
+	virtual bool AddChild(USlateWrapperComponent* Child) { return false; }
+};
diff --git a/Engine/Source/Runtime/UMG/Public/Components/SlateWrapperComponent.h b/Engine/Source/Runtime/UMG/Public/Components/SlateWrapperComponent.h
new file mode 100644
--- /dev/null
+++ b/Engine/Source/Runtime/UMG/Public/Components/SlateWrapperComponent.h
@@ -0,0 +1,40 @@
+// Copyright 1998-2014 Epic Games, Inc. All Rights Reserved.
+
+#pragma once
+
+#include "SlateWrapperTypes.h"
+#include "SlateWrapperComponent.generated.h"
+
+/** This is the base class for all Slate widget wrapper components */
+UCLASS(Abstract, hideCategories=(Activation, "Components|Activation", ComponentReplication, LOD, Rendering))
+class UMG_API USlateWrapperComponent : public UActorComponent
+{
+	GENERATED_UCLASS_BODY()
+
+#if WITH_EDITOR
+	UPROPERTY(Transient, EditInline, EditAnywhere, BlueprintReadWrite, Category=Layout, meta=( ShowOnlyInnerProperties ))
+	UPanelSlot* Slot;
+#endif
+
+	// UActorComponent interface
+	virtual void OnRegister() OVERRIDE;
+	virtual void OnUnregister() OVERRIDE;
+	// End of UActorComponent interface
+
+	TSharedRef<SWidget> GetWidget();
+
+#if WITH_EDITOR
+	virtual void ConnectEditorData() { }
+#endif
+
+	// Utility methods
+	//@TODO: Should move elsewhere
+	static EVisibility ConvertSerializedVisibilityToRuntime(TEnumAsByte<ESlateVisibility::Type> Input);
+	static FSizeParam ConvertSerializedSizeParamToRuntime(const FSlateChildSize& Input);
+
+protected:
+	virtual TSharedRef<SWidget> RebuildWidget();
+
+protected:
+	TSharedPtr<SWidget> MyWidget;
+};
diff --git a/Engine/Source/Runtime/UMG/Public/Components/SlateWrapperTypes.h b/Engine/Source/Runtime/UMG/Public/Components/SlateWrapperTypes.h
new file mode 100644
--- /dev/null
+++ b/Engine/Source/Runtime/UMG/Public/Components/SlateWrapperTypes.h
@@ -0,0 +1,62 @@
+// Copyright 1998-2014 Epic Games, Inc. All Rights Reserved.
+
+#pragma once
+
+#include "SlateWrapperTypes.generated.h"
+
+
+#define BIND_UOBJECT_ATTRIBUTE(Type, Function) \
+	TAttribute<Type>::Create( TAttribute<Type>::FGetter::CreateUObject( this, &ThisClass::Function ) )
+
+#define BIND_UOBJECT_DELEGATE(Type, Function) \
+	Type::CreateUObject( this, &ThisClass::Function )
+
+/** Is an entity visible? */
+UENUM()
+namespace ESlateVisibility
+{
+	enum Type
+	{
+		/** Default widget visibility - visible and can interactive with the cursor */
+		Visible,
+		/** Not visible and takes up no space in the layout; can never be clicked on because it takes up no space. */
+		Collapsed,
+		/** Not visible, but occupies layout space. Not interactive for obvious reasons. */
+		Hidden,
+		/** Visible to the user, but only as art. The cursors hit tests will never see this widget. */
+		HitTestInvisible,
+		/** Same as HitTestInvisible, but doesn't apply to child widgets. */
+		SelfHitTestInvisible
+	};
+}
+
+UENUM()
+namespace ESlateSizeRule
+{
+	enum Type
+	{
+		/** The container will size to fit the needs of the child widgets */
+		Automatic,
+		/** The container will fill the percentage of the container based on the Value 0..1 */
+		Fill
+	};
+}
+
+USTRUCT()
+struct FSlateChildSize
+{
+	GENERATED_USTRUCT_BODY()
+
+	//editcondition="SizeRule=ESlateSizeRule::Fill"
+	UPROPERTY(EditAnywhere, Category=Appearance, meta=(ClampMin="0", ClampMax="1"))
+	float Value;
+
+	UPROPERTY(EditAnywhere, Category=Appearance)
+	TEnumAsByte<ESlateSizeRule::Type> SizeRule;
+
+	FSlateChildSize()
+		: Value(1.0f)
+		, SizeRule(ESlateSizeRule::Fill)
+	{
+	}
+};
diff --git a/Engine/Source/Runtime/UMG/Public/Components/SpacerComponent.h b/Engine/Source/Runtime/UMG/Public/Components/SpacerComponent.h
new file mode 100644
--- /dev/null
+++ b/Engine/Source/Runtime/UMG/Public/Components/SpacerComponent.h
@@ -0,0 +1,23 @@
+// Copyright 1998-2014 Epic Games, Inc. All Rights Reserved.
+
+#pragma once
+
+#include "SpacerComponent.generated.h"
+
+/** A spacer widget; it does not have a visual representation, and just provides padding between other widgets. */
+UCLASS(meta=(BlueprintSpawnableComponent), ClassGroup=UserInterface)
+class UMG_API USpacerComponent : public USlateLeafWidgetComponent
+{
+	GENERATED_UCLASS_BODY()
+
+	/** The size of the spacer */
+	UPROPERTY(EditAnywhere, BlueprintReadWrite, Category=Appearance)
+	FVector2D Size;
+
+protected:
+	// USlateWrapperComponent interface
+	virtual TSharedRef<SWidget> RebuildWidget() OVERRIDE;
+	// End of USlateWrapperComponent interface
+
+	FVector2D GetSpacerSize() const;
+};
diff --git a/Engine/Source/Runtime/UMG/Public/Components/SpinningImageComponent.h b/Engine/Source/Runtime/UMG/Public/Components/SpinningImageComponent.h
new file mode 100644
--- /dev/null
+++ b/Engine/Source/Runtime/UMG/Public/Components/SpinningImageComponent.h
@@ -0,0 +1,21 @@
+// Copyright 1998-2014 Epic Games, Inc. All Rights Reserved.
+
+#pragma once
+
+#include "SpinningImageComponent.generated.h"
+
+/** A widget that displays a spinning image. */
+UCLASS(meta=(BlueprintSpawnableComponent), ClassGroup=UserInterface)
+class UMG_API USpinningImageComponent : public UImageComponent
+{
+	GENERATED_UCLASS_BODY()
+
+	/** The amount of time in seconds for a full rotation */
+	UPROPERTY(EditAnywhere, BlueprintReadOnly, Category=Appearance)
+	float Period;
+
+protected:
+	// USlateWrapperComponent interface
+	virtual TSharedRef<SWidget> RebuildWidget() OVERRIDE;
+	// End of USlateWrapperComponent interface
+};
diff --git a/Engine/Source/Runtime/UMG/Public/Components/TextBlockComponent.h b/Engine/Source/Runtime/UMG/Public/Components/TextBlockComponent.h
new file mode 100644
--- /dev/null
+++ b/Engine/Source/Runtime/UMG/Public/Components/TextBlockComponent.h
@@ -0,0 +1,37 @@
+// Copyright 1998-2014 Epic Games, Inc. All Rights Reserved.
+
+#pragma once
+
+#include "TextBlockComponent.generated.h"
+
+UCLASS(meta=(BlueprintSpawnableComponent), ClassGroup=UserInterface)
+class UMG_API UTextBlockComponent : public USlateLeafWidgetComponent
+{
+	GENERATED_UCLASS_BODY()
+
+public:
+	UPROPERTY(EditAnywhere, Category=TextBlock)
+	FText Text;
+
+	UPROPERTY(EditAnywhere, Category=TextBlock)
+	FSlateFontInfo Font;
+
+	UPROPERTY(EditAnywhere, Category=TextBlock)
+	FVector2D ShadowOffset;
+
+	//@TODO: Expose FSlateColor
+	UPROPERTY(EditAnywhere, Category=TextBlock)
+	FLinearColor ColorAndOpacity;
+
+	UPROPERTY(EditAnywhere, Category=TextBlock)
+	FLinearColor ShadowColorAndOpacity;
+
+protected:
+	// USlateWrapperComponent interface
+	virtual TSharedRef<SWidget> RebuildWidget() OVERRIDE;
+	// End of USlateWrapperComponent interface
+
+	FString GetText() const;
+	FSlateColor GetColorAndOpacity() const;
+	FLinearColor GetShadowColorAndOpacity() const;
+};
diff --git a/Engine/Source/Runtime/UMG/Public/Components/ThrobberComponent.h b/Engine/Source/Runtime/UMG/Public/Components/ThrobberComponent.h
new file mode 100644
--- /dev/null
+++ b/Engine/Source/Runtime/UMG/Public/Components/ThrobberComponent.h
@@ -0,0 +1,37 @@
+// Copyright 1998-2014 Epic Games, Inc. All Rights Reserved.
+
+#pragma once
+
+#include "ThrobberComponent.generated.h"
+
+/** A Throbber widget that shows several zooming circles in a row. */
+UCLASS(meta=(BlueprintSpawnableComponent), ClassGroup=UserInterface)
+class UMG_API UThrobberComponent : public USlateLeafWidgetComponent
+{
+	GENERATED_UCLASS_BODY()
+
+	/** How many pieces there are */
+	UPROPERTY(EditAnywhere, BlueprintReadOnly, Category=Appearance, meta=(ClampMin = "1", UIMin = "1", UIMax = "8"))
+	int32 NumberOfPieces;
+
+	/** Should the pieces animate horizontally? */
+	UPROPERTY(EditAnywhere, BlueprintReadWrite, Category=Appearance)
+	bool bAnimateHorizontally;
+
+	/** Should the pieces animate vertically? */
+	UPROPERTY(EditAnywhere, BlueprintReadWrite, Category=Appearance)
+	bool bAnimateVertically;
+
+	/** Should the pieces animate their opacity? */
+	UPROPERTY(EditAnywhere, BlueprintReadWrite, Category=Appearance)
+	bool bAnimateOpacity;
+
+	/** Image to use for each segment of the throbber */
+	UPROPERTY(EditAnywhere, BlueprintReadOnly, Category=Appearance)
+	FSlateBrush PieceImage;
+
+protected:
+	// USlateWrapperComponent interface
+	virtual TSharedRef<SWidget> RebuildWidget() OVERRIDE;
+	// End of USlateWrapperComponent interface
+};
diff --git a/Engine/Source/Runtime/UMG/Public/Components/VerticalBoxComponent.h b/Engine/Source/Runtime/UMG/Public/Components/VerticalBoxComponent.h
new file mode 100644
--- /dev/null
+++ b/Engine/Source/Runtime/UMG/Public/Components/VerticalBoxComponent.h
@@ -0,0 +1,46 @@
+// Copyright 1998-2014 Epic Games, Inc. All Rights Reserved.
+
+#pragma once
+
+#include "VerticalBoxComponent.generated.h"
+
+UCLASS(meta=(BlueprintSpawnableComponent), ClassGroup=UserInterface)
+class UMG_API UVerticalBoxComponent : public USlateNonLeafWidgetComponent
+{
+	GENERATED_UCLASS_BODY()
+
+	/** The items placed on the canvas */
+	UPROPERTY(EditAnywhere, EditInline, Category=Slots)
+	TArray<UVerticalBoxSlot*> Slots;
+
+	UVerticalBoxSlot* AddSlot(USlateWrapperComponent* Content);
+
+	// USlateNonLeafWidgetComponent
+	virtual int32 GetChildrenCount() const OVERRIDE;
+	virtual USlateWrapperComponent* GetChildAt(int32 Index) const OVERRIDE;
+	virtual bool AddChild(USlateWrapperComponent* Child) OVERRIDE
+	{
+		AddSlot(Child);
+		return true;
+	}
+	// End USlateNonLeafWidgetComponent
+
+#if WITH_EDITOR
+	// UObject interface
+	virtual void PostEditChangeProperty(struct FPropertyChangedEvent& PropertyChangedEvent) OVERRIDE;
+	// End of UObject interface
+
+	// USlateWrapperComponent interface
+	virtual void ConnectEditorData() OVERRIDE;
+	// End USlateWrapperComponent interface
+#endif
+
+protected:
+
+	TWeakPtr<class SVerticalBox> MyVerticalBox;
+
+protected:
+	// USlateWrapperComponent interface
+	virtual TSharedRef<SWidget> RebuildWidget() OVERRIDE;
+	// End of USlateWrapperComponent interface
+};
diff --git a/Engine/Source/Runtime/UMG/Public/Components/VerticalBoxSlot.h b/Engine/Source/Runtime/UMG/Public/Components/VerticalBoxSlot.h
new file mode 100644
--- /dev/null
+++ b/Engine/Source/Runtime/UMG/Public/Components/VerticalBoxSlot.h
@@ -0,0 +1,26 @@
+// Copyright 1998-2014 Epic Games, Inc. All Rights Reserved.
+
+#pragma once
+
+#include "SlateWrapperTypes.h"
+
+#include "VerticalBoxSlot.generated.h"
+
+UCLASS()
+class UMG_API UVerticalBoxSlot : public UPanelSlot
+{
+	GENERATED_UCLASS_BODY()
+	
+	UPROPERTY(EditAnywhere, Category=Appearance)
+	FMargin Padding;
+
+	/** How much space this slot should occupy in the direction of the panel. */
+	UPROPERTY(EditAnywhere, Category=Appearance)
+	FSlateChildSize Size;
+
+	UPROPERTY(EditAnywhere, Category=Appearance)
+	TEnumAsByte<EHorizontalAlignment> HorizontalAlignment;
+
+	UPROPERTY(EditAnywhere, Category=Appearance)
+	TEnumAsByte<EVerticalAlignment> VerticalAlignment;
+};
diff --git a/Engine/Source/Runtime/UMG/Public/UMG.h b/Engine/Source/Runtime/UMG/Public/UMG.h
new file mode 100644
--- /dev/null
+++ b/Engine/Source/Runtime/UMG/Public/UMG.h
@@ -0,0 +1,68 @@
+// Copyright 1998-2014 Epic Games, Inc. All Rights Reserved.
+
+/*=============================================================================
+	UMG.h: UMG module public header file.
+=============================================================================*/
+
+#pragma once
+
+
+/* Dependencies
+ *****************************************************************************/
+
+#include "Core.h"
+#include "CoreUObject.h"
+#include "Engine.h"
+#include "InputCore.h"
+#include "Slate.h"
+
+/* Components
+ *****************************************************************************/
+
+#include "Components/SlateWrapperTypes.h"
+
+#include "Components/PanelSlot.h"
+
+#include "Components/SlateWrapperComponent.h"
+
+#include "Components/SlateLeafWidgetComponent.h"
+#include "Components/SlateNonLeafWidgetComponent.h"
+
+#include "Components/ContentWidget.h"
+
+#include "Components/BoxPanelComponent.h"
+
+#include "Components/CanvasPanelSlot.h"
+#include "Components/CanvasPanelComponent.h"
+
+#include "Components/HorizontalBoxSlot.h"
+#include "Components/HorizontalBoxComponent.h"
+
+#include "Components/VerticalBoxSlot.h"
+#include "Components/VerticalBoxComponent.h"
+
+#include "Components/BorderComponent.h"
+#include "Components/ButtonComponent.h"
+#include "Components/CheckBoxComponent.h"
+#include "Components/CircularThrobberComponent.h"
+#include "Components/EditableTextBlockComponent.h"
+#include "Components/GridPanelComponent.h"
+#include "Components/ImageComponent.h"
+#include "Components/ScrollBoxComponent.h"
+#include "Components/SpacerComponent.h"
+#include "Components/SpinningImageComponent.h"
+#include "Components/TextBlockComponent.h"
+#include "Components/ThrobberComponent.h"
+
+/* Blueprint
+ *****************************************************************************/
+
+#include "Blueprint/WidgetTree.h"
+#include "Blueprint/UserWidget.h"
+#include "Blueprint/WidgetBlueprint.h"
+#include "Blueprint/WidgetBlueprintGeneratedClass.h"
+ 
+/* Interfaces
+ *****************************************************************************/
+
+#include "IUMGModule.h"
diff --git a/Engine/Source/Runtime/UMG/UMG.Build.cs b/Engine/Source/Runtime/UMG/UMG.Build.cs
--- a/Engine/Source/Runtime/UMG/UMG.Build.cs
+++ b/Engine/Source/Runtime/UMG/UMG.Build.cs
@@ -6,22 +6,24 @@ public class UMG : ModuleRules
 {
 	public UMG(TargetInfo Target)
 	{
+        PrivateIncludePaths.AddRange(
+            new string[] {
+                "Runtime/UMG/Private" // For PCH includes (because they don't work with relative paths, yet)
+            });
+
 		PublicDependencyModuleNames.AddRange(
 			new string[] {
 				"Core",
 				"CoreUObject",
-				"Engine",
-				"ShaderCore",
-				"RenderCore",
+                "Engine",
                 "InputCore",
-				"RHI",
+				"Slate",
 			}
-			);
+		);
 
 		PrivateDependencyModuleNames.AddRange(
 			new string[] {
-				"Slate"
 			}
-			);
+		);
 	}
 }
