UMG - Fleshing out a static editable textblock.  Better default colors.  More widgets support being dropped into.

diff --git a/Engine/Source/Runtime/UMG/Private/Components/ButtonComponent.cpp b/Engine/Source/Runtime/UMG/Private/Components/ButtonComponent.cpp
--- a/Engine/Source/Runtime/UMG/Private/Components/ButtonComponent.cpp
+++ b/Engine/Source/Runtime/UMG/Private/Components/ButtonComponent.cpp
@@ -24,8 +24,8 @@ UButtonComponent::UButtonComponent(const FPostConstructInitializeProperties& PCI
 	DesiredSizeScale = ButtonDefaults._DesiredSizeScale.Get();
 	ContentScale = ButtonDefaults._ContentScale.Get();
 
-	ButtonColorAndOpacity = ButtonDefaults._ButtonColorAndOpacity.Get();
-	ForegroundColor = ButtonDefaults._ForegroundColor.Get();
+	ButtonColorAndOpacity = FLinearColor::White;
+	ForegroundColor = FLinearColor::Black;
 }
 
 TSharedRef<SWidget> UButtonComponent::RebuildWidget()
diff --git a/Engine/Source/Runtime/UMG/Private/Components/ContentWidget.cpp b/Engine/Source/Runtime/UMG/Private/Components/ContentWidget.cpp
--- a/Engine/Source/Runtime/UMG/Private/Components/ContentWidget.cpp
+++ b/Engine/Source/Runtime/UMG/Private/Components/ContentWidget.cpp
@@ -41,6 +41,12 @@ USlateWrapperComponent* UContentWidget::GetChildAt(int32 Index) const
 	return Content;
 }
 
+bool UContentWidget::AddChild(USlateWrapperComponent* InContent)
+{
+	Content = InContent;
+	return true;
+}
+
 /////////////////////////////////////////////////////
 
 #undef LOCTEXT_NAMESPACE
diff --git a/Engine/Source/Runtime/UMG/Private/Components/EditableTextBlockComponent.cpp b/Engine/Source/Runtime/UMG/Private/Components/EditableTextBlockComponent.cpp
--- a/Engine/Source/Runtime/UMG/Private/Components/EditableTextBlockComponent.cpp
+++ b/Engine/Source/Runtime/UMG/Private/Components/EditableTextBlockComponent.cpp
@@ -10,11 +10,46 @@
 UEditableTextBlockComponent::UEditableTextBlockComponent(const FPostConstructInitializeProperties& PCIP)
 	: Super(PCIP)
 {
+	ForegroundColor = FLinearColor::Black;
+	BackgroundColor = FLinearColor::White;
+	ReadOnlyForegroundColor = FLinearColor::Black;
+
+	// HACK Special font initialization hack since there are no font assets yet for slate.
+	Font = FSlateFontInfo(TEXT("Slate/Fonts/Roboto-Bold.ttf"), 12);
+
+	// Grab other defaults from slate arguments.
+	SEditableTextBox::FArguments Defaults;
+	IsReadOnly = Defaults._IsReadOnly.Get();
+	IsPassword = Defaults._IsReadOnly.Get();
+	MinimumDesiredWidth = Defaults._MinDesiredWidth.Get();
+	Padding = Defaults._Padding.Get();
+	IsCaretMovedWhenGainFocus = Defaults._IsCaretMovedWhenGainFocus.Get();
+	SelectAllTextWhenFocused = Defaults._SelectAllTextWhenFocused.Get();
+	RevertTextOnEscape = Defaults._RevertTextOnEscape.Get();
+	ClearKeyboardFocusOnCommit = Defaults._ClearKeyboardFocusOnCommit.Get();
+	SelectAllTextOnCommit = Defaults._SelectAllTextOnCommit.Get();
 }
 
 TSharedRef<SWidget> UEditableTextBlockComponent::RebuildWidget()
 {
-	return SNew(SEditableTextBox);
+	FString FontPath = FPaths::EngineContentDir() / Font.FontName.ToString();
+
+	return SNew(SEditableTextBox)
+		.Text(Text)
+		.HintText(HintText)
+		.Font(FSlateFontInfo(FontPath, Font.Size))
+		.ForegroundColor(ForegroundColor)
+		.BackgroundColor(BackgroundColor)
+		.ReadOnlyForegroundColor(ReadOnlyForegroundColor)
+		.IsReadOnly(IsReadOnly)
+		.IsPassword(IsPassword)
+		.MinDesiredWidth(MinimumDesiredWidth)
+		.Padding(Padding)
+		.IsCaretMovedWhenGainFocus(IsCaretMovedWhenGainFocus)
+		.SelectAllTextWhenFocused(SelectAllTextWhenFocused)
+		.RevertTextOnEscape(RevertTextOnEscape)
+		.ClearKeyboardFocusOnCommit(ClearKeyboardFocusOnCommit)
+		.SelectAllTextOnCommit(SelectAllTextOnCommit);
 }
 
 /////////////////////////////////////////////////////
diff --git a/Engine/Source/Runtime/UMG/Private/Components/HorizontalBoxComponent.cpp b/Engine/Source/Runtime/UMG/Private/Components/HorizontalBoxComponent.cpp
--- a/Engine/Source/Runtime/UMG/Private/Components/HorizontalBoxComponent.cpp
+++ b/Engine/Source/Runtime/UMG/Private/Components/HorizontalBoxComponent.cpp
@@ -20,6 +20,12 @@ USlateWrapperComponent* UHorizontalBoxComponent::GetChildAt(int32 Index) const
 	return Slots[Index]->Content;
 }
 
+bool UHorizontalBoxComponent::AddChild(USlateWrapperComponent* Child)
+{
+	AddSlot(Child);
+	return true;
+}
+
 TSharedRef<SWidget> UHorizontalBoxComponent::RebuildWidget()
 {
 	TSharedRef<SHorizontalBox> NewCanvas = SNew(SHorizontalBox);
diff --git a/Engine/Source/Runtime/UMG/Private/Components/SlateWrapperComponent.cpp b/Engine/Source/Runtime/UMG/Private/Components/SlateWrapperComponent.cpp
--- a/Engine/Source/Runtime/UMG/Private/Components/SlateWrapperComponent.cpp
+++ b/Engine/Source/Runtime/UMG/Private/Components/SlateWrapperComponent.cpp
@@ -8,6 +8,21 @@
 USlateWrapperComponent::USlateWrapperComponent(const FPostConstructInitializeProperties& PCIP)
 	: Super(PCIP)
 {
+	bIsEnabled = true;
+}
+
+bool USlateWrapperComponent::GetIsEnabled() const
+{
+	return MyWidget.IsValid() ? MyWidget->IsEnabled() : bIsEnabled;
+}
+
+void USlateWrapperComponent::SetIsEnabled(bool bInIsEnabled)
+{
+	bIsEnabled = bInIsEnabled;
+	if ( MyWidget.IsValid() )
+	{
+		MyWidget->SetEnabled(bInIsEnabled);
+	}
 }
 
 void USlateWrapperComponent::OnRegister()
@@ -26,6 +41,7 @@ TSharedRef<SWidget> USlateWrapperComponent::GetWidget()
 	if ( !MyWidget.IsValid() )
 	{
 		MyWidget = RebuildWidget();
+		MyWidget->SetEnabled(bIsEnabled);
 	}
 
 	return MyWidget.ToSharedRef();
diff --git a/Engine/Source/Runtime/UMG/Private/Components/TextBlockComponent.cpp b/Engine/Source/Runtime/UMG/Private/Components/TextBlockComponent.cpp
--- a/Engine/Source/Runtime/UMG/Private/Components/TextBlockComponent.cpp
+++ b/Engine/Source/Runtime/UMG/Private/Components/TextBlockComponent.cpp
@@ -11,10 +11,12 @@ UTextBlockComponent::UTextBlockComponent(const FPostConstructInitializePropertie
 	: Super(PCIP)
 {
 	Text = LOCTEXT("TextBlockDefaultValue", "Text Block");
-	Font = FSlateFontInfo(TEXT("Slate/Fonts/Roboto-Bold.ttf"), 24);
 	ShadowOffset = FVector2D(1.0f, 1.0f);
 	ColorAndOpacity = FLinearColor::White;
 	ShadowColorAndOpacity = FLinearColor::Transparent;
+
+	// HACK Special font initialization hack since there are no font assets yet for slate.
+	Font = FSlateFontInfo(TEXT("Slate/Fonts/Roboto-Bold.ttf"), 24);
 }
 
 TSharedRef<SWidget> UTextBlockComponent::RebuildWidget()
@@ -26,7 +28,7 @@ TSharedRef<SWidget> UTextBlockComponent::RebuildWidget()
 		.Font(FSlateFontInfo(FontPath, Font.Size))
 		.ColorAndOpacity( BIND_UOBJECT_ATTRIBUTE(FSlateColor, GetColorAndOpacity) )
 		.ShadowOffset(ShadowOffset)
-		.ShadowColorAndOpacity( BIND_UOBJECT_ATTRIBUTE(FLinearColor, GetShadowColorAndOpacity) );
+		.ShadowColorAndOpacity(BIND_UOBJECT_ATTRIBUTE(FLinearColor, GetShadowColorAndOpacity));
 }
 
 FString UTextBlockComponent::GetText() const
diff --git a/Engine/Source/Runtime/UMG/Private/Components/VerticalBoxComponent.cpp b/Engine/Source/Runtime/UMG/Private/Components/VerticalBoxComponent.cpp
--- a/Engine/Source/Runtime/UMG/Private/Components/VerticalBoxComponent.cpp
+++ b/Engine/Source/Runtime/UMG/Private/Components/VerticalBoxComponent.cpp
@@ -20,6 +20,12 @@ USlateWrapperComponent* UVerticalBoxComponent::GetChildAt(int32 Index) const
 	return Slots[Index]->Content;
 }
 
+bool UVerticalBoxComponent::AddChild(USlateWrapperComponent* Child)
+{
+	AddSlot(Child);
+	return true;
+}
+
 TSharedRef<SWidget> UVerticalBoxComponent::RebuildWidget()
 {
 	TSharedRef<SVerticalBox> NewCanvas = SNew(SVerticalBox);
diff --git a/Engine/Source/Runtime/UMG/Public/Components/ContentWidget.h b/Engine/Source/Runtime/UMG/Public/Components/ContentWidget.h
--- a/Engine/Source/Runtime/UMG/Public/Components/ContentWidget.h
+++ b/Engine/Source/Runtime/UMG/Public/Components/ContentWidget.h
@@ -17,6 +17,7 @@ class UMG_API UContentWidget : public USlateNonLeafWidgetComponent
 
 	virtual int32 GetChildrenCount() const OVERRIDE;
 	virtual USlateWrapperComponent* GetChildAt(int32 Index) const OVERRIDE;
+	virtual bool AddChild(USlateWrapperComponent* InContent) OVERRIDE;
 
 protected:
 
diff --git a/Engine/Source/Runtime/UMG/Public/Components/EditableTextBlockComponent.h b/Engine/Source/Runtime/UMG/Public/Components/EditableTextBlockComponent.h
--- a/Engine/Source/Runtime/UMG/Public/Components/EditableTextBlockComponent.h
+++ b/Engine/Source/Runtime/UMG/Public/Components/EditableTextBlockComponent.h
@@ -12,6 +12,74 @@ class UMG_API UEditableTextBlockComponent : public USlateWrapperComponent
 
 protected:
 
+	/** The text content for this editable text box widget */
+	UPROPERTY(EditAnywhere, BlueprintReadWrite, Category=Content)
+	FText Text;
+
+	/** Hint text that appears when there is no text in the text box */
+	UPROPERTY(EditAnywhere, BlueprintReadWrite, Category=Content)
+	FText HintText;
+
+	/** Font color and opacity (overrides Style) */
+	UPROPERTY(EditAnywhere, Category=Appearance)
+	FSlateFontInfo Font;
+
+	/** Text color and opacity (overrides Style) */
+	UPROPERTY(EditAnywhere, BlueprintReadWrite, Category=Appearance)
+	FSlateColor ForegroundColor;
+
+	/** The color of the background/border around the editable text (overrides Style) */
+	UPROPERTY(EditAnywhere, BlueprintReadWrite, Category=Appearance)
+	FSlateColor BackgroundColor;
+
+	/** Text color and opacity when read-only (overrides Style) */
+	UPROPERTY(EditAnywhere, BlueprintReadWrite, Category=Appearance)
+	FSlateColor ReadOnlyForegroundColor;
+
+	/** Sets whether this text box can actually be modified interactively by the user */
+	UPROPERTY(EditAnywhere, BlueprintReadWrite, Category=Appearance)
+	bool IsReadOnly;
+
+	/** Sets whether this text box is for storing a password */
+	UPROPERTY(EditAnywhere, BlueprintReadWrite, Category=Appearance)
+	bool IsPassword;
+
+	/** Minimum width that a text block should be */
+	UPROPERTY(EditAnywhere, BlueprintReadWrite, Category=Appearance)
+	float MinimumDesiredWidth;
+
+	/** Padding between the box/border and the text widget inside (overrides Style) */
+	UPROPERTY(EditAnywhere, BlueprintReadWrite, Category=Appearance)
+	FMargin Padding;
+
+	/** Workaround as we loose focus when the auto completion closes. */
+	UPROPERTY(EditAnywhere, BlueprintReadWrite, Category=Behavior)
+	bool IsCaretMovedWhenGainFocus;
+
+	/** Whether to select all text when the user clicks to give focus on the widget */
+	UPROPERTY(EditAnywhere, BlueprintReadWrite, Category=Behavior)
+	bool SelectAllTextWhenFocused;
+
+	/** Whether to allow the user to back out of changes when they press the escape key */
+	UPROPERTY(EditAnywhere, BlueprintReadWrite, Category=Behavior)
+	bool RevertTextOnEscape;
+
+	/** Whether to clear keyboard focus when pressing enter to commit changes */
+	UPROPERTY(EditAnywhere, BlueprintReadWrite, Category=Behavior)
+	bool ClearKeyboardFocusOnCommit;
+
+	/** Whether to select all text when pressing enter to commit changes */
+	UPROPERTY(EditAnywhere, BlueprintReadWrite, Category=Behavior)
+	bool SelectAllTextOnCommit;
+
+	/** Called whenever the text is changed interactively by the user */
+	//SLATE_EVENT(FOnTextChanged, OnTextChanged)
+
+	/** Called whenever the text is committed.  This happens when the user presses enter or the text box loses focus. */
+	//SLATE_EVENT(FOnTextCommitted, OnTextCommitted)
+
+	/** Provide a alternative mechanism for error reporting. */
+	//SLATE_ARGUMENT(TSharedPtr<class IErrorReportingWidget>, ErrorReporting)
 
 protected:
 	// USlateWrapperComponent interface
diff --git a/Engine/Source/Runtime/UMG/Public/Components/HorizontalBoxComponent.h b/Engine/Source/Runtime/UMG/Public/Components/HorizontalBoxComponent.h
--- a/Engine/Source/Runtime/UMG/Public/Components/HorizontalBoxComponent.h
+++ b/Engine/Source/Runtime/UMG/Public/Components/HorizontalBoxComponent.h
@@ -18,6 +18,7 @@ class UMG_API UHorizontalBoxComponent : public USlateNonLeafWidgetComponent
 	// USlateNonLeafWidgetComponent
 	virtual int32 GetChildrenCount() const OVERRIDE;
 	virtual USlateWrapperComponent* GetChildAt(int32 Index) const OVERRIDE;
+	virtual bool UHorizontalBoxComponent::AddChild(USlateWrapperComponent* Child) OVERRIDE;
 	// End USlateNonLeafWidgetComponent
 
 #if WITH_EDITOR
diff --git a/Engine/Source/Runtime/UMG/Public/Components/SlateWrapperComponent.h b/Engine/Source/Runtime/UMG/Public/Components/SlateWrapperComponent.h
--- a/Engine/Source/Runtime/UMG/Public/Components/SlateWrapperComponent.h
+++ b/Engine/Source/Runtime/UMG/Public/Components/SlateWrapperComponent.h
@@ -16,6 +16,16 @@ class UMG_API USlateWrapperComponent : public UActorComponent
 	UPanelSlot* Slot;
 #endif
 
+	/** Sets whether this text box can actually be modified interactively by the user */
+	UPROPERTY(EditDefaultsOnly, Category=Behavior)
+	bool bIsEnabled;
+
+	UFUNCTION(BlueprintCallable, Category="Widget")
+	bool GetIsEnabled() const;
+
+	UFUNCTION(BlueprintCallable, Category="Widget")
+	void SetIsEnabled(bool bInIsEnabled);
+	
 	// UActorComponent interface
 	virtual void OnRegister() OVERRIDE;
 	virtual void OnUnregister() OVERRIDE;
diff --git a/Engine/Source/Runtime/UMG/Public/Components/TextBlockComponent.h b/Engine/Source/Runtime/UMG/Public/Components/TextBlockComponent.h
--- a/Engine/Source/Runtime/UMG/Public/Components/TextBlockComponent.h
+++ b/Engine/Source/Runtime/UMG/Public/Components/TextBlockComponent.h
@@ -10,20 +10,19 @@ class UMG_API UTextBlockComponent : public USlateLeafWidgetComponent
 	GENERATED_UCLASS_BODY()
 
 public:
-	UPROPERTY(EditAnywhere, Category=TextBlock)
+	UPROPERTY(EditAnywhere, Category=Content)
 	FText Text;
 
-	UPROPERTY(EditAnywhere, Category=TextBlock)
+	UPROPERTY(EditAnywhere, Category=Appearance)
 	FSlateFontInfo Font;
 
-	UPROPERTY(EditAnywhere, Category=TextBlock)
+	UPROPERTY(EditAnywhere, Category=Appearance)
 	FVector2D ShadowOffset;
 
-	//@TODO: Expose FSlateColor
-	UPROPERTY(EditAnywhere, Category=TextBlock)
-	FLinearColor ColorAndOpacity;
+	UPROPERTY(EditAnywhere, Category=Appearance)
+	FSlateColor ColorAndOpacity;
 
-	UPROPERTY(EditAnywhere, Category=TextBlock)
+	UPROPERTY(EditAnywhere, Category=Appearance)
 	FLinearColor ShadowColorAndOpacity;
 
 protected:
diff --git a/Engine/Source/Runtime/UMG/Public/Components/VerticalBoxComponent.h b/Engine/Source/Runtime/UMG/Public/Components/VerticalBoxComponent.h
--- a/Engine/Source/Runtime/UMG/Public/Components/VerticalBoxComponent.h
+++ b/Engine/Source/Runtime/UMG/Public/Components/VerticalBoxComponent.h
@@ -18,11 +18,7 @@ class UMG_API UVerticalBoxComponent : public USlateNonLeafWidgetComponent
 	// USlateNonLeafWidgetComponent
 	virtual int32 GetChildrenCount() const OVERRIDE;
 	virtual USlateWrapperComponent* GetChildAt(int32 Index) const OVERRIDE;
-	virtual bool AddChild(USlateWrapperComponent* Child) OVERRIDE
-	{
-		AddSlot(Child);
-		return true;
-	}
+	virtual bool AddChild(USlateWrapperComponent* Child) OVERRIDE;
 	// End USlateNonLeafWidgetComponent
 
 #if WITH_EDITOR
