Fix for broken mac build.

diff --git a/Engine/Source/Runtime/UMG/Public/Components/HorizontalBoxComponent.h b/Engine/Source/Runtime/UMG/Public/Components/HorizontalBoxComponent.h
--- a/Engine/Source/Runtime/UMG/Public/Components/HorizontalBoxComponent.h
+++ b/Engine/Source/Runtime/UMG/Public/Components/HorizontalBoxComponent.h
@@ -18,7 +18,7 @@ class UMG_API UHorizontalBoxComponent : public USlateNonLeafWidgetComponent
 	// USlateNonLeafWidgetComponent
 	virtual int32 GetChildrenCount() const OVERRIDE;
 	virtual USlateWrapperComponent* GetChildAt(int32 Index) const OVERRIDE;
-	virtual bool UHorizontalBoxComponent::AddChild(USlateWrapperComponent* Child) OVERRIDE;
+	virtual bool AddChild(USlateWrapperComponent* Child) OVERRIDE;
 	// End USlateNonLeafWidgetComponent
 
 #if WITH_EDITOR
