UMG - Adding editable text widget.

diff --git a/Engine/Source/Runtime/UMG/Private/Components/EditableText.cpp b/Engine/Source/Runtime/UMG/Private/Components/EditableText.cpp
new file mode 100644
--- /dev/null
+++ b/Engine/Source/Runtime/UMG/Private/Components/EditableText.cpp
@@ -0,0 +1,54 @@
+// Copyright 1998-2014 Epic Games, Inc. All Rights Reserved.
+
+#include "UMGPrivatePCH.h"
+
+#define LOCTEXT_NAMESPACE "UMG"
+
+/////////////////////////////////////////////////////
+// UEditableText
+
+UEditableText::UEditableText(const FPostConstructInitializeProperties& PCIP)
+	: Super(PCIP)
+{
+	ColorAndOpacity = FLinearColor::Black;
+
+	// HACK Special font initialization hack since there are no font assets yet for slate.
+	Font = FSlateFontInfo(TEXT("Slate/Fonts/Roboto-Bold.ttf"), 12);
+
+	// Grab other defaults from slate arguments.
+	SEditableTextBox::FArguments Defaults;
+	IsReadOnly = Defaults._IsReadOnly.Get();
+	IsPassword = Defaults._IsReadOnly.Get();
+	MinimumDesiredWidth = Defaults._MinDesiredWidth.Get();
+	IsCaretMovedWhenGainFocus = Defaults._IsCaretMovedWhenGainFocus.Get();
+	SelectAllTextWhenFocused = Defaults._SelectAllTextWhenFocused.Get();
+	RevertTextOnEscape = Defaults._RevertTextOnEscape.Get();
+	ClearKeyboardFocusOnCommit = Defaults._ClearKeyboardFocusOnCommit.Get();
+	SelectAllTextOnCommit = Defaults._SelectAllTextOnCommit.Get();
+}
+
+TSharedRef<SWidget> UEditableText::RebuildWidget()
+{
+	FString FontPath = FPaths::EngineContentDir() / Font.FontName.ToString();
+
+	SEditableText::FArguments Defaults;
+
+	return SNew(SEditableText)
+		.Text(Text)
+		//.Style(Style ? Style->TextBlockStyle : Defaults._Style)
+		.HintText(HintText)
+		.Font(FSlateFontInfo(FontPath, Font.Size))
+		.ColorAndOpacity(ColorAndOpacity)
+		.IsReadOnly(IsReadOnly)
+		.IsPassword(IsPassword)
+		.MinDesiredWidth(MinimumDesiredWidth)
+		.IsCaretMovedWhenGainFocus(IsCaretMovedWhenGainFocus)
+		.SelectAllTextWhenFocused(SelectAllTextWhenFocused)
+		.RevertTextOnEscape(RevertTextOnEscape)
+		.ClearKeyboardFocusOnCommit(ClearKeyboardFocusOnCommit)
+		.SelectAllTextOnCommit(SelectAllTextOnCommit);
+}
+
+/////////////////////////////////////////////////////
+
+#undef LOCTEXT_NAMESPACE
diff --git a/Engine/Source/Runtime/UMG/Public/Components/EditableText.h b/Engine/Source/Runtime/UMG/Public/Components/EditableText.h
new file mode 100644
--- /dev/null
+++ b/Engine/Source/Runtime/UMG/Public/Components/EditableText.h
@@ -0,0 +1,83 @@
+// Copyright 1998-2014 Epic Games, Inc. All Rights Reserved.
+
+#pragma once
+
+#include "EditableText.generated.h"
+
+/** Editable text box widget */
+UCLASS(meta=(BlueprintSpawnableComponent), ClassGroup=UserInterface)
+class UMG_API UEditableText : public USlateWrapperComponent
+{
+	GENERATED_UCLASS_BODY()
+
+protected:
+
+	/** The text content for this editable text box widget */
+	UPROPERTY(EditAnywhere, BlueprintReadWrite, Category=Content)
+	FText Text;
+
+	/** Hint text that appears when there is no text in the text box */
+	UPROPERTY(EditAnywhere, BlueprintReadWrite, Category=Content)
+	FText HintText;
+
+	/** Text style */
+	//UPROPERTY(EditAnywhere, Category=Appearance)
+	//UEditableTextWidgetStyle* Style;
+
+	/** Font color and opacity (overrides Style) */
+	UPROPERTY(EditAnywhere, Category=Appearance)
+	FSlateFontInfo Font;
+
+	/** Text color and opacity (overrides Style) */
+	UPROPERTY(EditAnywhere, BlueprintReadWrite, Category=Appearance)
+	FSlateColor ColorAndOpacity;
+
+	/** Sets whether this text box can actually be modified interactively by the user */
+	UPROPERTY(EditAnywhere, BlueprintReadWrite, Category=Appearance)
+	bool IsReadOnly;
+
+	/** Sets whether this text box is for storing a password */
+	UPROPERTY(EditAnywhere, BlueprintReadWrite, Category=Appearance)
+	bool IsPassword;
+
+	/** Minimum width that a text block should be */
+	UPROPERTY(EditAnywhere, BlueprintReadWrite, Category=Appearance)
+	float MinimumDesiredWidth;
+
+	/** Workaround as we loose focus when the auto completion closes. */
+	UPROPERTY(EditAnywhere, BlueprintReadWrite, Category=Behavior)
+	bool IsCaretMovedWhenGainFocus;
+
+	/** Whether to select all text when the user clicks to give focus on the widget */
+	UPROPERTY(EditAnywhere, BlueprintReadWrite, Category=Behavior)
+	bool SelectAllTextWhenFocused;
+
+	/** Whether to allow the user to back out of changes when they press the escape key */
+	UPROPERTY(EditAnywhere, BlueprintReadWrite, Category=Behavior)
+	bool RevertTextOnEscape;
+
+	/** Whether to clear keyboard focus when pressing enter to commit changes */
+	UPROPERTY(EditAnywhere, BlueprintReadWrite, Category=Behavior)
+	bool ClearKeyboardFocusOnCommit;
+
+	/** Whether to select all text when pressing enter to commit changes */
+	UPROPERTY(EditAnywhere, BlueprintReadWrite, Category=Behavior)
+	bool SelectAllTextOnCommit;
+
+	///** Background image for the selected text (overrides Style) */
+	//SLATE_ATTRIBUTE(const FSlateBrush*, BackgroundImageSelected)
+
+	///** Background image for the selection targeting effect (overrides Style) */
+	//SLATE_ATTRIBUTE(const FSlateBrush*, BackgroundImageSelectionTarget)
+
+	///** Background image for the composing text (overrides Style) */
+	//SLATE_ATTRIBUTE(const FSlateBrush*, BackgroundImageComposing)
+
+	///** Image brush used for the caret (overrides Style) */
+	//SLATE_ATTRIBUTE(const FSlateBrush*, CaretImage)
+
+protected:
+	// USlateWrapperComponent interface
+	virtual TSharedRef<SWidget> RebuildWidget() OVERRIDE;
+	// End of USlateWrapperComponent
+};
diff --git a/Engine/Source/Runtime/UMG/Public/UMG.h b/Engine/Source/Runtime/UMG/Public/UMG.h
--- a/Engine/Source/Runtime/UMG/Public/UMG.h
+++ b/Engine/Source/Runtime/UMG/Public/UMG.h
@@ -45,6 +45,7 @@
 #include "Components/ButtonComponent.h"
 #include "Components/CheckBoxComponent.h"
 #include "Components/CircularThrobberComponent.h"
+#include "Components/EditableText.h"
 #include "Components/EditableTextBlockComponent.h"
 #include "Components/GridPanelComponent.h"
 #include "Components/ImageComponent.h"
