UMG - Excluding abstract classes from the templates.  Experimenting with a way to expose sound overrides for checkbox and button.  Cleanup for checkbox.

diff --git a/Engine/Source/Editor/UMGEditor/Private/SUMGEditorWidgetTemplates.cpp b/Engine/Source/Editor/UMGEditor/Private/SUMGEditorWidgetTemplates.cpp
--- a/Engine/Source/Editor/UMGEditor/Private/SUMGEditorWidgetTemplates.cpp
+++ b/Engine/Source/Editor/UMGEditor/Private/SUMGEditorWidgetTemplates.cpp
@@ -124,7 +124,7 @@ void SUMGEditorWidgetTemplates::BuildClassWidgetList()
 	for ( TObjectIterator<UClass> ClassIt; ClassIt; ++ClassIt )
 	{
 		UClass* WidgetClass = *ClassIt;
-		if ( WidgetClass->IsChildOf(USlateWrapperComponent::StaticClass()) )
+		if ( WidgetClass->IsChildOf(USlateWrapperComponent::StaticClass()) && WidgetClass->HasAnyClassFlags(CLASS_Abstract) == false )
 		{
 			TSharedPtr<FWidgetTemplateClass> Template = MakeShareable(new FWidgetTemplateClass(WidgetClass));
 			Template->Name = WidgetClass->GetDisplayNameText();
diff --git a/Engine/Source/Runtime/UMG/Private/Components/ButtonComponent.cpp b/Engine/Source/Runtime/UMG/Private/Components/ButtonComponent.cpp
--- a/Engine/Source/Runtime/UMG/Private/Components/ButtonComponent.cpp
+++ b/Engine/Source/Runtime/UMG/Private/Components/ButtonComponent.cpp
@@ -34,6 +34,18 @@ TSharedRef<SWidget> UButtonComponent::RebuildWidget()
 		StylePtr = ButtonDefaults._ButtonStyle;
 	}
 
+	TOptional<FSlateSound> OptionalPressedSound;
+	if ( PressedSound.GetResourceObject() )
+	{
+		OptionalPressedSound = PressedSound;
+	}
+
+	TOptional<FSlateSound> OptionalHoveredSound;
+	if ( HoveredSound.GetResourceObject() )
+	{
+		OptionalHoveredSound = HoveredSound;
+	}
+
 	TSharedRef<SButton> NewButton = SNew(SButton)
 		.ButtonStyle(StylePtr)
 		.HAlign(HorizontalAlignment)
@@ -43,7 +55,10 @@ TSharedRef<SWidget> UButtonComponent::RebuildWidget()
 		.DesiredSizeScale(DesiredSizeScale)
 		.ContentScale(ContentScale)
 		.ButtonColorAndOpacity(BIND_UOBJECT_ATTRIBUTE(FSlateColor, GetButtonColor))
-		.ForegroundColor(BIND_UOBJECT_ATTRIBUTE(FSlateColor, GetForegroundColor));
+		.ForegroundColor(BIND_UOBJECT_ATTRIBUTE(FSlateColor, GetForegroundColor))
+		.PressedSoundOverride(OptionalPressedSound)
+		.HoveredSoundOverride(OptionalHoveredSound)
+		;
 
 	NewButton->SetContent(Content ? Content->GetWidget() : SNullWidget::NullWidget);
 	
diff --git a/Engine/Source/Runtime/UMG/Private/Components/CheckBoxComponent.cpp b/Engine/Source/Runtime/UMG/Private/Components/CheckBoxComponent.cpp
--- a/Engine/Source/Runtime/UMG/Private/Components/CheckBoxComponent.cpp
+++ b/Engine/Source/Runtime/UMG/Private/Components/CheckBoxComponent.cpp
@@ -12,29 +12,54 @@ UCheckBoxComponent::UCheckBoxComponent(const FPostConstructInitializeProperties&
 {
 	SCheckBox::FArguments CheckBoxDefaults;
 
-	CheckBoxStyleName = NAME_None;
 	bIsChecked = false;
 
 	HorizontalAlignment = CheckBoxDefaults._HAlign;
 	Padding = CheckBoxDefaults._Padding.Get();
 
-	ForegroundColor = CheckBoxDefaults._ForegroundColor.Get();
-	BorderBackgroundColor = CheckBoxDefaults._BorderBackgroundColor.Get();
+	ForegroundColor = FLinearColor::White;
+	BorderBackgroundColor = FLinearColor::White;
 }
 
 TSharedRef<SWidget> UCheckBoxComponent::RebuildWidget()
 {
+	const FCheckBoxStyle* StylePtr = ( Style != NULL ) ? Style->GetStyle<FCheckBoxStyle>() : NULL;
+	if ( StylePtr == NULL )
+	{
+		SCheckBox::FArguments Defaults;
+		StylePtr = Defaults._Style;
+	}
+
+	TOptional<FSlateSound> OptionalCheckedSound;
+	if ( CheckedSound.GetResourceObject() )
+	{
+		OptionalCheckedSound = CheckedSound;
+	}
+
+	TOptional<FSlateSound> OptionalUncheckedSound;
+	if ( UncheckedSound.GetResourceObject() )
+	{
+		OptionalUncheckedSound = UncheckedSound;
+	}
+
+	TOptional<FSlateSound> OptionalHoveredSound;
+	if ( HoveredSound.GetResourceObject() )
+	{
+		OptionalHoveredSound = HoveredSound;
+	}
+
 	return SNew(SCheckBox)
-// 		SLATE_ATTRIBUTE( ESlateCheckBoxType::Type, Type )
-// 		SLATE_DEFAULT_SLOT( FArguments, Content )
-		//.Style( FEditorStyle::Get(), CheckBoxStyleName )
-// 		SLATE_ARGUMENT( const FCheckBoxStyle *, CheckBoxStyle )
+		.Style(StylePtr)
 		.OnCheckStateChanged( BIND_UOBJECT_DELEGATE(FOnCheckStateChanged, SlateOnCheckStateChangedCallback) )
 		.IsChecked( BIND_UOBJECT_ATTRIBUTE(ESlateCheckBoxState::Type, GetCheckState) )
 		.HAlign( HorizontalAlignment )
 		.Padding( Padding )
 		.ForegroundColor( ForegroundColor )
-		.BorderBackgroundColor( BorderBackgroundColor );
+		.BorderBackgroundColor( BorderBackgroundColor )
+		.CheckedSoundOverride(OptionalCheckedSound)
+		.UncheckedSoundOverride(OptionalUncheckedSound)
+		.HoveredSoundOverride(OptionalHoveredSound)
+		;
 // 		SLATE_ATTRIBUTE( bool, ReadOnly )
 // 		SLATE_ARGUMENT( bool, IsFocusable )
 // 		SLATE_EVENT( FOnGetContent, OnGetMenuContent )
diff --git a/Engine/Source/Runtime/UMG/Public/Components/ButtonComponent.h b/Engine/Source/Runtime/UMG/Public/Components/ButtonComponent.h
--- a/Engine/Source/Runtime/UMG/Public/Components/ButtonComponent.h
+++ b/Engine/Source/Runtime/UMG/Public/Components/ButtonComponent.h
@@ -45,6 +45,12 @@ public:
 	UPROPERTY(EditAnywhere, BlueprintReadWrite, Category=Appearance)
 	FSlateColor ForegroundColor;
 
+	UPROPERTY(EditAnywhere, Category=Sound)
+	FSlateSound PressedSound;
+
+	UPROPERTY(EditAnywhere, Category=Sound)
+	FSlateSound HoveredSound;
+	
 	/** Called when the button is clicked */
 	UPROPERTY(BlueprintAssignable)
 	FOnButtonClickedEvent OnClicked;
diff --git a/Engine/Source/Runtime/UMG/Public/Components/CheckBoxComponent.h b/Engine/Source/Runtime/UMG/Public/Components/CheckBoxComponent.h
--- a/Engine/Source/Runtime/UMG/Public/Components/CheckBoxComponent.h
+++ b/Engine/Source/Runtime/UMG/Public/Components/CheckBoxComponent.h
@@ -14,22 +14,8 @@ class UMG_API UCheckBoxComponent : public USlateWrapperComponent
 
 protected:
 	/** Style of the check box */
-	UPROPERTY(EditAnywhere, BlueprintReadOnly, Category=Appearance)
-	USlateWidgetStyleAsset* CheckBoxStyle;
-
-
-	/** Type of check box */
-	//SLATE_ATTRIBUTE( ESlateCheckBoxType::Type, Type )
-
-	/** Content to be placed next to the check box, or for a toggle button, the content to be placed inside the button */
-	//SLATE_DEFAULT_SLOT( FArguments, Content )
-
-	/** The name of the visual style for this checkbox */
-	UPROPERTY(EditAnywhere, BlueprintReadOnly, Category=Appearance)
-	FName CheckBoxStyleName;
-
-	/** The style structure (image references) for this checkbox' various states */
-	//SLATE_ARGUMENT( const FCheckBoxStyle *, CheckBoxStyle )
+	UPROPERTY(EditAnywhere, BlueprintReadOnly, Category=Style)
+	USlateWidgetStyleAsset* Style;
 
 	/** Called when the checked state has changed */
 	UPROPERTY(BlueprintAssignable)
@@ -60,6 +46,15 @@ protected:
 
 	//SLATE_EVENT( FOnGetContent, OnGetMenuContent )
 
+	UPROPERTY(EditAnywhere, Category=Sound)
+	FSlateSound CheckedSound;
+
+	UPROPERTY(EditAnywhere, Category=Sound)
+	FSlateSound UncheckedSound;
+
+	UPROPERTY(EditAnywhere, Category=Sound)
+	FSlateSound HoveredSound;
+
 protected:
 	// USlateWrapperComponent interface
 	virtual TSharedRef<SWidget> RebuildWidget() OVERRIDE;
