UMG - Checkpointing more work to try and get a main menu system working.  Adding a slider.  Root creation is now possible.  Can sort of delete, super buggy.  Less polling, more active function calls for most things.

diff --git a/Engine/Source/Editor/LevelEditor/Private/SLevelEditor.cpp b/Engine/Source/Editor/LevelEditor/Private/SLevelEditor.cpp
--- a/Engine/Source/Editor/LevelEditor/Private/SLevelEditor.cpp
+++ b/Engine/Source/Editor/LevelEditor/Private/SLevelEditor.cpp
@@ -162,11 +162,6 @@ void SLevelEditor::Construct( const SLevelEditor::FArguments& InArgs)
 
 	FEditorDelegates::MapChange.AddSP(this, &SLevelEditor::HandleEditorMapChange);
 	HandleEditorMapChange(MapChangeEventFlags::NewMap);
-
-	if ( FParse::Param(FCommandLine::Get(), TEXT("umg")) )
-	{
-		FModuleManager::LoadModuleChecked<IUMGEditorModule>("UMGEditor");
-	}
 }
 
 void SLevelEditor::Initialize( const TSharedRef<SDockTab>& OwnerTab, const TSharedRef<SWindow>& OwnerWindow )
diff --git a/Engine/Source/Editor/UMGEditor/Private/SUMGDesigner.cpp b/Engine/Source/Editor/UMGEditor/Private/SUMGDesigner.cpp
--- a/Engine/Source/Editor/UMGEditor/Private/SUMGDesigner.cpp
+++ b/Engine/Source/Editor/UMGEditor/Private/SUMGDesigner.cpp
@@ -286,7 +286,7 @@ FReply SUMGDesigner::OnMouseMove(const FGeometry& MyGeometry, const FPointerEven
 		}
 	}
 
-	return FReply::Handled();
+	return FReply::Unhandled();
 }
 
 bool SUMGDesigner::GetArrangedWidget(TSharedRef<SWidget> Widget, FArrangedWidget& ArrangedWidget) const
@@ -574,9 +574,10 @@ FReply SUMGDesigner::OnDrop(const FGeometry& MyGeometry, const FDragDropEvent& D
 		FArrangedWidget ArrangedWidget(SNullWidget::NullWidget, FGeometry());
 		USlateWrapperComponent* Template = GetTemplateAtCursor(MyGeometry, DragDropEvent, ArrangedWidget);
 
+		UWidgetBlueprint* BP = GetBlueprint();
+
 		if ( Template && Template->IsA(USlateNonLeafWidgetComponent::StaticClass()) )
 		{
-			UWidgetBlueprint* BP = CastChecked<UWidgetBlueprint>(BlueprintEditor.Pin()->GetBlueprintObj());
 			USlateNonLeafWidgetComponent* Parent = Cast<USlateNonLeafWidgetComponent>(Template);
 
 			USlateWrapperComponent* Widget = DragDropOp->Template->Create(BP->WidgetTree);
@@ -586,6 +587,19 @@ FReply SUMGDesigner::OnDrop(const FGeometry& MyGeometry, const FDragDropEvent& D
 			//@TODO UMG When we add a child blindly we need to default the slot size to the preferred size of the widget if the container supports such things.
 			//@TODO UMG We may need a desired size canvas, where the slots have no size, they only give you position, alternatively, maybe slots that don't clip, so center is still easy.
 
+			// Update the selected template to be the newly created one.
+			SelectedTemplate = Widget;
+
+			FBlueprintEditorUtils::MarkBlueprintAsStructurallyModified(BP);
+
+			return FReply::Handled();
+		}
+		else if ( BP->WidgetTree->WidgetTemplates.Num() == 0 )
+		{
+			// No existing templates so just create it and make it the root widget.
+			USlateWrapperComponent* Widget = DragDropOp->Template->Create(BP->WidgetTree);
+			SelectedTemplate = Widget;
+
 			FBlueprintEditorUtils::MarkBlueprintAsStructurallyModified(BP);
 
 			return FReply::Handled();
diff --git a/Engine/Source/Editor/UMGEditor/Private/SUMGEditorTree.cpp b/Engine/Source/Editor/UMGEditor/Private/SUMGEditorTree.cpp
--- a/Engine/Source/Editor/UMGEditor/Private/SUMGEditorTree.cpp
+++ b/Engine/Source/Editor/UMGEditor/Private/SUMGEditorTree.cpp
@@ -37,6 +37,17 @@ void SUMGEditorTree::Construct(const FArguments& InArgs, TSharedPtr<FBlueprintEd
 			]
 		]
 
+		+ SVerticalBox::Slot()
+		.AutoHeight()
+		[
+			SNew(SButton)
+			.OnClicked(this, &SUMGEditorTree::DeleteSelected)
+			[
+				SNew(STextBlock)
+				.Text(NSLOCTEXT("SUMGEditorTree", "DeleteSelected", "Delete Selected"))
+			]
+		]
+
 		+ SVerticalBox::Slot()
 		.FillHeight(1.0f)
 		[
@@ -120,7 +131,11 @@ void SUMGEditorTree::WidgetHierarchy_OnGetChildren(USlateWrapperComponent* InPar
 	{
 		for ( int32 i = 0; i < Widget->GetChildrenCount(); i++ )
 		{
-			OutChildren.Add( Widget->GetChildAt(i) );
+			USlateWrapperComponent* Child = Widget->GetChildAt(i);
+			if ( Child )
+			{
+				OutChildren.Add(Child);
+			}
 		}
 	}
 }
@@ -192,6 +207,28 @@ FReply SUMGEditorTree::CreateTestUI()
 	return FReply::Handled();
 }
 
+FReply SUMGEditorTree::DeleteSelected()
+{
+	TArray<USlateWrapperComponent*> SelectedItems = WidgetTreeView->GetSelectedItems();
+	if ( SelectedItems.Num() > 0 )
+	{
+		UWidgetBlueprint* BP = GetBlueprint();
+
+		bool bRemoved = false;
+		for ( USlateWrapperComponent* Item : SelectedItems )
+		{
+			bRemoved = BP->WidgetTree->RemoveWidget(Item);
+		}
+
+		if ( bRemoved )
+		{
+			FBlueprintEditorUtils::MarkBlueprintAsStructurallyModified(BP);
+		}
+	}
+
+	return FReply::Handled();
+}
+
 void SUMGEditorTree::RefreshTree()
 {
 	RootWidgets.Reset();
diff --git a/Engine/Source/Editor/UMGEditor/Private/SUMGEditorTree.h b/Engine/Source/Editor/UMGEditor/Private/SUMGEditorTree.h
--- a/Engine/Source/Editor/UMGEditor/Private/SUMGEditorTree.h
+++ b/Engine/Source/Editor/UMGEditor/Private/SUMGEditorTree.h
@@ -29,6 +29,8 @@ private:
 
 	FReply CreateTestUI();
 	void RefreshTree();
+
+	FReply DeleteSelected();
 	
 private:
 
diff --git a/Engine/Source/Editor/UMGEditor/Private/SUMGEditorWidgetTemplates.cpp b/Engine/Source/Editor/UMGEditor/Private/SUMGEditorWidgetTemplates.cpp
--- a/Engine/Source/Editor/UMGEditor/Private/SUMGEditorWidgetTemplates.cpp
+++ b/Engine/Source/Editor/UMGEditor/Private/SUMGEditorWidgetTemplates.cpp
@@ -16,6 +16,7 @@
 #include "WidgetTemplate.h"
 #include "WidgetTemplateClass.h"
 #include "WidgetTemplateButton.h"
+#include "WidgetTemplateCheckBox.h"
 
 class SWidgetTemplateItem : public SCompoundWidget
 {
@@ -136,6 +137,7 @@ void SUMGEditorWidgetTemplates::BuildClassWidgetList()
 void SUMGEditorWidgetTemplates::BuildSpecialWidgetList()
 {
 	AddWidgetTemplate(MakeShareable(new FWidgetTemplateButton()));
+	AddWidgetTemplate(MakeShareable(new FWidgetTemplateCheckBox()));
 
 	//TODO Make this pluggable.
 }
diff --git a/Engine/Source/Editor/UnrealEd/Private/Editor.cpp b/Engine/Source/Editor/UnrealEd/Private/Editor.cpp
--- a/Engine/Source/Editor/UnrealEd/Private/Editor.cpp
+++ b/Engine/Source/Editor/UnrealEd/Private/Editor.cpp
@@ -529,6 +529,11 @@ void UEditorEngine::Init(IEngineLoop* InEngineLoop)
 		FModuleManager::Get().LoadModule(TEXT("DeviceProfileEditor"));
 		FModuleManager::Get().LoadModule(TEXT("SourceCodeAccess"));
 
+		if ( FParse::Param(FCommandLine::Get(), TEXT("umg")) )
+		{
+			FModuleManager::Get().LoadModule(TEXT("UMGEditor"));
+		}
+
 		if( FParse::Param( FCommandLine::Get(),TEXT( "PListEditor" ) ) )
 		{
 			FModuleManager::Get().LoadModule(TEXT("PListEditor"));
diff --git a/Engine/Source/Runtime/Engine/Private/BlueprintGeneratedClass.cpp b/Engine/Source/Runtime/Engine/Private/BlueprintGeneratedClass.cpp
--- a/Engine/Source/Runtime/Engine/Private/BlueprintGeneratedClass.cpp
+++ b/Engine/Source/Runtime/Engine/Private/BlueprintGeneratedClass.cpp
@@ -279,7 +279,10 @@ void UBlueprintGeneratedClass::BindDynamicDelegates(AActor* InInstance)
 
 	for (int32 Index = 0; Index < DynamicBindingObjects.Num(); ++Index)
 	{
-		DynamicBindingObjects[Index]->BindDynamicDelegates(InInstance);
+		if ( ensure(DynamicBindingObjects[Index] != NULL) )
+		{
+			DynamicBindingObjects[Index]->BindDynamicDelegates(InInstance);
+		}
 	}
 
 	// call on super class, if it's a BlueprintGeneratedClass
diff --git a/Engine/Source/Runtime/UMG/Private/Components/ButtonComponent.cpp b/Engine/Source/Runtime/UMG/Private/Components/ButtonComponent.cpp
--- a/Engine/Source/Runtime/UMG/Private/Components/ButtonComponent.cpp
+++ b/Engine/Source/Runtime/UMG/Private/Components/ButtonComponent.cpp
@@ -50,19 +50,18 @@ TSharedRef<SWidget> UButtonComponent::RebuildWidget()
 		.ButtonStyle(StylePtr)
 		.HAlign(HorizontalAlignment)
 		.VAlign(VerticalAlignment)
-		.ContentPadding(BIND_UOBJECT_ATTRIBUTE(FMargin, GetContentPadding))
-		.OnClicked(BIND_UOBJECT_DELEGATE(FOnClicked, SlateOnClickedCallback))
+		.ContentPadding(ContentPadding)
 		.DesiredSizeScale(DesiredSizeScale)
 		.ContentScale(ContentScale)
-		.ButtonColorAndOpacity(BIND_UOBJECT_ATTRIBUTE(FSlateColor, GetButtonColor))
-		.ForegroundColor(BIND_UOBJECT_ATTRIBUTE(FSlateColor, GetForegroundColor))
+		.ButtonColorAndOpacity(ButtonColorAndOpacity)
+		.ForegroundColor(ForegroundColor)
 		.PressedSoundOverride(OptionalPressedSound)
 		.HoveredSoundOverride(OptionalHoveredSound)
+		.OnClicked(BIND_UOBJECT_DELEGATE(FOnClicked, HandleOnClicked))
 		;
 
 	NewButton->SetContent(Content ? Content->GetWidget() : SNullWidget::NullWidget);
 	
-	MyButton = NewButton;
 	return NewButton;
 }
 
@@ -70,34 +69,49 @@ void UButtonComponent::SetContent(USlateWrapperComponent* InContent)
 {
 	Super::SetContent(InContent);
 
-	TSharedPtr<SButton> Button = MyButton.Pin();
-	if (Button.IsValid())
+	if ( ButtonWidget().IsValid() )
 	{
-		Button->SetContent(InContent ? InContent->GetWidget() : SNullWidget::NullWidget);
+		ButtonWidget()->SetContent(InContent ? InContent->GetWidget() : SNullWidget::NullWidget);
 	}
 }
 
-FMargin UButtonComponent::GetContentPadding() const
+FLinearColor UButtonComponent::GetButtonColorAndOpacity()
 {
-	return ContentPadding;
+	return ButtonWidget()->GetBorderBackgroundColor().GetSpecifiedColor();
 }
 
-FReply UButtonComponent::SlateOnClickedCallback()
+void UButtonComponent::SetButtonColorAndOpacity(FLinearColor InButtonColorAndOpacity)
 {
-	OnClicked.Broadcast();
-	return FReply::Handled();
+	ButtonWidget()->SetBorderBackgroundColor(InButtonColorAndOpacity);
+}
+
+FLinearColor UButtonComponent::GetForegroundColor()
+{
+	return ButtonWidget()->GetForegroundColor().GetSpecifiedColor();
+}
+
+void UButtonComponent::SetForegroundColor(FLinearColor InForegroundColor)
+{
+	ButtonWidget()->SetForegroundColor(InForegroundColor);
 }
 
-FSlateColor UButtonComponent::GetButtonColor() const
+FMargin UButtonComponent::GetContentPadding()
 {
-	return ButtonColorAndOpacity;
+	return ButtonWidget()->GetContentPadding();
 }
 
-FSlateColor UButtonComponent::GetForegroundColor() const
+void UButtonComponent::SetContentPadding(FMargin InContentPadding)
 {
-	return ForegroundColor;
+	ButtonWidget()->SetContentPadding(InContentPadding);
 }
 
+FReply UButtonComponent::HandleOnClicked()
+{
+	OnClicked.Broadcast();
+	return FReply::Handled();
+}
+
+
 /////////////////////////////////////////////////////
 
 #undef LOCTEXT_NAMESPACE
diff --git a/Engine/Source/Runtime/UMG/Private/Components/ContentWidget.cpp b/Engine/Source/Runtime/UMG/Private/Components/ContentWidget.cpp
--- a/Engine/Source/Runtime/UMG/Private/Components/ContentWidget.cpp
+++ b/Engine/Source/Runtime/UMG/Private/Components/ContentWidget.cpp
@@ -47,6 +47,17 @@ bool UContentWidget::AddChild(USlateWrapperComponent* InContent, FVector2D Posit
 	return true;
 }
 
+bool UContentWidget::RemoveChild(USlateWrapperComponent* Child)
+{
+	if ( Content == Child )
+	{
+		Content = NULL;
+		return true;
+	}
+
+	return false;
+}
+
 /////////////////////////////////////////////////////
 
 #undef LOCTEXT_NAMESPACE
diff --git a/Engine/Source/Runtime/UMG/Private/Components/EditableText.cpp b/Engine/Source/Runtime/UMG/Private/Components/EditableText.cpp
--- a/Engine/Source/Runtime/UMG/Private/Components/EditableText.cpp
+++ b/Engine/Source/Runtime/UMG/Private/Components/EditableText.cpp
@@ -47,15 +47,30 @@ TSharedRef<SWidget> UEditableText::RebuildWidget()
 		.RevertTextOnEscape(RevertTextOnEscape)
 		.ClearKeyboardFocusOnCommit(ClearKeyboardFocusOnCommit)
 		.SelectAllTextOnCommit(SelectAllTextOnCommit)
-		.OnTextChanged(BIND_UOBJECT_DELEGATE(FOnTextChanged, SlateOnTextChanged))
+		.OnTextChanged(BIND_UOBJECT_DELEGATE(FOnTextChanged, HandleOnTextChanged))
+		.OnTextCommitted(BIND_UOBJECT_DELEGATE(FOnTextCommitted, HandleOnTextCommitted))
 		;
 }
 
-void UEditableText::SlateOnTextChanged(const FText& Text)
+#if WITH_EDITOR
+
+void UEditableText::PostEditChangeProperty(struct FPropertyChangedEvent& PropertyChangedEvent)
+{
+	
+}
+
+#endif
+
+void UEditableText::HandleOnTextChanged(const FText& Text)
 {
 	OnTextChanged.Broadcast(Text);
 }
 
+void UEditableText::HandleOnTextCommitted(const FText& Text, ETextCommit::Type CommitMethod)
+{
+	OnTextCommitted.Broadcast(Text, CommitMethod);
+}
+
 /////////////////////////////////////////////////////
 
 #undef LOCTEXT_NAMESPACE
diff --git a/Engine/Source/Runtime/UMG/Private/Components/HorizontalBoxComponent.cpp b/Engine/Source/Runtime/UMG/Private/Components/HorizontalBoxComponent.cpp
--- a/Engine/Source/Runtime/UMG/Private/Components/HorizontalBoxComponent.cpp
+++ b/Engine/Source/Runtime/UMG/Private/Components/HorizontalBoxComponent.cpp
@@ -27,6 +27,22 @@ bool UHorizontalBoxComponent::AddChild(USlateWrapperComponent* Child, FVector2D
 	return true;
 }
 
+bool UHorizontalBoxComponent::RemoveChild(USlateWrapperComponent* Child)
+{
+	for ( int32 SlotIndex = 0; SlotIndex < Slots.Num(); ++SlotIndex )
+	{
+		UHorizontalBoxSlot* Slot = Slots[SlotIndex];
+
+		if ( Slot->Content == Child )
+		{
+			Slots.RemoveAt(SlotIndex);
+			return true;
+		}
+	}
+
+	return false;
+}
+
 TSharedRef<SWidget> UHorizontalBoxComponent::RebuildWidget()
 {
 	TSharedRef<SHorizontalBox> NewCanvas = SNew(SHorizontalBox);
diff --git a/Engine/Source/Runtime/UMG/Private/Components/PanelSlot.cpp b/Engine/Source/Runtime/UMG/Private/Components/PanelSlot.cpp
--- a/Engine/Source/Runtime/UMG/Private/Components/PanelSlot.cpp
+++ b/Engine/Source/Runtime/UMG/Private/Components/PanelSlot.cpp
@@ -18,4 +18,14 @@ void UPanelSlot::Resize(const FVector2D& Direction, const FVector2D& Amount)
 bool UPanelSlot::CanResize(const FVector2D& Direction) const
 {
 	return false;
-}
\ No newline at end of file
+}
+
+void UPanelSlot::MoveTo(const FVector2D& Location)
+{
+
+}
+
+bool UPanelSlot::CanMove() const
+{
+	return false;
+}
diff --git a/Engine/Source/Runtime/UMG/Private/Components/SlateWrapperComponent.cpp b/Engine/Source/Runtime/UMG/Private/Components/SlateWrapperComponent.cpp
--- a/Engine/Source/Runtime/UMG/Private/Components/SlateWrapperComponent.cpp
+++ b/Engine/Source/Runtime/UMG/Private/Components/SlateWrapperComponent.cpp
@@ -10,6 +10,7 @@ USlateWrapperComponent::USlateWrapperComponent(const FPostConstructInitializePro
 {
 	bIsEnabled = true;
 	bIsVariable = true;
+	Visiblity = ESlateVisibility::Visible;
 }
 
 bool USlateWrapperComponent::GetIsEnabled() const
@@ -26,6 +27,21 @@ void USlateWrapperComponent::SetIsEnabled(bool bInIsEnabled)
 	}
 }
 
+TEnumAsByte<ESlateVisibility::Type> USlateWrapperComponent::GetVisibility()
+{
+	return USlateWrapperComponent::ConvertRuntimeToSerializedVisiblity(MyWidget->GetVisibility());
+}
+
+void USlateWrapperComponent::SetVisibility(TEnumAsByte<ESlateVisibility::Type> InVisibility)
+{
+	MyWidget->SetVisibility(USlateWrapperComponent::ConvertSerializedVisibilityToRuntime(InVisibility));
+}
+
+void USlateWrapperComponent::SetToolTipText(const FText& InToolTipText)
+{
+	MyWidget->SetToolTipText(ToolTipText);
+}
+
 void USlateWrapperComponent::OnRegister()
 {
 	Super::OnRegister();
@@ -48,6 +64,11 @@ TSharedRef<SWidget> USlateWrapperComponent::GetWidget()
 	{
 		MyWidget = RebuildWidget();
 		MyWidget->SetEnabled(bIsEnabled);
+		
+		if ( !ToolTipText.IsEmpty() )
+		{
+			MyWidget->SetToolTipText(ToolTipText);
+		}
 	}
 
 	return MyWidget.ToSharedRef();
@@ -79,6 +100,35 @@ EVisibility USlateWrapperComponent::ConvertSerializedVisibilityToRuntime(TEnumAs
 	}
 }
 
+TEnumAsByte<ESlateVisibility::Type> USlateWrapperComponent::ConvertRuntimeToSerializedVisiblity(const EVisibility& Input)
+{
+	if ( Input == EVisibility::Visible )
+	{
+		return ESlateVisibility::Visible;
+	}
+	else if ( Input == EVisibility::Collapsed )
+	{
+		return ESlateVisibility::Collapsed;
+	}
+	else if ( Input == EVisibility::Hidden )
+	{
+		return ESlateVisibility::Hidden;
+	}
+	else if ( Input == EVisibility::HitTestInvisible )
+	{
+		return ESlateVisibility::HitTestInvisible;
+	}
+	else if ( Input == EVisibility::SelfHitTestInvisible )
+	{
+		return ESlateVisibility::SelfHitTestInvisible;
+	}
+	else
+	{
+		check(false);
+		return ESlateVisibility::Visible;
+	}
+}
+
 FSizeParam USlateWrapperComponent::ConvertSerializedSizeParamToRuntime(const FSlateChildSize& Input)
 {
 	switch ( Input.SizeRule )
diff --git a/Engine/Source/Runtime/UMG/Private/Components/Slider.cpp b/Engine/Source/Runtime/UMG/Private/Components/Slider.cpp
new file mode 100644
--- /dev/null
+++ b/Engine/Source/Runtime/UMG/Private/Components/Slider.cpp
@@ -0,0 +1,63 @@
+// Copyright 1998-2014 Epic Games, Inc. All Rights Reserved.
+
+#include "UMGPrivatePCH.h"
+
+#define LOCTEXT_NAMESPACE "UMG"
+
+/////////////////////////////////////////////////////
+// USlider
+
+USlider::USlider(const FPostConstructInitializeProperties& PCIP)
+	: Super(PCIP)
+{
+	Orientation = EOrientation::Orient_Horizontal;
+	SliderBarColor = FLinearColor::White;
+	SliderHandleColor = FLinearColor::White;
+}
+
+TSharedRef<SWidget> USlider::RebuildWidget()
+{
+	// TODO UMG Bind these
+	//FOnMouseCaptureBeginEvent OnMouseCaptureBegin;
+	//FOnMouseCaptureEndEvent OnMouseCaptureEnd;
+
+	TSharedRef<SSlider> NewWidget = SNew(SSlider)
+		.Orientation(Orientation)
+		.SliderBarColor(SliderBarColor)
+		.SliderHandleColor(SliderHandleColor)
+		.Value(Value)
+		.OnMouseCaptureBegin(BIND_UOBJECT_DELEGATE(FSimpleDelegate, HandleOnMouseCaptureBegin))
+		.OnMouseCaptureEnd(BIND_UOBJECT_DELEGATE(FSimpleDelegate, HandleOnMouseCaptureEnd))
+		.OnValueChanged(BIND_UOBJECT_DELEGATE(FOnFloatValueChanged, HandleOnValueChanged));
+
+	return NewWidget;
+}
+
+void USlider::HandleOnValueChanged(float InValue)
+{
+	OnValueChanged.Broadcast(InValue);
+}
+
+void USlider::HandleOnMouseCaptureBegin()
+{
+	OnMouseCaptureBegin.Broadcast();
+}
+
+void USlider::HandleOnMouseCaptureEnd()
+{
+	OnMouseCaptureEnd.Broadcast();
+}
+
+float USlider::GetValue()
+{
+	return SliderWidget()->GetValue();
+}
+
+void USlider::SetValue(float InValue)
+{
+	return SliderWidget()->SetValue(InValue);
+}
+
+/////////////////////////////////////////////////////
+
+#undef LOCTEXT_NAMESPACE
diff --git a/Engine/Source/Runtime/UMG/Private/Components/VerticalBoxComponent.cpp b/Engine/Source/Runtime/UMG/Private/Components/VerticalBoxComponent.cpp
--- a/Engine/Source/Runtime/UMG/Private/Components/VerticalBoxComponent.cpp
+++ b/Engine/Source/Runtime/UMG/Private/Components/VerticalBoxComponent.cpp
@@ -27,6 +27,22 @@ bool UVerticalBoxComponent::AddChild(USlateWrapperComponent* Child, FVector2D Po
 	return true;
 }
 
+bool UVerticalBoxComponent::RemoveChild(USlateWrapperComponent* Child)
+{
+	for ( int32 SlotIndex = 0; SlotIndex < Slots.Num(); ++SlotIndex )
+	{
+		UVerticalBoxSlot* Slot = Slots[SlotIndex];
+
+		if ( Slot->Content == Child )
+		{
+			Slots.RemoveAt(SlotIndex);
+			return true;
+		}
+	}
+
+	return false;
+}
+
 TSharedRef<SWidget> UVerticalBoxComponent::RebuildWidget()
 {
 	TSharedRef<SVerticalBox> NewCanvas = SNew(SVerticalBox);
@@ -79,7 +95,11 @@ void UVerticalBoxComponent::ConnectEditorData()
 {
 	for ( UVerticalBoxSlot* Slot : Slots )
 	{
-		Slot->Content->Slot = Slot;
+		if ( Slot->Content )
+		{
+			Slot->Content->Slot = Slot;
+		}
+		//TODO UMG Should we auto delete empty slots?
 	}
 }
 
diff --git a/Engine/Source/Runtime/UMG/Private/UMGModule.cpp b/Engine/Source/Runtime/UMG/Private/UMGModule.cpp
--- a/Engine/Source/Runtime/UMG/Private/UMGModule.cpp
+++ b/Engine/Source/Runtime/UMG/Private/UMGModule.cpp
@@ -16,6 +16,12 @@ public:
 	/** Called right after the module DLL has been loaded and the module object has been created */
 	virtual void StartupModule() OVERRIDE
 	{
+#if WITH_EDITOR
+		if ( FParse::Param(FCommandLine::Get(), TEXT("umg")) )
+		{
+			FModuleManager::Get().LoadModule(TEXT("UMGEditor"));
+		}
+#endif
 	}
 
 	/** Called before the module is unloaded, right before the module object is destroyed. */
diff --git a/Engine/Source/Runtime/UMG/Private/UserWidget.cpp b/Engine/Source/Runtime/UMG/Private/UserWidget.cpp
--- a/Engine/Source/Runtime/UMG/Private/UserWidget.cpp
+++ b/Engine/Source/Runtime/UMG/Private/UserWidget.cpp
@@ -9,6 +9,10 @@ AUserWidget::AUserWidget(const FPostConstructInitializeProperties& PCIP)
 	: Super(PCIP)
 {
 	RootWidget = SNullWidget::NullWidget;
+	Visiblity = ESlateVisibility::Visible;
+	PrimaryActorTick.bCanEverTick = true;
+	PrimaryActorTick.bTickEvenWhenPaused = true;
+	bShowCursorWhenVisible = false;
 }
 
 void AUserWidget::PostActorCreated()
@@ -29,6 +33,14 @@ void AUserWidget::RerunConstructionScripts()
 	RebuildWrapperWidget();
 }
 
+void AUserWidget::TickActor(float DeltaTime, enum ELevelTick TickType, FActorTickFunction& ThisTickFunction)
+{
+	if ( bShowCursorWhenVisible && GetIsVisible() )
+	{
+		GetWorld()->GetFirstLocalPlayerFromController()->PlayerController->bShowMouseCursor = true;
+	}
+}
+
 USlateWrapperComponent* AUserWidget::GetWidgetHandle(TSharedRef<SWidget> InWidget)
 {
 	return WidgetToComponent.FindRef(InWidget);
@@ -63,6 +75,7 @@ void AUserWidget::RebuildWrapperWidget()
 	{
 		UGameViewportClient* Viewport = World->GetGameViewport();
 		Viewport->AddViewportWidgetContent(RootWidget.ToSharedRef());
+		RootWidget->SetVisibility(USlateWrapperComponent::ConvertSerializedVisibilityToRuntime(Visiblity));
 	}
 }
 
@@ -96,3 +109,25 @@ TSharedRef<SWidget> AUserWidget::GetRootWidget()
 {
 	return RootWidget.ToSharedRef();
 }
+
+void AUserWidget::Show()
+{
+	RootWidget->SetVisibility(EVisibility::Visible);
+	OnVisibilityChanged.Broadcast(ESlateVisibility::Visible);
+}
+
+void AUserWidget::Hide()
+{
+	RootWidget->SetVisibility(EVisibility::Hidden);
+	OnVisibilityChanged.Broadcast(ESlateVisibility::Hidden);
+}
+
+bool AUserWidget::GetIsVisible()
+{
+	return RootWidget->GetVisibility().IsVisible();
+}
+
+TEnumAsByte<ESlateVisibility::Type> AUserWidget::GetVisiblity()
+{
+	return USlateWrapperComponent::ConvertRuntimeToSerializedVisiblity(RootWidget->GetVisibility());
+}
diff --git a/Engine/Source/Runtime/UMG/Private/WidgetTree.cpp b/Engine/Source/Runtime/UMG/Private/WidgetTree.cpp
--- a/Engine/Source/Runtime/UMG/Private/WidgetTree.cpp
+++ b/Engine/Source/Runtime/UMG/Private/WidgetTree.cpp
@@ -9,3 +9,52 @@ UWidgetTree::UWidgetTree(const FPostConstructInitializeProperties& PCIP)
 	: Super(PCIP)
 {
 }
+
+void UWidgetTree::RenameWidget(USlateWrapperComponent* Widget, FString& NewName)
+{
+	Widget->Rename(*NewName);
+
+	// TODO Update nodes in the blueprint!
+}
+
+USlateWrapperComponent* UWidgetTree::FindWidget(FString& Name) const
+{
+	FString ExistingName;
+	for ( USlateWrapperComponent* Widget : WidgetTemplates )
+	{
+		Widget->GetName(ExistingName);
+		if ( ExistingName.Equals(Name, ESearchCase::IgnoreCase) )
+		{
+			return Widget;
+		}
+	}
+
+	return NULL;
+}
+
+bool UWidgetTree::RemoveWidget(USlateWrapperComponent* Widget)
+{
+	USlateWrapperComponent* Parent = NULL;
+
+	bool bRemoved = false;
+
+	for ( USlateWrapperComponent* Template : WidgetTemplates )
+	{
+		USlateNonLeafWidgetComponent* NonLeafTemplate = Cast<USlateNonLeafWidgetComponent>(Template);
+		if ( NonLeafTemplate )
+		{
+			if ( NonLeafTemplate->RemoveChild(Widget) )
+			{
+				bRemoved = true;
+				break;
+			}
+		}
+	}
+
+	if ( bRemoved )
+	{
+		WidgetTemplates.Remove(Widget);
+	}
+
+	return bRemoved;
+}
\ No newline at end of file
diff --git a/Engine/Source/Runtime/UMG/Public/Blueprint/UserWidget.h b/Engine/Source/Runtime/UMG/Public/Blueprint/UserWidget.h
--- a/Engine/Source/Runtime/UMG/Public/Blueprint/UserWidget.h
+++ b/Engine/Source/Runtime/UMG/Public/Blueprint/UserWidget.h
@@ -4,17 +4,44 @@
 
 #include "UserWidget.generated.h"
 
+DECLARE_DYNAMIC_MULTICAST_DELEGATE_OneParam(FOnVisibilityChangedEvent, ESlateVisibility::Type, Visibility);
+
 UCLASS(Abstract, hideCategories=(Object, Actor))
 class UMG_API AUserWidget : public AActor
 {
 	GENERATED_UCLASS_BODY()
 
+	UPROPERTY(EditDefaultsOnly, Category=Behavior)
+	TEnumAsByte<ESlateVisibility::Type> Visiblity;
+
+	/** Called when the visibility changes. */
+	UPROPERTY(BlueprintAssignable)
+	FOnVisibilityChangedEvent OnVisibilityChanged;
+
+	UPROPERTY(EditDefaultsOnly, Category=Behavior)
+	uint32 bShowCursorWhenVisible : 1;
+
 	// AActor interface
 	virtual void PostActorCreated() OVERRIDE;
 	virtual void Destroyed() OVERRIDE;
 	virtual void RerunConstructionScripts() OVERRIDE;
+	virtual void TickActor(float DeltaTime, enum ELevelTick TickType, FActorTickFunction& ThisTickFunction) OVERRIDE;
 	// End of AActor interface
 
+	/*  */
+	UFUNCTION(BlueprintCallable, Category="Appearance")
+	void Show();
+
+	/*  */
+	UFUNCTION(BlueprintCallable, Category="Appearance")
+	void Hide();
+
+	UFUNCTION(BlueprintPure, Category="Appearance")
+	bool GetIsVisible();
+
+	UFUNCTION(BlueprintPure, Category="Appearance")
+	TEnumAsByte<ESlateVisibility::Type> GetVisiblity();
+
 	USlateWrapperComponent* GetWidgetHandle(TSharedRef<SWidget> InWidget);
 
 	TSharedRef<SWidget> GetRootWidget();
@@ -22,6 +49,7 @@ class UMG_API AUserWidget : public AActor
 	TSharedPtr<SWidget> GetWidgetFromName(const FString& Name) const;
 	USlateWrapperComponent* GetHandleFromName(const FString& Name) const;
 
+
 private:
 	TSharedPtr<SWidget> RootWidget;
 	TMap< TWeakPtr<SWidget>, USlateWrapperComponent* > WidgetToComponent;
diff --git a/Engine/Source/Runtime/UMG/Public/Blueprint/WidgetTree.h b/Engine/Source/Runtime/UMG/Public/Blueprint/WidgetTree.h
--- a/Engine/Source/Runtime/UMG/Public/Blueprint/WidgetTree.h
+++ b/Engine/Source/Runtime/UMG/Public/Blueprint/WidgetTree.h
@@ -13,12 +13,11 @@ class UMG_API UWidgetTree : public UObject
 	UPROPERTY()
 	TArray<class USlateWrapperComponent*> WidgetTemplates;
 
-	void RenameWidget(USlateWrapperComponent* Widget, FString& NewName)
-	{
-		Widget->Rename(*NewName);
+	void RenameWidget(USlateWrapperComponent* Widget, FString& NewName);
 
-		// TODO Update nodes in the blueprint!
-	}
+	class USlateWrapperComponent* FindWidget(FString& Name) const;
+
+	bool RemoveWidget(class USlateWrapperComponent* Widget);
 
 	template< class T >
 	T* ConstructWidget(TSubclassOf<class USlateWrapperComponent> WidgetType)
@@ -28,19 +27,4 @@ class UMG_API UWidgetTree : public UObject
 
 		return (T*)Widget;
 	}
-
-	class USlateWrapperComponent* FindWidget(FString& Name) const
-	{
-		FString ExistingName;
-		for ( USlateWrapperComponent* Widget : WidgetTemplates )
-		{
-			Widget->GetName(ExistingName);
-			if ( ExistingName.Equals(Name, ESearchCase::IgnoreCase) )
-			{
-				return Widget;
-			}
-		}
-
-		return NULL;
-	}
 };
diff --git a/Engine/Source/Runtime/UMG/Public/Components/ButtonComponent.h b/Engine/Source/Runtime/UMG/Public/Components/ButtonComponent.h
--- a/Engine/Source/Runtime/UMG/Public/Components/ButtonComponent.h
+++ b/Engine/Source/Runtime/UMG/Public/Components/ButtonComponent.h
@@ -14,53 +14,86 @@ class UMG_API UButtonComponent : public UContentWidget
 
 public:
 	/** Style of the button */
-	UPROPERTY(EditAnywhere, BlueprintReadOnly, Category=Style, meta=( DisplayThumbnail = "true" ))
+	UPROPERTY(EditDefaultsOnly, Category=Style, meta=( DisplayThumbnail = "true" ))
 	USlateWidgetStyleAsset* Style;
 
 	/** Horizontal positioning of the content within the button */
-	UPROPERTY(EditAnywhere, BlueprintReadOnly, Category=Appearance)
+	UPROPERTY(EditDefaultsOnly, Category=Appearance)
 	TEnumAsByte<EHorizontalAlignment> HorizontalAlignment;
 
 	/** Vertical positioning of the content within the button */
-	UPROPERTY(EditAnywhere, BlueprintReadOnly, Category=Appearance)
+	UPROPERTY(EditDefaultsOnly, Category=Appearance)
 	TEnumAsByte<EVerticalAlignment> VerticalAlignment;
 
 	/** The padding to add around the button content. */
-	UPROPERTY(EditAnywhere, BlueprintReadWrite, Category=Appearance)
+	UPROPERTY(EditDefaultsOnly, Category=Appearance)
 	FMargin ContentPadding;
 
 	/** The scaling factor for the button border */
-	UPROPERTY(EditAnywhere, BlueprintReadOnly, Category=Appearance)
+	UPROPERTY(EditDefaultsOnly, Category=Appearance)
 	FVector2D DesiredSizeScale;
 
 	/** The scaling factor for the button content */
-	UPROPERTY(EditAnywhere, BlueprintReadOnly, Category=Appearance)
+	UPROPERTY(EditDefaultsOnly, Category=Appearance)
 	FVector2D ContentScale;
 
 	/** The color multiplier for the button images */
-	UPROPERTY(EditAnywhere, BlueprintReadWrite, Category=Appearance)
-	FSlateColor ButtonColorAndOpacity;
+	UPROPERTY(EditDefaultsOnly, Category=Appearance)
+	FLinearColor ButtonColorAndOpacity;
 
 	/** The foreground color of the button */
-	UPROPERTY(EditAnywhere, BlueprintReadWrite, Category=Appearance)
-	FSlateColor ForegroundColor;
+	UPROPERTY(EditDefaultsOnly, Category=Appearance)
+	FLinearColor ForegroundColor;
 
-	UPROPERTY(EditAnywhere, Category=Sound)
+	UPROPERTY(EditDefaultsOnly, Category=Sound)
 	FSlateSound PressedSound;
 
-	UPROPERTY(EditAnywhere, Category=Sound)
+	UPROPERTY(EditDefaultsOnly, Category=Sound)
 	FSlateSound HoveredSound;
 	
 	/** Called when the button is clicked */
 	UPROPERTY(BlueprintAssignable)
 	FOnButtonClickedEvent OnClicked;
 
+	/*  */
+	UFUNCTION(BlueprintPure, Category="Appearance")
+	FLinearColor GetButtonColorAndOpacity();
+	
+	/*  */
+	UFUNCTION(BlueprintCallable, Category="Appearance")
+	void SetButtonColorAndOpacity(FLinearColor InButtonColorAndOpacity);
+
+	/*  */
+	UFUNCTION(BlueprintPure, Category="Appearance")
+	FLinearColor GetForegroundColor();
+
+	/*  */
+	UFUNCTION(BlueprintCallable, Category="Appearance")
+	void SetForegroundColor(FLinearColor InForegroundColor);
+	
+	/*  */
+	UFUNCTION(BlueprintPure, Category="Appearance")
+	FMargin GetContentPadding();
+
+	/*  */
+	UFUNCTION(BlueprintCallable, Category="Appearance")
+	void SetContentPadding(FMargin InContentPadding);
+
+
 	// UContentWidget interface
 	virtual void SetContent(USlateWrapperComponent* Content) OVERRIDE;
 	// End UContentWidget interface
 
 protected:
-	TWeakPtr<class SButton> MyButton;
+	TSharedPtr<class SButton> ButtonWidget() const
+	{
+		if ( MyWidget.IsValid() )
+		{
+			return StaticCastSharedRef<SButton>(MyWidget.ToSharedRef());
+		}
+
+		return TSharedPtr<class SButton>();
+	}
 
 protected:
 	// USlateWrapperComponent interface
@@ -68,8 +101,6 @@ protected:
 	// End of USlateWrapperComponent interface
 
 	FMargin GetContentPadding() const;
-	FSlateColor GetButtonColor() const;
-	FSlateColor GetForegroundColor() const;
 
-	virtual FReply SlateOnClickedCallback();
+	FReply HandleOnClicked();
 };
diff --git a/Engine/Source/Runtime/UMG/Public/Components/ContentWidget.h b/Engine/Source/Runtime/UMG/Public/Components/ContentWidget.h
--- a/Engine/Source/Runtime/UMG/Public/Components/ContentWidget.h
+++ b/Engine/Source/Runtime/UMG/Public/Components/ContentWidget.h
@@ -15,9 +15,12 @@ class UMG_API UContentWidget : public USlateNonLeafWidgetComponent
 	UFUNCTION()
 	virtual void SetContent(USlateWrapperComponent* InContent);
 
+	// USlateNonLeafWidgetComponent
 	virtual int32 GetChildrenCount() const OVERRIDE;
 	virtual USlateWrapperComponent* GetChildAt(int32 Index) const OVERRIDE;
-	virtual bool AddChild(USlateWrapperComponent* InContent, FVector2D Position) OVERRIDE;
+	virtual bool AddChild(USlateWrapperComponent* Child, FVector2D Position) OVERRIDE;
+	virtual bool RemoveChild(USlateWrapperComponent* Child) OVERRIDE;
+	// End USlateNonLeafWidgetComponent
 
 protected:
 
diff --git a/Engine/Source/Runtime/UMG/Public/Components/EditableText.h b/Engine/Source/Runtime/UMG/Public/Components/EditableText.h
--- a/Engine/Source/Runtime/UMG/Public/Components/EditableText.h
+++ b/Engine/Source/Runtime/UMG/Public/Components/EditableText.h
@@ -5,7 +5,7 @@
 #include "EditableText.generated.h"
 
 DECLARE_DYNAMIC_MULTICAST_DELEGATE_OneParam(FOnEditableTextChangedEvent, const FText&, Text);
-//DECLARE_DYNAMIC_MULTICAST_DELEGATE_TwoParams(FOnEditableTextCommittedEvent, const FText&, Text, ETextCommit::Type, Type);
+DECLARE_DYNAMIC_MULTICAST_DELEGATE_TwoParams(FOnEditableTextCommittedEvent, const FText&, Text, ETextCommit::Type, CommitMethod);
 
 /** Editable text box widget */
 UCLASS(meta=(BlueprintSpawnableComponent), ClassGroup=UserInterface)
@@ -16,55 +16,55 @@ class UMG_API UEditableText : public USlateWrapperComponent
 protected:
 
 	/** The text content for this editable text box widget */
-	UPROPERTY(EditAnywhere, BlueprintReadWrite, Category=Content)
+	UPROPERTY(EditDefaultsOnly, Category=Content)
 	FText Text;
 
 	/** Hint text that appears when there is no text in the text box */
-	UPROPERTY(EditAnywhere, BlueprintReadWrite, Category=Content)
+	UPROPERTY(EditDefaultsOnly, Category=Content)
 	FText HintText;
 
 	/** Text style */
-	//UPROPERTY(EditAnywhere, Category=Appearance)
-	//UEditableTextWidgetStyle* Style;
+	UPROPERTY(EditDefaultsOnly, Category=Style, meta=( DisplayThumbnail = "true" ))
+	USlateWidgetStyleAsset* Style;
 
 	/** Font color and opacity (overrides Style) */
-	UPROPERTY(EditAnywhere, Category=Appearance)
+	UPROPERTY(EditDefaultsOnly, Category=Appearance)
 	FSlateFontInfo Font;
 
 	/** Text color and opacity (overrides Style) */
-	UPROPERTY(EditAnywhere, BlueprintReadWrite, Category=Appearance)
+	UPROPERTY(EditDefaultsOnly, Category=Appearance)
 	FSlateColor ColorAndOpacity;
 
 	/** Sets whether this text box can actually be modified interactively by the user */
-	UPROPERTY(EditAnywhere, BlueprintReadWrite, Category=Appearance)
+	UPROPERTY(EditDefaultsOnly, Category=Appearance)
 	bool IsReadOnly;
 
 	/** Sets whether this text box is for storing a password */
-	UPROPERTY(EditAnywhere, BlueprintReadWrite, Category=Appearance)
+	UPROPERTY(EditDefaultsOnly, Category=Appearance)
 	bool IsPassword;
 
 	/** Minimum width that a text block should be */
-	UPROPERTY(EditAnywhere, BlueprintReadWrite, Category=Appearance)
+	UPROPERTY(EditDefaultsOnly, Category=Appearance)
 	float MinimumDesiredWidth;
 
 	/** Workaround as we loose focus when the auto completion closes. */
-	UPROPERTY(EditAnywhere, BlueprintReadWrite, Category=Behavior)
+	UPROPERTY(EditDefaultsOnly, Category=Behavior)
 	bool IsCaretMovedWhenGainFocus;
 
 	/** Whether to select all text when the user clicks to give focus on the widget */
-	UPROPERTY(EditAnywhere, BlueprintReadWrite, Category=Behavior)
+	UPROPERTY(EditDefaultsOnly, Category=Behavior)
 	bool SelectAllTextWhenFocused;
 
 	/** Whether to allow the user to back out of changes when they press the escape key */
-	UPROPERTY(EditAnywhere, BlueprintReadWrite, Category=Behavior)
+	UPROPERTY(EditDefaultsOnly, Category=Behavior)
 	bool RevertTextOnEscape;
 
 	/** Whether to clear keyboard focus when pressing enter to commit changes */
-	UPROPERTY(EditAnywhere, BlueprintReadWrite, Category=Behavior)
+	UPROPERTY(EditDefaultsOnly, Category=Behavior)
 	bool ClearKeyboardFocusOnCommit;
 
 	/** Whether to select all text when pressing enter to commit changes */
-	UPROPERTY(EditAnywhere, BlueprintReadWrite, Category=Behavior)
+	UPROPERTY(EditDefaultsOnly, Category=Behavior)
 	bool SelectAllTextOnCommit;
 
 	///** Background image for the selected text (overrides Style) */
@@ -84,13 +84,20 @@ protected:
 	FOnEditableTextChangedEvent OnTextChanged;
 
 	/** Called whenever the text is committed.  This happens when the user presses enter or the text box loses focus. */
-	//UPROPERTY(BlueprintAssignable)
-	//FOnTextCommittedEvent OnTextCommitted;
+	UPROPERTY(BlueprintAssignable)
+	FOnEditableTextCommittedEvent OnTextCommitted;
 
 protected:
 	// USlateWrapperComponent interface
 	virtual TSharedRef<SWidget> RebuildWidget() OVERRIDE;
 	// End of USlateWrapperComponent
 
-	void SlateOnTextChanged(const FText& Text);
+#if WITH_EDITOR
+	// UObject interface
+	virtual void PostEditChangeProperty(struct FPropertyChangedEvent& PropertyChangedEvent) OVERRIDE;
+	// End of UObject interface
+#endif
+
+	void HandleOnTextChanged(const FText& Text);
+	void HandleOnTextCommitted(const FText& Text, ETextCommit::Type CommitMethod);
 };
diff --git a/Engine/Source/Runtime/UMG/Public/Components/HorizontalBoxComponent.h b/Engine/Source/Runtime/UMG/Public/Components/HorizontalBoxComponent.h
--- a/Engine/Source/Runtime/UMG/Public/Components/HorizontalBoxComponent.h
+++ b/Engine/Source/Runtime/UMG/Public/Components/HorizontalBoxComponent.h
@@ -19,6 +19,7 @@ class UMG_API UHorizontalBoxComponent : public USlateNonLeafWidgetComponent
 	virtual int32 GetChildrenCount() const OVERRIDE;
 	virtual USlateWrapperComponent* GetChildAt(int32 Index) const OVERRIDE;
 	virtual bool AddChild(USlateWrapperComponent* Child, FVector2D Position) OVERRIDE;
+	virtual bool RemoveChild(USlateWrapperComponent* Child) OVERRIDE;
 	// End USlateNonLeafWidgetComponent
 
 #if WITH_EDITOR
diff --git a/Engine/Source/Runtime/UMG/Public/Components/PanelSlot.h b/Engine/Source/Runtime/UMG/Public/Components/PanelSlot.h
--- a/Engine/Source/Runtime/UMG/Public/Components/PanelSlot.h
+++ b/Engine/Source/Runtime/UMG/Public/Components/PanelSlot.h
@@ -15,4 +15,8 @@ class UMG_API UPanelSlot : public UObject
 	virtual void Resize(const FVector2D& Direction, const FVector2D& Amount);
 
 	virtual bool CanResize(const FVector2D& Direction) const;
+
+	virtual void MoveTo(const FVector2D& Location);
+
+	virtual bool CanMove() const;
 };
diff --git a/Engine/Source/Runtime/UMG/Public/Components/SlateNonLeafWidgetComponent.h b/Engine/Source/Runtime/UMG/Public/Components/SlateNonLeafWidgetComponent.h
--- a/Engine/Source/Runtime/UMG/Public/Components/SlateNonLeafWidgetComponent.h
+++ b/Engine/Source/Runtime/UMG/Public/Components/SlateNonLeafWidgetComponent.h
@@ -12,5 +12,13 @@ class UMG_API USlateNonLeafWidgetComponent : public USlateWrapperComponent
 	virtual int32 GetChildrenCount() const { return 0; }
 	virtual USlateWrapperComponent* GetChildAt(int32 Index) const { return NULL; }
 
-	virtual bool AddChild(USlateWrapperComponent* Child, FVector2D Position) { return false; }
+	virtual bool AddChild(USlateWrapperComponent* Child, FVector2D Position)
+	{
+		return false;
+	}
+
+	virtual bool RemoveChild(USlateWrapperComponent* Child)
+	{
+		return false;
+	}
 };
diff --git a/Engine/Source/Runtime/UMG/Public/Components/SlateWrapperComponent.h b/Engine/Source/Runtime/UMG/Public/Components/SlateWrapperComponent.h
--- a/Engine/Source/Runtime/UMG/Public/Components/SlateWrapperComponent.h
+++ b/Engine/Source/Runtime/UMG/Public/Components/SlateWrapperComponent.h
@@ -22,16 +22,40 @@ class UMG_API USlateWrapperComponent : public UActorComponent
 	UPanelSlot* Slot;
 #endif
 
-	/** Sets whether this text box can actually be modified interactively by the user */
+	/** Sets whether this widget can be modified interactively by the user */
 	UPROPERTY(EditDefaultsOnly, Category=Behavior)
 	bool bIsEnabled;
 
+	//TODO UMG ToolTipWidget
+
+	UPROPERTY(EditDefaultsOnly, Category=Behavior)
+	FText ToolTipText;
+
+	UPROPERTY(EditDefaultsOnly, Category=Behavior)
+	TEnumAsByte<ESlateVisibility::Type> Visiblity;
+
+	UPROPERTY(EditDefaultsOnly, Category=Behavior)
+	TEnumAsByte<EMouseCursor::Type> Cursor;
+
 	UFUNCTION(BlueprintCallable, Category="Widget")
 	bool GetIsEnabled() const;
 
 	UFUNCTION(BlueprintCallable, Category="Widget")
 	void SetIsEnabled(bool bInIsEnabled);
-	
+
+	UFUNCTION(BlueprintCallable, Category="Widget")
+	void SetToolTipText(const FText& InToolTipText);
+
+	UFUNCTION(BlueprintCallable, Category="Widget")
+	TEnumAsByte<ESlateVisibility::Type> GetVisibility();
+
+	UFUNCTION(BlueprintCallable, Category="Widget")
+	void SetVisibility(TEnumAsByte<ESlateVisibility::Type> InVisibility);
+
+	//TOptional< TEnumAsByte<EMouseCursor::Type> > GetCursor() const;
+
+	//void SetCursor(const TAttribute< TOptional< TEnumAsByte<EMouseCursor::Type> > >& InCursor);
+
 	// UActorComponent interface
 	virtual void OnRegister() OVERRIDE;
 	virtual void OnUnregister() OVERRIDE;
@@ -50,8 +74,10 @@ class UMG_API USlateWrapperComponent : public UActorComponent
 #endif
 
 	// Utility methods
-	//@TODO: Should move elsewhere
+	//@TODO UMG: Should move elsewhere
 	static EVisibility ConvertSerializedVisibilityToRuntime(TEnumAsByte<ESlateVisibility::Type> Input);
+	static TEnumAsByte<ESlateVisibility::Type> ConvertRuntimeToSerializedVisiblity(const EVisibility& Input);
+
 	static FSizeParam ConvertSerializedSizeParamToRuntime(const FSlateChildSize& Input);
 
 protected:
diff --git a/Engine/Source/Runtime/UMG/Public/Components/SlateWrapperTypes.h b/Engine/Source/Runtime/UMG/Public/Components/SlateWrapperTypes.h
--- a/Engine/Source/Runtime/UMG/Public/Components/SlateWrapperTypes.h
+++ b/Engine/Source/Runtime/UMG/Public/Components/SlateWrapperTypes.h
@@ -59,4 +59,10 @@ struct FSlateChildSize
 		, SizeRule(ESlateSizeRule::Fill)
 	{
 	}
+
+	FSlateChildSize(ESlateSizeRule::Type InSizeRule)
+		: Value(1.0f)
+		, SizeRule(InSizeRule)
+	{
+	}
 };
diff --git a/Engine/Source/Runtime/UMG/Public/Components/Slider.h b/Engine/Source/Runtime/UMG/Public/Components/Slider.h
new file mode 100644
--- /dev/null
+++ b/Engine/Source/Runtime/UMG/Public/Components/Slider.h
@@ -0,0 +1,80 @@
+// Copyright 1998-2014 Epic Games, Inc. All Rights Reserved.
+
+#pragma once
+
+#include "Slider.generated.h"
+
+DECLARE_DYNAMIC_MULTICAST_DELEGATE(FOnMouseCaptureBeginEvent);
+DECLARE_DYNAMIC_MULTICAST_DELEGATE(FOnMouseCaptureEndEvent);
+DECLARE_DYNAMIC_MULTICAST_DELEGATE_OneParam(FOnFloatValueChangedEvent, float, Value);
+
+/** Slider widget */
+UCLASS(meta=(BlueprintSpawnableComponent), ClassGroup=UserInterface)
+class UMG_API USlider : public USlateLeafWidgetComponent
+{
+	GENERATED_UCLASS_BODY()
+
+	///** Whether the slidable area should be indented to fit the handle. */
+	//SLATE_ATTRIBUTE(bool, IndentHandle)
+
+	///** Whether the handle is interactive or fixed. */
+	//SLATE_ATTRIBUTE(bool, Locked)
+
+	///** The style used to draw the slider. */
+	//SLATE_STYLE_ARGUMENT(FSliderStyle, Style)
+
+	/** The slider's orientation. */
+	UPROPERTY(EditDefaultsOnly, Category=Appearance)
+	TEnumAsByte<EOrientation> Orientation;
+
+	/** The color to draw the slider bar in. */
+	UPROPERTY(EditDefaultsOnly, Category=Appearance)
+	FLinearColor SliderBarColor;
+
+	/** The color to draw the slider handle in. */
+	UPROPERTY(EditDefaultsOnly, Category=Appearance)
+	FLinearColor SliderHandleColor;
+
+	/** The volume value to display. */
+	UPROPERTY(EditDefaultsOnly, Category=Appearance)
+	float Value;
+
+	/** Invoked when the mouse is pressed and a capture begins. */
+	UPROPERTY(BlueprintAssignable)
+	FOnMouseCaptureBeginEvent OnMouseCaptureBegin;
+
+	/** Invoked when the mouse is released and a capture ends. */
+	UPROPERTY(BlueprintAssignable)
+	FOnMouseCaptureEndEvent OnMouseCaptureEnd;
+
+	/** Called when the value is changed by slider or typing. */
+	UPROPERTY(BlueprintAssignable)
+	FOnFloatValueChangedEvent OnValueChanged;
+
+	/*  */
+	UFUNCTION(BlueprintPure, Category="Behavior")
+	float GetValue();
+
+	/*  */
+	UFUNCTION(BlueprintCallable, Category="Behavior")
+	void SetValue(float InValue);
+
+protected:
+	TSharedPtr<class SSlider> SliderWidget() const
+	{
+		if ( MyWidget.IsValid() )
+		{
+			return StaticCastSharedRef<SSlider>(MyWidget.ToSharedRef());
+		}
+
+		return TSharedPtr<class SSlider>();
+	}
+
+	// USlateWrapperComponent interface
+	virtual TSharedRef<SWidget> RebuildWidget() OVERRIDE;
+	// End of USlateWrapperComponent interface
+
+	void HandleOnValueChanged(float InValue);
+	void HandleOnMouseCaptureBegin();
+	void HandleOnMouseCaptureEnd();
+};
diff --git a/Engine/Source/Runtime/UMG/Public/Components/VerticalBoxComponent.h b/Engine/Source/Runtime/UMG/Public/Components/VerticalBoxComponent.h
--- a/Engine/Source/Runtime/UMG/Public/Components/VerticalBoxComponent.h
+++ b/Engine/Source/Runtime/UMG/Public/Components/VerticalBoxComponent.h
@@ -19,6 +19,7 @@ class UMG_API UVerticalBoxComponent : public USlateNonLeafWidgetComponent
 	virtual int32 GetChildrenCount() const OVERRIDE;
 	virtual USlateWrapperComponent* GetChildAt(int32 Index) const OVERRIDE;
 	virtual bool AddChild(USlateWrapperComponent* Child, FVector2D Position) OVERRIDE;
+	virtual bool RemoveChild(USlateWrapperComponent* Child) OVERRIDE;
 	// End USlateNonLeafWidgetComponent
 
 #if WITH_EDITOR
diff --git a/Engine/Source/Runtime/UMG/Public/UMG.h b/Engine/Source/Runtime/UMG/Public/UMG.h
--- a/Engine/Source/Runtime/UMG/Public/UMG.h
+++ b/Engine/Source/Runtime/UMG/Public/UMG.h
@@ -43,19 +43,21 @@
 #include "VerticalBoxSlot.h"
 #include "VerticalBoxComponent.h"
 
+#include "TextBlockComponent.h"
 #include "BorderComponent.h"
 #include "ButtonComponent.h"
 #include "CheckBoxComponent.h"
-#include "CircularThrobberComponent.h"
 #include "EditableText.h"
 #include "EditableTextBlockComponent.h"
-#include "GridPanelComponent.h"
 #include "ImageComponent.h"
+#include "Slider.h"
+
+#include "GridPanelComponent.h"
 #include "ScrollBoxComponent.h"
 #include "SpacerComponent.h"
 #include "SpinningImageComponent.h"
-#include "TextBlockComponent.h"
 #include "ThrobberComponent.h"
+#include "CircularThrobberComponent.h"
 
 
 /* Blueprint
