UMG - Working on the focus stealing aspect of the UI, just rough solution for now, needs a windows manager concept.

diff --git a/Engine/Source/Runtime/UMG/Private/UserWidget.cpp b/Engine/Source/Runtime/UMG/Private/UserWidget.cpp
--- a/Engine/Source/Runtime/UMG/Private/UserWidget.cpp
+++ b/Engine/Source/Runtime/UMG/Private/UserWidget.cpp
@@ -2,6 +2,106 @@
 
 #include "UMGPrivatePCH.h"
 
+#include "SceneViewport.h"
+
+class SViewportWidgetHost : public SCompoundWidget
+{
+	SLATE_BEGIN_ARGS(SViewportWidgetHost)
+	{ }
+		SLATE_DEFAULT_SLOT(FArguments, Content)
+	SLATE_END_ARGS()
+
+	void Construct(const FArguments& InArgs, bool bInModal)
+	{
+		bModal = bInModal;
+
+		ChildSlot
+			[
+				InArgs._Content.Widget
+			];
+	}
+
+	virtual bool OnHitTest(const FGeometry& MyGeometry, FVector2D InAbsoluteCursorPosition) OVERRIDE
+	{
+		return true;
+	}
+
+	virtual bool SupportsKeyboardFocus() const OVERRIDE
+	{
+		return true;
+	}
+
+	FReply OnKeyboardFocusReceived(const FGeometry& MyGeometry, const FKeyboardFocusEvent& InKeyboardFocusEvent)
+	{
+		// if we support focus, release the focus captures and lock the focus to this widget 
+		if ( SupportsKeyboardFocus() )
+		{
+			return FReply::Handled().ReleaseMouseCapture().ReleaseJoystickCapture().LockMouseToWidget(SharedThis(this));
+		}
+		else
+		{
+			return FReply::Unhandled();
+		}
+	}
+
+	virtual FCursorReply OnCursorQuery(const FGeometry& MyGeometry, const FPointerEvent& CursorEvent) const
+	{
+		// If we support focus, show the default mouse cursor 
+		if ( SupportsKeyboardFocus() )
+		{
+			return FCursorReply::Cursor(EMouseCursor::Default);
+		}
+		else
+		{
+			return SCompoundWidget::OnCursorQuery(MyGeometry, CursorEvent);
+		}
+	}
+
+	//FReply OnKeyDown(const FGeometry& MyGeometry, const FKeyboardEvent& InKeyboardEvent)
+	//{
+	//	return FReply::Handled();
+	//}
+
+	//FReply OnKeyUp(const FGeometry& MyGeometry, const FKeyboardEvent& InKeyboardEvent)
+	//{
+	//	return FReply::Handled();
+	//}
+
+	FReply OnMouseButtonDown(const FGeometry& MyGeometry, const FPointerEvent& MouseEvent)
+	{
+		if ( bModal )
+		{
+			return FReply::Handled();
+		}
+
+		return FReply::Unhandled();
+	}
+
+	FReply OnMouseButtonUp(const FGeometry& MyGeometry, const FPointerEvent& MouseEvent)
+	{
+		if ( bModal )
+		{
+			return FReply::Handled();
+		}
+
+		return FReply::Unhandled();
+	}
+
+	FReply OnMouseMove(const FGeometry& MyGeometry, const FPointerEvent& MouseEvent)
+	{
+		if ( bModal )
+		{
+			return FReply::Handled();
+		}
+
+		return FReply::Unhandled();
+	}
+
+protected:
+	bool bModal;
+};
+
+
 /////////////////////////////////////////////////////
 // AUserWidget
 
@@ -13,6 +113,11 @@ AUserWidget::AUserWidget(const FPostConstructInitializeProperties& PCIP)
 	PrimaryActorTick.bCanEverTick = true;
 	PrimaryActorTick.bTickEvenWhenPaused = true;
 	bShowCursorWhenVisible = false;
+	bAbsoluteLayout = false;
+	bModal = false;
+	AbsoluteSize = FVector2D(100, 100);
+	HorizontalAlignment = HAlign_Fill;
+	VerticalAlignment = VAlign_Fill;
 }
 
 void AUserWidget::PostActorCreated()
@@ -35,10 +140,17 @@ void AUserWidget::RerunConstructionScripts()
 
 void AUserWidget::TickActor(float DeltaTime, enum ELevelTick TickType, FActorTickFunction& ThisTickFunction)
 {
-	if ( bShowCursorWhenVisible && GetIsVisible() )
-	{
-		GetWorld()->GetFirstLocalPlayerFromController()->PlayerController->bShowMouseCursor = true;
-	}
+	//UWorld* World = GetWorld();
+
+	//if ( World && World->IsGameWorld() )
+	//{
+	//	if ( bShowCursorWhenVisible && GetIsVisible() )
+	//	{
+	//		World->GetFirstLocalPlayerFromController()->PlayerController->bShowMouseCursor = true;
+
+	//		UGameViewportClient* Viewport = World->GetGameViewport();
+	//	}
+	//}
 }
 
 USlateWrapperComponent* AUserWidget::GetWidgetHandle(TSharedRef<SWidget> InWidget)
@@ -57,7 +169,36 @@ void AUserWidget::RebuildWrapperWidget()
 	// Add the first component to the root of the widget surface.
 	if ( SlateWrapperComponents.Num() > 0 )
 	{
-		RootWidget = SlateWrapperComponents[0]->GetWidget();
+		if ( bAbsoluteLayout )
+		{
+			RootWidget =
+				SNew(SCanvas)
+
+				+ SCanvas::Slot()
+				.Position(AbsolutePosition)
+				.Size(AbsoluteSize)
+				.VAlign(VerticalAlignment)
+				.HAlign(HorizontalAlignment)
+				[
+					SlateWrapperComponents[0]->GetWidget()
+				];
+		}
+		else
+		{
+			TSharedRef<SVerticalBox> VerticalBox = SNew(SVerticalBox);
+
+			auto& NewSlot = VerticalBox->AddSlot()
+				.Padding(Padding)
+				.HAlign(HorizontalAlignment)
+				.VAlign(VerticalAlignment)
+				[
+					SlateWrapperComponents[0]->GetWidget()
+				];
+
+			NewSlot.SizeParam = USlateWrapperComponent::ConvertSerializedSizeParamToRuntime(Size);
+
+			RootWidget = VerticalBox;
+		}
 	}
 
 	// Place all of our top-level children Slate wrapped components into the overlay
@@ -73,9 +214,24 @@ void AUserWidget::RebuildWrapperWidget()
 	UWorld* World = GetWorld();
 	if ( World && World->IsGameWorld() )
 	{
+		TSharedRef<SViewportWidgetHost> WidgetHost = SNew(SViewportWidgetHost, (bool)bModal)
+		[
+			RootWidget.ToSharedRef()
+		];
+
+		RootWidget = WidgetHost;
+
 		UGameViewportClient* Viewport = World->GetGameViewport();
 		Viewport->AddViewportWidgetContent(RootWidget.ToSharedRef());
-		RootWidget->SetVisibility(USlateWrapperComponent::ConvertSerializedVisibilityToRuntime(Visiblity));
+
+		if ( Visiblity == ESlateVisibility::Visible )
+		{
+			Show();
+		}
+		else
+		{
+			Hide();
+		}
 	}
 }
 
@@ -114,12 +270,40 @@ void AUserWidget::Show()
 {
 	RootWidget->SetVisibility(EVisibility::Visible);
 	OnVisibilityChanged.Broadcast(ESlateVisibility::Visible);
+
+	// If this is a game world add the widget to the current worlds viewport.
+	UWorld* World = GetWorld();
+	if ( World && World->IsGameWorld() )
+	{
+		UGameViewportClient* Viewport = World->GetGameViewport();
+		TWeakPtr<SViewport> GameViewportWidget = Viewport->GetGameViewport()->GetViewportWidget();
+		if ( GameViewportWidget.IsValid() )
+		{
+			GameViewportWidget.Pin()->SetWidgetToFocusOnActivate(RootWidget);
+			FSlateApplication::Get().SetKeyboardFocus(RootWidget);
+		}
+	}
 }
 
 void AUserWidget::Hide()
 {
 	RootWidget->SetVisibility(EVisibility::Hidden);
 	OnVisibilityChanged.Broadcast(ESlateVisibility::Hidden);
+
+	// If this is a game world add the widget to the current worlds viewport.
+	UWorld* World = GetWorld();
+	if ( World && World->IsGameWorld() )
+	{
+		UGameViewportClient* Viewport = World->GetGameViewport();
+		TWeakPtr<SViewport> GameViewportWidget = Viewport->GetGameViewport()->GetViewportWidget();
+		if ( GameViewportWidget.IsValid() )
+		{
+			//TODO UMG this isn't what should manage focus, a higher level window controller, probably the viewport needs to understand
+			// the Widget stack, and the dialog stack.
+			GameViewportWidget.Pin()->ClearWidgetToFocusOnActivate();
+			FSlateApplication::Get().SetKeyboardFocus(TSharedPtr<SWidget>());
+		}
+	}
 }
 
 bool AUserWidget::GetIsVisible()
diff --git a/Engine/Source/Runtime/UMG/Public/Blueprint/UserWidget.h b/Engine/Source/Runtime/UMG/Public/Blueprint/UserWidget.h
--- a/Engine/Source/Runtime/UMG/Public/Blueprint/UserWidget.h
+++ b/Engine/Source/Runtime/UMG/Public/Blueprint/UserWidget.h
@@ -18,9 +18,63 @@ class UMG_API AUserWidget : public AActor
 	UPROPERTY(BlueprintAssignable)
 	FOnVisibilityChangedEvent OnVisibilityChanged;
 
+	/** Controls whether the cursor is automatically visible when this widget is visible. */
 	UPROPERTY(EditDefaultsOnly, Category=Behavior)
 	uint32 bShowCursorWhenVisible : 1;
 
+	UPROPERTY(EditDefaultsOnly, Category=Behavior)
+	uint32 bModal : 1;
+
+	/**  */
+	UPROPERTY(EditDefaultsOnly, Category=Appearance)
+	bool bAbsoluteLayout;
+
+	UPROPERTY(EditAnywhere, Category=Appearance, meta=( EditCondition="!bAbsoluteLayout" ))
+	FMargin Padding;
+
+	/** How much space this slot should occupy in the direction of the panel. */
+	UPROPERTY(EditAnywhere, Category=Appearance, meta=( EditCondition="!bAbsoluteLayout" ))
+	FSlateChildSize Size;
+
+	/** Position. */
+	UPROPERTY(EditDefaultsOnly, Category=Appearance, meta=( EditCondition="bAbsoluteLayout" ))
+	FVector2D AbsolutePosition;
+
+	/** Size. */
+	UPROPERTY(EditDefaultsOnly, Category=Appearance, meta=( EditCondition="bAbsoluteLayout" ))
+	FVector2D AbsoluteSize;
+
+	/**
+	* Horizontal pivot position
+	*  Given a top aligned slot, where '+' represents the
+	*  anchor point defined by PositionAttr.
+	*
+	*   Left				Center				Right
+	+ _ _ _ _            _ _ + _ _          _ _ _ _ +
+	|		  |		   | 		   |	  |		    |
+	| _ _ _ _ |        | _ _ _ _ _ |	  | _ _ _ _ |
+	*
+	*  Note: FILL is NOT supported in absolute layout
+	*/
+	UPROPERTY(EditAnywhere, Category=Appearance)
+	TEnumAsByte<EHorizontalAlignment> HorizontalAlignment;
+
+	/**
+	* Vertical pivot position
+	*   Given a left aligned slot, where '+' represents the
+	*   anchor point defined by PositionAttr.
+	*
+	*   Top					Center			  Bottom
+	*	+_ _ _ _ _		 _ _ _ _ _		 _ _ _ _ _
+	*	|         |		| 		  |		|		  |
+	*	|         |     +		  |		|		  |
+	*	| _ _ _ _ |		| _ _ _ _ |		+ _ _ _ _ |
+	*
+	*  Note: FILL is NOT supported in absolute layout
+	*/
+	UPROPERTY(EditAnywhere, Category=Appearance)
+	TEnumAsByte<EVerticalAlignment> VerticalAlignment;
+
 	// AActor interface
 	virtual void PostActorCreated() OVERRIDE;
 	virtual void Destroyed() OVERRIDE;
