UMG - Prototyping continues on the property data binding.  Can now data bind a few core widget properties, and a few properties on Text Component.

diff --git a/Engine/Source/Editor/UMGEditor/Private/DetailCustomizations.cpp b/Engine/Source/Editor/UMGEditor/Private/DetailCustomizations.cpp
--- a/Engine/Source/Editor/UMGEditor/Private/DetailCustomizations.cpp
+++ b/Engine/Source/Editor/UMGEditor/Private/DetailCustomizations.cpp
@@ -224,7 +224,7 @@ FText FBlueprintWidgetCustomization::GetCurrentBindingText(TSharedRef<IPropertyH
 
 		for ( const FDelegateEditorBinding& Binding : Blueprint->Bindings )
 		{
-			if ( Binding.PropertyName == PropertyName )
+			if ( Binding.ObjectName == OuterObjects[ObjectIndex]->GetName() && Binding.PropertyName == PropertyName )
 			{
 				FName FoundName = Blueprint->GetFieldNameFromClassByGuid<UFunction>(Blueprint->GeneratedClass, Binding.MemberGuid);
 				return FText::FromName(FoundName);
diff --git a/Engine/Source/Runtime/UMG/Private/Components/SlateWrapperComponent.cpp b/Engine/Source/Runtime/UMG/Private/Components/SlateWrapperComponent.cpp
--- a/Engine/Source/Runtime/UMG/Private/Components/SlateWrapperComponent.cpp
+++ b/Engine/Source/Runtime/UMG/Private/Components/SlateWrapperComponent.cpp
@@ -63,12 +63,12 @@ TSharedRef<SWidget> USlateWrapperComponent::GetWidget()
 	if ( !MyWidget.IsValid() )
 	{
 		MyWidget = RebuildWidget();
-		MyWidget->SetEnabled(bIsEnabled);
-		MyWidget->SetVisibility(ConvertSerializedVisibilityToRuntime(Visiblity));
+		MyWidget->SetEnabled(OPTIONAL_BINDING(bool, bIsEnabled));
+		MyWidget->SetVisibility(OPTIONAL_BINDING_CONVERT(ESlateVisibility::Type, Visiblity, EVisibility, ConvertVisibility));
 
 		if ( !ToolTipText.IsEmpty() )
 		{
-			MyWidget->SetToolTipText(ToolTipText);
+			MyWidget->SetToolTipText(OPTIONAL_BINDING(FText, ToolTipText));
 		}
 	}
 
diff --git a/Engine/Source/Runtime/UMG/Private/Components/TextBlockComponent.cpp b/Engine/Source/Runtime/UMG/Private/Components/TextBlockComponent.cpp
--- a/Engine/Source/Runtime/UMG/Private/Components/TextBlockComponent.cpp
+++ b/Engine/Source/Runtime/UMG/Private/Components/TextBlockComponent.cpp
@@ -34,35 +34,17 @@ TSharedRef<SWidget> UTextBlockComponent::RebuildWidget()
 		StylePtr = Defaults._TextStyle;
 	}
 
-//	TAttribute<FText>::Create(TAttribute<FText>::FGetter::From::CreateUObject(this, &GlueText));
-
+	TAttribute<FText> TextBinding = OPTIONAL_BINDING(FText, Text);
+	TAttribute<FSlateColor> ColorAndOpacityBinding = OPTIONAL_BINDING(FSlateColor, ColorAndOpacity);
+	TAttribute<FLinearColor> ShadowColorAndOpacityBinding = OPTIONAL_BINDING(FLinearColor, ShadowColorAndOpacity);
+	
 	return SNew(STextBlock)
-		.Text(BIND_UOBJECT_ATTRIBUTE(FText, GetText))
-		.TextStyle(StylePtr)
-		.Font(FSlateFontInfo(FontPath, Font.Size))
-		.ColorAndOpacity( BIND_UOBJECT_ATTRIBUTE(FSlateColor, GetColorAndOpacity) )
-		.ShadowOffset(ShadowOffset)
-		.ShadowColorAndOpacity(BIND_UOBJECT_ATTRIBUTE(FLinearColor, GetShadowColorAndOpacity));
-}
-
-FText UTextBlockComponent::GetText() const
-{
-	if ( TextDelegate.IsBound() )
-	{
-		return TextDelegate.Execute();
-	}
-
-	return Text;
-}
-
-FSlateColor UTextBlockComponent::GetColorAndOpacity() const
-{
-	return ColorAndOpacity;
-}
-
-FLinearColor UTextBlockComponent::GetShadowColorAndOpacity() const
-{
-	return ShadowColorAndOpacity;
+		.Text( TextBinding )
+		.TextStyle( StylePtr )
+		.Font( FSlateFontInfo(FontPath, Font.Size) )
+		.ColorAndOpacity( ColorAndOpacityBinding )
+		.ShadowOffset( ShadowOffset )
+		.ShadowColorAndOpacity( ShadowColorAndOpacityBinding );
 }
 
 /////////////////////////////////////////////////////
diff --git a/Engine/Source/Runtime/UMG/Public/Components/SlateWrapperComponent.h b/Engine/Source/Runtime/UMG/Public/Components/SlateWrapperComponent.h
--- a/Engine/Source/Runtime/UMG/Public/Components/SlateWrapperComponent.h
+++ b/Engine/Source/Runtime/UMG/Public/Components/SlateWrapperComponent.h
@@ -5,12 +5,21 @@
 #include "SlateWrapperTypes.h"
 #include "SlateWrapperComponent.generated.h"
 
+#define OPTIONAL_BINDING(ReturnType, MemberName) MemberName ## Delegate.IsBound() ? TAttribute< ReturnType >::Create(MemberName ## Delegate.GetUObject(), MemberName ## Delegate.GetFunctionName()) : TAttribute< ReturnType >(MemberName)
+#define OPTIONAL_BINDING_CONVERT(ReturnType, MemberName, ConvertedType, ConversionFunction) MemberName ## Delegate.IsBound() ? TAttribute< ConvertedType >::Create(TAttribute< ConvertedType >::FGetter::CreateUObject(this, &ThisClass::ConversionFunction, TAttribute< ReturnType >::Create(MemberName ## Delegate.GetUObject(), MemberName ## Delegate.GetFunctionName()))) : ConversionFunction(TAttribute< ReturnType >(MemberName))
+
 /** This is the base class for all Slate widget wrapper components */
 UCLASS(Abstract, hideCategories=(Activation, "Components|Activation", ComponentReplication, LOD, Rendering))
 class UMG_API USlateWrapperComponent : public UActorComponent
 {
 	GENERATED_UCLASS_BODY()
 
+	DECLARE_DYNAMIC_DELEGATE_RetVal(bool, FGetBool);
+	DECLARE_DYNAMIC_DELEGATE_RetVal(FText, FGetText);
+	DECLARE_DYNAMIC_DELEGATE_RetVal(FLinearColor, FGetSlateColor);
+	DECLARE_DYNAMIC_DELEGATE_RetVal(FLinearColor, FGetLinearColor);
+	DECLARE_DYNAMIC_DELEGATE_RetVal(ESlateVisibility::Type, FGetSlateVisibility);
+
 	UPROPERTY(EditAnywhere, Category=Variable)
 	FString Name;
 
@@ -26,14 +35,23 @@ class UMG_API USlateWrapperComponent : public UActorComponent
 	UPROPERTY(EditDefaultsOnly, Category=Behavior)
 	bool bIsEnabled;
 
+	UPROPERTY()
+	FGetBool bIsEnabledDelegate;
+
 	//TODO UMG ToolTipWidget
 
 	UPROPERTY(EditDefaultsOnly, Category=Behavior)
 	FText ToolTipText;
 
+	UPROPERTY()
+	FGetText ToolTipTextDelegate;
+
 	UPROPERTY(EditDefaultsOnly, Category=Behavior)
 	TEnumAsByte<ESlateVisibility::Type> Visiblity;
 
+	UPROPERTY()
+	FGetSlateVisibility VisiblityDelegate;
+
 	UPROPERTY(EditDefaultsOnly, Category=Behavior)
 	TEnumAsByte<EMouseCursor::Type> Cursor;
 
@@ -83,6 +101,11 @@ class UMG_API USlateWrapperComponent : public UActorComponent
 protected:
 	virtual TSharedRef<SWidget> RebuildWidget();
 
+	EVisibility ConvertVisibility(TAttribute<ESlateVisibility::Type> SerializedType) const
+	{
+		return ConvertSerializedVisibilityToRuntime(SerializedType.Get());
+	}
+
 protected:
 	TSharedPtr<SWidget> MyWidget;
 };
diff --git a/Engine/Source/Runtime/UMG/Public/Components/TextBlockComponent.h b/Engine/Source/Runtime/UMG/Public/Components/TextBlockComponent.h
--- a/Engine/Source/Runtime/UMG/Public/Components/TextBlockComponent.h
+++ b/Engine/Source/Runtime/UMG/Public/Components/TextBlockComponent.h
@@ -9,14 +9,12 @@ class UMG_API UTextBlockComponent : public USlateLeafWidgetComponent
 {
 	GENERATED_UCLASS_BODY()
 
-	DECLARE_DYNAMIC_DELEGATE_RetVal(FText, FTextDelegate);
-
 public:
 	UPROPERTY(EditDefaultsOnly, Category=Content)
 	FText Text;
 
 	UPROPERTY()
-	FTextDelegate TextDelegate;
+	FGetText TextDelegate;
 	
 	UPROPERTY(EditAnywhere, Category=Appearance)
 	FSlateFontInfo Font;
@@ -25,11 +23,17 @@ public:
 	FVector2D ShadowOffset;
 
 	UPROPERTY(EditAnywhere, Category=Appearance)
-	FSlateColor ColorAndOpacity;
+	FLinearColor ColorAndOpacity;
+
+	UPROPERTY()
+	FGetSlateColor ColorAndOpacityDelegate;
 
 	UPROPERTY(EditAnywhere, Category=Appearance)
 	FLinearColor ShadowColorAndOpacity;
 
+	UPROPERTY()
+	FGetLinearColor ShadowColorAndOpacityDelegate;
+
 	/** Text to display as the content of the button */
 	UPROPERTY(EditAnywhere, BlueprintReadOnly, Category=Style, meta=( DisplayThumbnail = "true" ))
 	USlateWidgetStyleAsset* Style;
@@ -39,7 +43,6 @@ protected:
 	virtual TSharedRef<SWidget> RebuildWidget() OVERRIDE;
 	// End of USlateWrapperComponent interface
 
-	FText GetText() const;
 	FSlateColor GetColorAndOpacity() const;
 	FLinearColor GetShadowColorAndOpacity() const;
 };
