UMG - Can now dynamically create the WidgetBlueprint.  Started a global set of functions for controlling UMG UI.

diff --git a/Engine/Source/Runtime/UMG/Private/UserWidget.cpp b/Engine/Source/Runtime/UMG/Private/UserWidget.cpp
--- a/Engine/Source/Runtime/UMG/Private/UserWidget.cpp
+++ b/Engine/Source/Runtime/UMG/Private/UserWidget.cpp
@@ -120,9 +120,11 @@ AUserWidget::AUserWidget(const FPostConstructInitializeProperties& PCIP)
 	VerticalAlignment = VAlign_Fill;
 }
 
-void AUserWidget::PostActorCreated()
+void AUserWidget::OnConstruction(const FTransform& Transform)
 {
-	Super::PostActorCreated();
+	Super::OnConstruction(Transform);
+
+	RebuildWrapperWidget();
 }
 
 void AUserWidget::Destroyed()
diff --git a/Engine/Source/Runtime/UMG/Private/WidgetBlueprintLibrary.cpp b/Engine/Source/Runtime/UMG/Private/WidgetBlueprintLibrary.cpp
new file mode 100644
--- /dev/null
+++ b/Engine/Source/Runtime/UMG/Private/WidgetBlueprintLibrary.cpp
@@ -0,0 +1,33 @@
+// Copyright 1998-2014 Epic Games, Inc. All Rights Reserved.
+
+#include "UMGPrivatePCH.h"
+
+#include "WidgetBlueprintLibrary.h"
+
+#define LOCTEXT_NAMESPACE "UMG"
+
+/////////////////////////////////////////////////////
+// UWidgetBlueprintLibrary
+
+UWidgetBlueprintLibrary::UWidgetBlueprintLibrary(const FPostConstructInitializeProperties& PCIP)
+: Super(PCIP)
+{
+}
+
+AUserWidget* UWidgetBlueprintLibrary::Create(UObject* WorldContextObject, TSubclassOf<AUserWidget> WidgetType)
+{
+	if ( WidgetType->HasAnyClassFlags(CLASS_Abstract) )
+	{
+		// TODO Error?
+		return NULL;
+	}
+
+	//TODO ConstructObject after AUserWidget is no longer an actor.
+
+	UWorld* World = GEngine->GetWorldFromContextObject(WorldContextObject);
+	AUserWidget* NewWidget = World->SpawnActor<AUserWidget>(WidgetType);
+	
+	return NewWidget;
+}
+
+#undef LOCTEXT_NAMESPACE
\ No newline at end of file
diff --git a/Engine/Source/Runtime/UMG/Public/Blueprint/UserWidget.h b/Engine/Source/Runtime/UMG/Public/Blueprint/UserWidget.h
--- a/Engine/Source/Runtime/UMG/Public/Blueprint/UserWidget.h
+++ b/Engine/Source/Runtime/UMG/Public/Blueprint/UserWidget.h
@@ -76,7 +76,7 @@ class UMG_API AUserWidget : public AActor
 	TEnumAsByte<EVerticalAlignment> VerticalAlignment;
 
 	// AActor interface
-	virtual void PostActorCreated() OVERRIDE;
+	virtual void OnConstruction(const FTransform& Transform) OVERRIDE;
 	virtual void Destroyed() OVERRIDE;
 	virtual void RerunConstructionScripts() OVERRIDE;
 	virtual void TickActor(float DeltaTime, enum ELevelTick TickType, FActorTickFunction& ThisTickFunction) OVERRIDE;
diff --git a/Engine/Source/Runtime/UMG/Public/Blueprint/WidgetBlueprintLibrary.h b/Engine/Source/Runtime/UMG/Public/Blueprint/WidgetBlueprintLibrary.h
new file mode 100644
--- /dev/null
+++ b/Engine/Source/Runtime/UMG/Public/Blueprint/WidgetBlueprintLibrary.h
@@ -0,0 +1,14 @@
+// Copyright 1998-2014 Epic Games, Inc. All Rights Reserved.
+
+#pragma once
+#include "WidgetBlueprintLibrary.generated.h"
+
+UCLASS(MinimalAPI)
+class UWidgetBlueprintLibrary : public UBlueprintFunctionLibrary
+{
+	GENERATED_UCLASS_BODY()
+
+	/** Creates a widget */
+	UFUNCTION(BlueprintCallable, meta=( HidePin="WorldContextObject", DefaultToSelf="WorldContextObject", UnsafeDuringActorConstruction = "true", FriendlyName = "Create Widget" ), Category="User Interface|Widget")
+	static class AUserWidget* Create(UObject* WorldContextObject, TSubclassOf<class AUserWidget> WidgetType);
+};
\ No newline at end of file
