UMG - Completely divorced the UI system from needing to be based on AActor and Components.

diff --git a/Engine/Source/Editor/UMGEditor/Classes/WidgetBlueprintFactory.h b/Engine/Source/Editor/UMGEditor/Classes/WidgetBlueprintFactory.h
--- a/Engine/Source/Editor/UMGEditor/Classes/WidgetBlueprintFactory.h
+++ b/Engine/Source/Editor/UMGEditor/Classes/WidgetBlueprintFactory.h
@@ -15,7 +15,7 @@ class UWidgetBlueprintFactory : public UFactory
 
 	// The parent class of the created blueprint
 	UPROPERTY(EditAnywhere, Category=WidgetBlueprintFactory, meta=(AllowAbstract = ""))
-	TSubclassOf<class AUserWidget> ParentClass;
+	TSubclassOf<class UUserWidget> ParentClass;
 
 	// Begin UFactory Interface
 	virtual bool ConfigureProperties() OVERRIDE;
diff --git a/Engine/Source/Editor/UMGEditor/Private/DetailCustomizations.cpp b/Engine/Source/Editor/UMGEditor/Private/DetailCustomizations.cpp
--- a/Engine/Source/Editor/UMGEditor/Private/DetailCustomizations.cpp
+++ b/Engine/Source/Editor/UMGEditor/Private/DetailCustomizations.cpp
@@ -283,6 +283,12 @@ FText FBlueprintWidgetCustomization::GetCurrentBindingText(TSharedRef<IPropertyH
 	FName PropertyName = PropertyHandle->GetProperty()->GetFName();
 	for ( int32 ObjectIndex = 0; ObjectIndex < OuterObjects.Num(); ObjectIndex++ )
 	{
+		// Ignore null outter objects
+		if ( OuterObjects[ObjectIndex] == NULL )
+		{
+			continue;
+		}
+
 		//TODO UMG handle multiple things selected
 
 		for ( const FDelegateEditorBinding& Binding : Blueprint->Bindings )
diff --git a/Engine/Source/Editor/UMGEditor/Private/SUMGDesigner.cpp b/Engine/Source/Editor/UMGEditor/Private/SUMGDesigner.cpp
--- a/Engine/Source/Editor/UMGEditor/Private/SUMGDesigner.cpp
+++ b/Engine/Source/Editor/UMGEditor/Private/SUMGDesigner.cpp
@@ -7,6 +7,8 @@
 #include "SKismetInspector.h"
 #include "BlueprintEditorUtils.h"
 
+#include "WidgetBlueprintEditor.h"
+
 #include "WidgetTemplateDragDropOp.h"
 
 #define LOCTEXT_NAMESPACE "UMG"
@@ -95,9 +97,8 @@ static bool LocateWidgetsUnderCursor_Helper(FArrangedWidget& Candidate, FVector2
 /////////////////////////////////////////////////////
 // SUMGDesigner
 
-void SUMGDesigner::Construct(const FArguments& InArgs, TSharedPtr<FBlueprintEditor> InBlueprintEditor)
+void SUMGDesigner::Construct(const FArguments& InArgs, TSharedPtr<FWidgetBlueprintEditor> InBlueprintEditor)
 {
-	PreviewWidgetActor = NULL;
 	BlueprintEditor = InBlueprintEditor;
 	CurrentHandle = DH_NONE;
 	SelectedTemplate = NULL;
@@ -164,6 +165,7 @@ void SUMGDesigner::OnBlueprintChanged(UBlueprint* InBlueprint)
 {
 	if ( InBlueprint )
 	{
+		
 	}
 }
 
@@ -173,6 +175,8 @@ void SUMGDesigner::OnObjectPropertyChanged(UObject* ObjectBeingModified, FProper
 	{
 		return;
 	}
+
+	//UpdatePreview(InBlueprint);
 }
 
 void SUMGDesigner::ShowDetailsForObjects(TArray<USlateWrapperComponent*> Widgets)
@@ -211,7 +215,7 @@ USlateWrapperComponent* SUMGDesigner::GetTemplateAtCursor(const FGeometry& MyGeo
 
 	PreviewSurface->SetVisibility(EVisibility::HitTestInvisible);
 
-	AUserWidget* WidgetActor = PreviewWidgetActor.Get();
+	UUserWidget* WidgetActor = BlueprintEditor.Pin()->GetPreview();
 	if ( WidgetActor )
 	{
 		USlateWrapperComponent* PreviewHandle = NULL;
@@ -494,11 +498,8 @@ FCursorReply SUMGDesigner::OnCursorQuery(const FGeometry& MyGeometry, const FPoi
 
 void SUMGDesigner::Tick(const FGeometry& AllottedGeometry, const double InCurrentTime, const float InDeltaTime)
 {
-	AActor* PreviewActor = BlueprintEditor.Pin()->GetPreviewActor();
-	AUserWidget* WidgetActor = Cast<AUserWidget>(PreviewActor);
-	PreviewWidgetActor = WidgetActor;
-
-	if (WidgetActor)
+	UUserWidget* WidgetActor = BlueprintEditor.Pin()->GetPreview();
+	if ( WidgetActor )
 	{
 		TSharedRef<SWidget> CurrentWidget = WidgetActor->GetRootWidget();
 
@@ -518,7 +519,7 @@ void SUMGDesigner::Tick(const FGeometry& AllottedGeometry, const double InCurren
 			.VAlign(VAlign_Center)
 			[
 				SNew(STextBlock)
-				.Text(LOCTEXT("NoWrappedWidget", "No actor; Open the viewport and tab back"))
+				.Text(LOCTEXT("NoWrappedWidget", "No Widget Preview"))
 			]
 		];
 	}
@@ -526,7 +527,6 @@ void SUMGDesigner::Tick(const FGeometry& AllottedGeometry, const double InCurren
 	// Update the selected widget to match the selected template.
 	if ( SelectedTemplate )
 	{
-		AUserWidget* WidgetActor = PreviewWidgetActor.Get();
 		if ( WidgetActor )
 		{
 			// Set the selected widget so that we can draw the highlight
diff --git a/Engine/Source/Editor/UMGEditor/Private/SUMGDesigner.h b/Engine/Source/Editor/UMGEditor/Private/SUMGDesigner.h
--- a/Engine/Source/Editor/UMGEditor/Private/SUMGDesigner.h
+++ b/Engine/Source/Editor/UMGEditor/Private/SUMGDesigner.h
@@ -13,7 +13,7 @@ public:
 	SLATE_BEGIN_ARGS( SUMGDesigner ) {}
 	SLATE_END_ARGS()
 
-	void Construct(const FArguments& InArgs, TSharedPtr<class FBlueprintEditor> InBlueprintEditor);
+	void Construct(const FArguments& InArgs, TSharedPtr<class FWidgetBlueprintEditor> InBlueprintEditor);
 	virtual ~SUMGDesigner();
 
 	// SWidget interface
@@ -69,8 +69,8 @@ private:
 	DragHandle HitTestDragHandles(const FGeometry& AllottedGeometry, const FPointerEvent& PointerEvent) const;
 
 private:
-	TWeakPtr<FBlueprintEditor> BlueprintEditor;
-	TWeakObjectPtr<AUserWidget> PreviewWidgetActor;
+	TWeakPtr<FWidgetBlueprintEditor> BlueprintEditor;
+
 	TWeakPtr<SWidget> PreviewWidget;
 
 	USlateWrapperComponent* SelectedTemplate;
diff --git a/Engine/Source/Editor/UMGEditor/Private/UMGBlueprintEditorExtensionHook.cpp b/Engine/Source/Editor/UMGEditor/Private/UMGBlueprintEditorExtensionHook.cpp
--- a/Engine/Source/Editor/UMGEditor/Private/UMGBlueprintEditorExtensionHook.cpp
+++ b/Engine/Source/Editor/UMGEditor/Private/UMGBlueprintEditorExtensionHook.cpp
@@ -46,36 +46,20 @@ public:
 
 	virtual TSharedRef<SWidget> CreateTabBody(const FWorkflowTabSpawnInfo& Info) const OVERRIDE
 	{
-		TSharedPtr<FBlueprintEditor> BlueprintEditorPtr = BlueprintEditor.Pin();
+		TSharedPtr<FWidgetBlueprintEditor> BlueprintEditorPtr = StaticCastSharedPtr<FWidgetBlueprintEditor>(BlueprintEditor.Pin());
 
-		TSharedPtr<SWidget> Result;
-		if ( BlueprintEditorPtr->CanAccessComponentsMode() )
-		{
-			Result = BlueprintEditorPtr->GetSCSViewport();
-		}
-
-		if ( Result.IsValid() )
-		{
-			return SNew(SHorizontalBox)
-			+ SHorizontalBox::Slot()
-			.FillWidth(0)
-			[
-				// HACK: Super hack to make sure we have a preview actor.
-				Result.ToSharedRef()
-			]
+		//TSharedPtr<SWidget> Result;
+		//if ( BlueprintEditorPtr->CanAccessComponentsMode() )
+		//{
+		//	Result = BlueprintEditorPtr->GetSCSViewport();
+		//}
 
+		return SNew(SHorizontalBox)
 			+ SHorizontalBox::Slot()
 			.FillWidth(1)
 			[
 				SNew(SUMGDesigner, BlueprintEditorPtr)
 			];
-		}
-		else
-		{
-			return SNew(SErrorText)
-				.BackgroundColor(FLinearColor::Transparent)
-				.ErrorText(LOCTEXT("SCSViewportView_Unavailable", "Viewport is not available for this Blueprint.").ToString());
-		}
 	}
 };
 
@@ -142,11 +126,11 @@ public:
 //////////////////////////////////////////////////////////////////////////
 // FComponentsEditorModeOverride
 
-class FComponentsEditorModeOverride : public FBlueprintComponentsApplicationMode
+class FComponentsEditorModeOverride : public FBlueprintDefaultsApplicationMode
 {
 public:
 	FComponentsEditorModeOverride(TSharedPtr<class FBlueprintEditor> InBlueprintEditor)
-		: FBlueprintComponentsApplicationMode(InBlueprintEditor)
+		: FBlueprintDefaultsApplicationMode(InBlueprintEditor)
 	{
 		//BlueprintComponentsTabFactories.RegisterFactory(MakeShareable(new FConstructionScriptEditorSummoner(InBlueprintEditor)));
 		//BlueprintComponentsTabFactories.RegisterFactory(MakeShareable(new FSCSViewportSummoner(InBlueprintEditor)));
@@ -155,10 +139,11 @@ public:
 		//BlueprintComponentsTabFactories.RegisterFactory(MakeShareable(new FFindResultsSummoner(InBlueprintEditor)));
 
 		//BlueprintComponentsTabFactories.UnregisterFactory(FBlueprintEditorTabs::ConstructionScriptEditorID);
-		BlueprintComponentsTabFactories.UnregisterFactory(FBlueprintEditorTabs::SCSViewportID);
-		BlueprintComponentsTabFactories.RegisterFactory(MakeShareable(new FSlatePreviewSummoner(InBlueprintEditor)));
-		BlueprintComponentsTabFactories.RegisterFactory(MakeShareable(new FSlateTreeSummoner(InBlueprintEditor)));
-		BlueprintComponentsTabFactories.RegisterFactory(MakeShareable(new FWidgetTemplatesSummoner(InBlueprintEditor)));
+		//BlueprintDefaultsTabFactories.UnregisterFactory(FBlueprintEditorTabs::SCSViewportID);
+		BlueprintDefaultsTabFactories.RegisterFactory(MakeShareable(new FSelectionDetailsSummoner(InBlueprintEditor)));
+		BlueprintDefaultsTabFactories.RegisterFactory(MakeShareable(new FSlatePreviewSummoner(InBlueprintEditor)));
+		BlueprintDefaultsTabFactories.RegisterFactory(MakeShareable(new FSlateTreeSummoner(InBlueprintEditor)));
+		BlueprintDefaultsTabFactories.RegisterFactory(MakeShareable(new FWidgetTemplatesSummoner(InBlueprintEditor)));
 
 		TabLayout = FTabManager::NewLayout( "Standalone_UMGEditor_Layout_v2" )
 		->AddArea
@@ -224,14 +209,14 @@ public:
 
 	virtual void PreDeactivateMode() OVERRIDE
 	{
-		FBlueprintComponentsApplicationMode::PreDeactivateMode();
+		FBlueprintDefaultsApplicationMode::PreDeactivateMode();
 
 		GetBlueprintEditor()->GetInspector()->EnableComponentDetailsCustomization(false);
 	}
 
 	virtual void PostActivateMode() OVERRIDE
 	{
-		FBlueprintComponentsApplicationMode::PostActivateMode();
+		FBlueprintDefaultsApplicationMode::PostActivateMode();
 
 		GetBlueprintEditor()->GetInspector()->EnableComponentDetailsCustomization(false);
 
@@ -305,14 +290,14 @@ class FBlueprintEditorExtenderForSlatePreview
 public:
 	static TSharedRef<FApplicationMode> OnModeCreated(const FName ModeName, TSharedRef<FApplicationMode> InMode)
 	{
-		if (ModeName == FBlueprintEditorApplicationModes::BlueprintComponentsMode)
+		if (ModeName == FBlueprintEditorApplicationModes::BlueprintDefaultsMode)
 		{
 			//@TODO: Bit of a lie - push GetBlueprint up, or pass in editor!
 			auto LieMode = StaticCastSharedRef<FComponentsEditorModeOverride>(InMode);
 
 			if (UBlueprint* BP = LieMode->GetBlueprint())
 			{
-				if (BP->ParentClass->IsChildOf(AUserWidget::StaticClass()))
+				if (BP->ParentClass->IsChildOf(UUserWidget::StaticClass()))
 				{
 					return MakeShareable(new FComponentsEditorModeOverride(LieMode->GetBlueprintEditor()));
 				}
@@ -325,7 +310,7 @@ public:
 
 			if ( UBlueprint* BP = LieMode->GetBlueprint() )
 			{
-				if ( BP->ParentClass->IsChildOf(AUserWidget::StaticClass()) )
+				if ( BP->ParentClass->IsChildOf(UUserWidget::StaticClass()) )
 				{
 					return MakeShareable(new FGraphEditorModeOverride(LieMode->GetBlueprintEditor(), ModeName));
 				}
diff --git a/Engine/Source/Editor/UMGEditor/Private/WidgetBlueprintEditor.cpp b/Engine/Source/Editor/UMGEditor/Private/WidgetBlueprintEditor.cpp
--- a/Engine/Source/Editor/UMGEditor/Private/WidgetBlueprintEditor.cpp
+++ b/Engine/Source/Editor/UMGEditor/Private/WidgetBlueprintEditor.cpp
@@ -7,16 +7,142 @@
 #define LOCTEXT_NAMESPACE "UMG"
 
 FWidgetBlueprintEditor::FWidgetBlueprintEditor()
+	: PreviewScene(FPreviewScene::ConstructionValues().AllowAudioPlayback(true).ShouldSimulatePhysics(true))
+	, PreviewBlueprint(NULL)
 {
 }
 
 FWidgetBlueprintEditor::~FWidgetBlueprintEditor()
 {
+	UWidgetBlueprint* Blueprint = GetWidgetBlueprintObj();
+	if ( Blueprint )
+	{
+		Blueprint->OnChanged().RemoveAll(this);
+	}
+}
+
+void FWidgetBlueprintEditor::OnBlueprintChanged(UBlueprint* InBlueprint)
+{
+	if ( InBlueprint )
+	{
+		UpdatePreview(InBlueprint, true);
+	}
 }
 
 void FWidgetBlueprintEditor::InitWidgetBlueprintEditor(const EToolkitMode::Type Mode, const TSharedPtr< IToolkitHost >& InitToolkitHost, const TArray<UBlueprint*>& InBlueprints, bool bShouldOpenInDefaultsMode)
 {
 	InitBlueprintEditor(Mode, InitToolkitHost, InBlueprints, bShouldOpenInDefaultsMode);
+
+	UWidgetBlueprint* Blueprint = GetWidgetBlueprintObj();
+	Blueprint->OnChanged().AddSP(this, &FWidgetBlueprintEditor::OnBlueprintChanged);
+
+	UpdatePreview(GetWidgetBlueprintObj(), true);
+}
+
+void FWidgetBlueprintEditor::Tick(float DeltaTime)
+{
+	FBlueprintEditor::Tick(DeltaTime);
+
+	// Tick the preview scene world.
+	if ( !GIntraFrameDebuggingGameThread )
+	{
+		bool bIsSimulateEnabled = true;
+		bool bIsRealTime = true;
+
+		// Allow full tick only if preview simulation is enabled and we're not currently in an active SIE or PIE session
+		if ( bIsSimulateEnabled && GEditor->PlayWorld == NULL && !GEditor->bIsSimulatingInEditor )
+		{
+			PreviewScene.GetWorld()->Tick(bIsRealTime ? LEVELTICK_All : LEVELTICK_TimeOnly, DeltaTime);
+		}
+		else
+		{
+			PreviewScene.GetWorld()->Tick(bIsRealTime ? LEVELTICK_ViewportsOnly : LEVELTICK_TimeOnly, DeltaTime);
+		}
+	}
+}
+
+void FWidgetBlueprintEditor::NotifyPostChange(const FPropertyChangedEvent& PropertyChangedEvent, UProperty* PropertyThatChanged)
+{
+	UpdatePreview(GetWidgetBlueprintObj(), true);
+}
+
+UWidgetBlueprint* FWidgetBlueprintEditor::GetWidgetBlueprintObj() const
+{
+	return Cast<UWidgetBlueprint>(GetBlueprintObj());
+}
+
+UUserWidget* FWidgetBlueprintEditor::GetPreview() const
+{
+	//// Note: The weak ptr can become stale if the actor is reinstanced due to a Blueprint change, etc. In that case we look to see if we can find the new instance in the preview world and then update the weak ptr.
+	//if ( PreviewWidgetActorPtr.IsStale(true) )
+	//{
+	//	UWorld* PreviewWorld = PreviewWidgetActorPtr->GetWorld();
+	//	for ( TActorIterator<AActor> It(PreviewWorld); It; ++It )
+	//	{
+	//		AActor* Actor = *It;
+	//		if ( !Actor->IsPendingKillPending()
+	//			&& Actor->GetClass()->ClassGeneratedBy == PreviewBlueprint )
+	//		{
+	//			PreviewActorPtr = Actor;
+	//			break;
+	//		}
+	//	}
+	//}
+
+	return PreviewWidgetActorPtr.Get();
+}
+
+void FWidgetBlueprintEditor::DestroyPreview()
+{
+	UUserWidget* PreviewActor = GetPreview();
+	if ( PreviewActor != NULL )
+	{
+		check(PreviewScene.GetWorld());
+
+		PreviewActor->ClearFlags(RF_Standalone);
+		PreviewActor->MarkPendingKill();
+	}
+
+	//if ( PreviewBlueprint != NULL )
+	//{
+	//	if ( PreviewBlueprint->SimpleConstructionScript != NULL
+	//		&& PreviewActor == PreviewBlueprint->SimpleConstructionScript->GetComponentEditorActorInstance() )
+	//	{
+	//		// Ensure that all editable component references are cleared
+	//		PreviewBlueprint->SimpleConstructionScript->ClearEditorComponentReferences();
+
+	//		// Clear the reference to the preview actor instance
+	//		PreviewBlueprint->SimpleConstructionScript->SetComponentEditorActorInstance(NULL);
+	//	}
+
+	//	PreviewBlueprint = NULL;
+	//}
+}
+
+void FWidgetBlueprintEditor::UpdatePreview(UBlueprint* InBlueprint, bool bInForceFullUpdate)
+{
+	UUserWidget* PreviewActor = GetPreview();
+
+	// Signal that we're going to be constructing editor components
+	if ( InBlueprint != NULL && InBlueprint->SimpleConstructionScript != NULL )
+	{
+		InBlueprint->SimpleConstructionScript->BeginEditorComponentConstruction();
+	}
+
+	// If the Blueprint is changing
+	if ( InBlueprint != PreviewBlueprint || bInForceFullUpdate )
+	{
+		// Destroy the previous actor instance
+		DestroyPreview();
+
+		// Save the Blueprint we're creating a preview for
+		PreviewBlueprint = InBlueprint;
+
+		PreviewActor = ConstructObject<UUserWidget>(PreviewBlueprint->GeneratedClass, PreviewScene.GetWorld()->GetCurrentLevel());
+		PreviewActor->SetFlags(RF_Standalone);
+
+		PreviewWidgetActorPtr = PreviewActor;
+	}
 }
 
 #undef LOCTEXT_NAMESPACE
diff --git a/Engine/Source/Editor/UMGEditor/Private/WidgetBlueprintEditor.h b/Engine/Source/Editor/UMGEditor/Private/WidgetBlueprintEditor.h
--- a/Engine/Source/Editor/UMGEditor/Private/WidgetBlueprintEditor.h
+++ b/Engine/Source/Editor/UMGEditor/Private/WidgetBlueprintEditor.h
@@ -4,6 +4,8 @@
 
 #include "Editor/Kismet/Public/BlueprintEditor.h"
 
+#include "PreviewScene.h"
+
 /**
  * widget blueprint editor (extends Blueprint editor)
  */
@@ -15,4 +17,25 @@ public:
 	virtual ~FWidgetBlueprintEditor();
 
 	void InitWidgetBlueprintEditor(const EToolkitMode::Type Mode, const TSharedPtr< IToolkitHost >& InitToolkitHost, const TArray<UBlueprint*>& InBlueprints, bool bShouldOpenInDefaultsMode);
+
+	virtual void Tick(float DeltaTime) OVERRIDE;
+	virtual void NotifyPostChange(const FPropertyChangedEvent& PropertyChangedEvent, UProperty* PropertyThatChanged) OVERRIDE;
+
+	class UWidgetBlueprint* GetWidgetBlueprintObj() const;
+
+	UUserWidget* GetPreview() const;
+
+private:
+	void OnBlueprintChanged(UBlueprint* InBlueprint);
+
+	void DestroyPreview();
+	void UpdatePreview(UBlueprint* InBlueprint, bool bInForceFullUpdate);
+
+private:
+	FPreviewScene PreviewScene;
+
+	/** The Blueprint associated with the current preview */
+	UBlueprint* PreviewBlueprint;
+
+	mutable TWeakObjectPtr<UUserWidget> PreviewWidgetActorPtr;
 };
diff --git a/Engine/Source/Editor/UMGEditor/Private/WidgetBlueprintFactory.cpp b/Engine/Source/Editor/UMGEditor/Private/WidgetBlueprintFactory.cpp
--- a/Engine/Source/Editor/UMGEditor/Private/WidgetBlueprintFactory.cpp
+++ b/Engine/Source/Editor/UMGEditor/Private/WidgetBlueprintFactory.cpp
@@ -17,7 +17,7 @@ UWidgetBlueprintFactory::UWidgetBlueprintFactory(const class FPostConstructIniti
 	bCreateNew = true;
 	bEditAfterNew = true;
 	SupportedClass = UWidgetBlueprint::StaticClass();
-	ParentClass = AUserWidget::StaticClass();
+	ParentClass = UUserWidget::StaticClass();
 }
 
 bool UWidgetBlueprintFactory::ConfigureProperties()
@@ -37,7 +37,7 @@ UObject* UWidgetBlueprintFactory::FactoryCreateNew(UClass* Class, UObject* InPar
 		ParentClass = UInterface::StaticClass();
 	}
 
-	if ( ( ParentClass == NULL ) || !FKismetEditorUtilities::CanCreateBlueprintOfClass(ParentClass) || !ParentClass->IsChildOf(AUserWidget::StaticClass()) )
+	if ( ( ParentClass == NULL ) || !FKismetEditorUtilities::CanCreateBlueprintOfClass(ParentClass) || !ParentClass->IsChildOf(UUserWidget::StaticClass()) )
 	{
 		FFormatNamedArguments Args;
 		Args.Add( TEXT("ClassName"), (ParentClass != NULL) ? FText::FromString( ParentClass->GetName() ) : LOCTEXT("Null", "(null)") );
diff --git a/Engine/Source/Runtime/UMG/Private/Components/ContentWidget.cpp b/Engine/Source/Runtime/UMG/Private/Components/ContentWidget.cpp
--- a/Engine/Source/Runtime/UMG/Private/Components/ContentWidget.cpp
+++ b/Engine/Source/Runtime/UMG/Private/Components/ContentWidget.cpp
@@ -15,10 +15,10 @@
 UContentWidget::UContentWidget(const FPostConstructInitializeProperties& PCIP)
 	: Super(PCIP)
 {
-	PrimaryComponentTick.bCanEverTick = true;
+	//PrimaryComponentTick.bCanEverTick = true;
 
-	bAutoActivate = true;
-	bTickInEditor = true;
+	//bAutoActivate = true;
+	//bTickInEditor = true;
 }
 
 USlateWrapperComponent* UContentWidget::GetContent()
diff --git a/Engine/Source/Runtime/UMG/Private/Components/SlateNonLeafWidgetComponent.cpp b/Engine/Source/Runtime/UMG/Private/Components/SlateNonLeafWidgetComponent.cpp
--- a/Engine/Source/Runtime/UMG/Private/Components/SlateNonLeafWidgetComponent.cpp
+++ b/Engine/Source/Runtime/UMG/Private/Components/SlateNonLeafWidgetComponent.cpp
@@ -15,10 +15,10 @@
 USlateNonLeafWidgetComponent::USlateNonLeafWidgetComponent(const FPostConstructInitializeProperties& PCIP)
 	: Super(PCIP)
 {
-	PrimaryComponentTick.bCanEverTick = true;
+//	PrimaryComponentTick.bCanEverTick = true;
 
-	bAutoActivate = true;
-	bTickInEditor = true;
+//	bAutoActivate = true;
+//	bTickInEditor = true;
 }
 
 /////////////////////////////////////////////////////
diff --git a/Engine/Source/Runtime/UMG/Private/Components/SlateWrapperComponent.cpp b/Engine/Source/Runtime/UMG/Private/Components/SlateWrapperComponent.cpp
--- a/Engine/Source/Runtime/UMG/Private/Components/SlateWrapperComponent.cpp
+++ b/Engine/Source/Runtime/UMG/Private/Components/SlateWrapperComponent.cpp
@@ -42,15 +42,10 @@ void USlateWrapperComponent::SetToolTipText(const FText& InToolTipText)
 	MyWidget->SetToolTipText(ToolTipText);
 }
 
-void USlateWrapperComponent::OnRegister()
-{
-	Super::OnRegister();
-}
-
 void USlateWrapperComponent::OnUnregister()
 {
 	MyWidget = NULL;
-	Super::OnUnregister();
+//	Super::OnUnregister();
 }
 
 void USlateWrapperComponent::PostEditChangeProperty(struct FPropertyChangedEvent& PropertyChangedEvent)
@@ -78,7 +73,7 @@ TSharedRef<SWidget> USlateWrapperComponent::GetWidget()
 TSharedRef<SWidget> USlateWrapperComponent::RebuildWidget()
 {
 	ensureMsg(false, TEXT("You must implement RebuildWidget() in your child class"));
-	return SNullWidget::NullWidget;
+	return SNew(SSpacer);
 }
 
 EVisibility USlateWrapperComponent::ConvertSerializedVisibilityToRuntime(TEnumAsByte<ESlateVisibility::Type> Input)
diff --git a/Engine/Source/Runtime/UMG/Private/Components/WidgetHost.cpp b/Engine/Source/Runtime/UMG/Private/Components/WidgetHost.cpp
new file mode 100644
--- /dev/null
+++ b/Engine/Source/Runtime/UMG/Private/Components/WidgetHost.cpp
@@ -0,0 +1,59 @@
+// Copyright 1998-2014 Epic Games, Inc. All Rights Reserved.
+
+#include "UMGPrivatePCH.h"
+
+/////////////////////////////////////////////////////
+// UWidgetHost
+
+UWidgetHost::UWidgetHost(const FPostConstructInitializeProperties& PCIP)
+	: Super(PCIP)
+{
+//	PrimaryComponentTick.bCanEverTick = true;
+
+	bIsVariable = true;
+
+	//bAutoActivate = true;
+	//bTickInEditor = true;
+}
+
+int32 UWidgetHost::GetChildrenCount() const
+{
+	return Widget != NULL ? 1 : 0;
+}
+
+USlateWrapperComponent* UWidgetHost::GetChildAt(int32 Index) const
+{
+	if ( Widget )
+	{
+		return Widget->GetRootWidgetComponent();
+	}
+	
+	return NULL;
+}
+
+TSharedRef<SWidget> UWidgetHost::RebuildWidget()
+{
+	if ( Widget )
+	{
+		//Widget = GetWorld()->SpawnActor<UUserWidget>(WidgetClass);
+		return Widget->GetRootWidget();
+	}
+
+	return SNew(SSpacer);
+}
+
+#if WITH_EDITOR
+void UWidgetHost::ConnectEditorData()
+{
+
+}
+
+void UWidgetHost::PostEditChangeProperty(struct FPropertyChangedEvent& PropertyChangedEvent)
+{
+
+}
+#endif
+
+/////////////////////////////////////////////////////
+
+#undef LOCTEXT_NAMESPACE
diff --git a/Engine/Source/Runtime/UMG/Private/UserWidget.cpp b/Engine/Source/Runtime/UMG/Private/UserWidget.cpp
--- a/Engine/Source/Runtime/UMG/Private/UserWidget.cpp
+++ b/Engine/Source/Runtime/UMG/Private/UserWidget.cpp
@@ -103,15 +103,14 @@ protected:
 
 
 /////////////////////////////////////////////////////
-// AUserWidget
+// UUserWidget
 
-AUserWidget::AUserWidget(const FPostConstructInitializeProperties& PCIP)
+UUserWidget::UUserWidget(const FPostConstructInitializeProperties& PCIP)
 	: Super(PCIP)
 {
-	RootWidget = SNullWidget::NullWidget;
 	Visiblity = ESlateVisibility::Visible;
-	PrimaryActorTick.bCanEverTick = true;
-	PrimaryActorTick.bTickEvenWhenPaused = true;
+	/*PrimaryActorTick.bCanEverTick = true;
+	PrimaryActorTick.bTickEvenWhenPaused = true;*/
 	bShowCursorWhenVisible = false;
 	bAbsoluteLayout = false;
 	bModal = false;
@@ -120,80 +119,51 @@ AUserWidget::AUserWidget(const FPostConstructInitializeProperties& PCIP)
 	VerticalAlignment = VAlign_Fill;
 }
 
-void AUserWidget::OnConstruction(const FTransform& Transform)
+void UUserWidget::PostInitProperties()
 {
-	Super::OnConstruction(Transform);
+	Super::PostInitProperties();
 
-	RebuildWrapperWidget();
-}
-
-void AUserWidget::Destroyed()
-{
-	RootWidget.Reset();
-	Super::Destroyed();
-}
+	Components.Reset();
 
-void AUserWidget::RerunConstructionScripts()
-{
-	Super::RerunConstructionScripts();
+	// Only do this if this widget is of a blueprint class
+	UWidgetBlueprintGeneratedClass* BGClass = Cast<UWidgetBlueprintGeneratedClass>(GetClass());
+	if ( BGClass != NULL )
+	{
+		BGClass->InitializeWidget(this);
+	}
 
 	RebuildWrapperWidget();
 }
 
-void AUserWidget::TickActor(float DeltaTime, enum ELevelTick TickType, FActorTickFunction& ThisTickFunction)
+UWorld* UUserWidget::GetWorld() const
 {
-	//if ( BindingEvaluator )
-	//{
-	//	// TODO UMG pass in this as a parameter, so it has the context to set properties?
-	//	ProcessEvent(BindingEvaluator, NULL);
-	//}
-	//static AUserWidget* TestWidget;
-	//static FWidgetNode_Base* LinkedPtr = NULL;
-
-	//FText GlueText()
-	//{
-	//	LinkedPtr->EvaluateGraphExposedInputs.Execute(TestWidget);
-	//	FWidgetNode_Text* TextNode = (FWidgetNode_Text*)LinkedPtr;
-	//	return TextNode->Text;
-	//}
-
-	//TextBlock->LinkedTextAttribute = TAttribute<FText>::Create(TAttribute<FText>::FGetter::CreateStatic(&GlueText));
-
-	//UWorld* World = GetWorld();
-	//if ( World && World->IsGameWorld() )
-	//{
-	//	if ( UWidgetBlueprintGeneratedClass* WidgetBlueprintClass = Cast<UWidgetBlueprintGeneratedClass>(GetClass()) )
-	//	{
-	//		for ( UStructProperty* Attribute : WidgetBlueprintClass->WidgetNodeProperties )
-	//		{
-
-	//		}
-	//	}
+	UObject* Outer = GetOuter();
+	if ( Outer == NULL )
+	{
+		return NULL;
+	}
 
+	// TODO UMG Global UI elements should go where?  who will be their outer.  Currently CreateWidget node makes the level their owner.
+	if ( ULevel* Level = Cast<ULevel>(Outer) )
+	{
+		return Level->OwningWorld;
+	}
 
-	//	//if ( bShowCursorWhenVisible && GetIsVisible() )
-	//	//{
-	//	//	World->GetFirstLocalPlayerFromController()->PlayerController->bShowMouseCursor = true;
-	//	//	UGameViewportClient* Viewport = World->GetGameViewport();
-	//	//}
-	//}
+	return NULL;
 }
 
-USlateWrapperComponent* AUserWidget::GetWidgetHandle(TSharedRef<SWidget> InWidget)
+USlateWrapperComponent* UUserWidget::GetWidgetHandle(TSharedRef<SWidget> InWidget)
 {
 	return WidgetToComponent.FindRef(InWidget);
 }
 
-void AUserWidget::RebuildWrapperWidget()
+void UUserWidget::RebuildWrapperWidget()
 {
 	WidgetToComponent.Reset();
-	RootWidget = SNullWidget::NullWidget;
+	RootWidget = SNew(SSpacer);
 	
-	TArray<USlateWrapperComponent*> SlateWrapperComponents;
-	GetComponents(SlateWrapperComponents);
-
 	// Add the first component to the root of the widget surface.
-	if ( SlateWrapperComponents.Num() > 0 )
+	if ( Components.Num() > 0 )
 	{
 		if ( bAbsoluteLayout )
 		{
@@ -206,7 +176,7 @@ void AUserWidget::RebuildWrapperWidget()
 				.VAlign(VerticalAlignment)
 				.HAlign(HorizontalAlignment)
 				[
-					SlateWrapperComponents[0]->GetWidget()
+					Components[0]->GetWidget()
 				];
 		}
 		else
@@ -218,7 +188,7 @@ void AUserWidget::RebuildWrapperWidget()
 				.HAlign(HorizontalAlignment)
 				.VAlign(VerticalAlignment)
 				[
-					SlateWrapperComponents[0]->GetWidget()
+					Components[0]->GetWidget()
 				];
 
 			NewSlot.SizeParam = USlateWrapperComponent::ConvertSerializedSizeParamToRuntime(Size);
@@ -228,9 +198,9 @@ void AUserWidget::RebuildWrapperWidget()
 	}
 
 	// Place all of our top-level children Slate wrapped components into the overlay
-	for (int32 ComponentIndex = 0; ComponentIndex < SlateWrapperComponents.Num(); ++ComponentIndex)
+	for ( int32 ComponentIndex = 0; ComponentIndex < Components.Num(); ++ComponentIndex )
 	{
-		USlateWrapperComponent* Handle = SlateWrapperComponents[ComponentIndex];
+		USlateWrapperComponent* Handle = Components[ComponentIndex];
 		TSharedRef<SWidget> Widget = Handle->GetWidget();
 
 		WidgetToComponent.Add(Widget, Handle);
@@ -261,7 +231,7 @@ void AUserWidget::RebuildWrapperWidget()
 	}
 }
 
-TSharedPtr<SWidget> AUserWidget::GetWidgetFromName(const FString& Name) const
+TSharedPtr<SWidget> UUserWidget::GetWidgetFromName(const FString& Name) const
 {
 	for ( auto& Entry : WidgetToComponent )
 	{
@@ -274,7 +244,7 @@ TSharedPtr<SWidget> AUserWidget::GetWidgetFromName(const FString& Name) const
 	return TSharedPtr<SWidget>();
 }
 
-USlateWrapperComponent* AUserWidget::GetHandleFromName(const FString& Name) const
+USlateWrapperComponent* UUserWidget::GetHandleFromName(const FString& Name) const
 {
 	for ( auto& Entry : WidgetToComponent )
 	{
@@ -287,12 +257,22 @@ USlateWrapperComponent* AUserWidget::GetHandleFromName(const FString& Name) cons
 	return NULL;
 }
 
-TSharedRef<SWidget> AUserWidget::GetRootWidget()
+TSharedRef<SWidget> UUserWidget::GetRootWidget()
 {
 	return RootWidget.ToSharedRef();
 }
 
-void AUserWidget::Show()
+USlateWrapperComponent* UUserWidget::GetRootWidgetComponent()
+{
+	if ( Components.Num() > 0 )
+	{
+		return Components[0];
+	}
+
+	return NULL;
+}
+
+void UUserWidget::Show()
 {
 	RootWidget->SetVisibility(EVisibility::Visible);
 	OnVisibilityChanged.Broadcast(ESlateVisibility::Visible);
@@ -311,7 +291,7 @@ void AUserWidget::Show()
 	}
 }
 
-void AUserWidget::Hide()
+void UUserWidget::Hide()
 {
 	RootWidget->SetVisibility(EVisibility::Hidden);
 	OnVisibilityChanged.Broadcast(ESlateVisibility::Hidden);
@@ -332,12 +312,12 @@ void AUserWidget::Hide()
 	}
 }
 
-bool AUserWidget::GetIsVisible()
+bool UUserWidget::GetIsVisible()
 {
 	return RootWidget->GetVisibility().IsVisible();
 }
 
-TEnumAsByte<ESlateVisibility::Type> AUserWidget::GetVisiblity()
+TEnumAsByte<ESlateVisibility::Type> UUserWidget::GetVisiblity()
 {
 	return USlateWrapperComponent::ConvertRuntimeToSerializedVisiblity(RootWidget->GetVisibility());
 }
diff --git a/Engine/Source/Runtime/UMG/Private/WidgetActor.cpp b/Engine/Source/Runtime/UMG/Private/WidgetActor.cpp
new file mode 100644
--- /dev/null
+++ b/Engine/Source/Runtime/UMG/Private/WidgetActor.cpp
@@ -0,0 +1,15 @@
+// Copyright 1998-2014 Epic Games, Inc. All Rights Reserved.
+
+#include "UMGPrivatePCH.h"
+
+AWidgetActor::AWidgetActor(const FPostConstructInitializeProperties& PCIP)
+	: Super(PCIP)
+{
+	PrimaryActorTick.bCanEverTick = true;
+	PrimaryActorTick.bTickEvenWhenPaused = true;
+}
+
+void AWidgetActor::TickActor(float DeltaTime, enum ELevelTick TickType, FActorTickFunction& ThisTickFunction)
+{
+
+}
diff --git a/Engine/Source/Runtime/UMG/Private/WidgetBlueprintGeneratedClass.cpp b/Engine/Source/Runtime/UMG/Private/WidgetBlueprintGeneratedClass.cpp
--- a/Engine/Source/Runtime/UMG/Private/WidgetBlueprintGeneratedClass.cpp
+++ b/Engine/Source/Runtime/UMG/Private/WidgetBlueprintGeneratedClass.cpp
@@ -33,10 +33,8 @@ void UWidgetBlueprintGeneratedClass::Link(FArchive& Ar, bool bRelinkExistingProp
 	}
 }
 
-void UWidgetBlueprintGeneratedClass::CreateComponentsForActor(AActor* Actor) const
+void UWidgetBlueprintGeneratedClass::InitializeWidget(UUserWidget* Actor) const
 {
-	Super::CreateComponentsForActor(Actor);
-
 	// Duplicate the graph, keeping track of what was duplicated
 	TMap<UObject*, UObject*> DuplicatedObjectList;
 
@@ -45,7 +43,7 @@ void UWidgetBlueprintGeneratedClass::CreateComponentsForActor(AActor* Actor) con
 
 	UWidgetTree* ClonedTree = CastChecked<UWidgetTree>(StaticDuplicateObjectEx(Parameters));
 
-	AUserWidget* WidgetActor = CastChecked<AUserWidget>(Actor);
+	UUserWidget* WidgetActor = CastChecked<UUserWidget>(Actor);
 	UClass* ActorClass = Actor->GetClass();
 
 	for ( USlateWrapperComponent* Widget : ClonedTree->WidgetTemplates )
@@ -56,13 +54,13 @@ void UWidgetBlueprintGeneratedClass::CreateComponentsForActor(AActor* Actor) con
 			continue;
 		}
 
-		FName NewName = *( FString::Printf(TEXT("WidgetComponent__%d"), Actor->SerializedComponents.Num()) );
+		FName NewName = *( FString::Printf(TEXT("WidgetComponent__%d"), Actor->Components.Num()) );
 
 		FString VariableName = Widget->GetName();
 
 		Widget->bCreatedByConstructionScript = true; // Indicate it comes from a blueprint so it gets cleared when we rerun construction scripts
-		Actor->SerializedComponents.Add(Widget); // Add to array so it gets saved
-		Widget->SetNetAddressable();	// This component has a stable name that can be referenced for replication
+		Actor->Components.Add(Widget); // Add to array so it gets saved
+//		Widget->SetNetAddressable();	// This component has a stable name that can be referenced for replication
 
 		// Find property with the same name as the template and assign the new Timeline to it
 		UObjectPropertyBase* Prop = FindField<UObjectPropertyBase>(ActorClass, *VariableName);
@@ -96,7 +94,7 @@ void UWidgetBlueprintGeneratedClass::CreateComponentsForActor(AActor* Actor) con
 			}
 		}
 
-		Widget->RegisterComponent();
+//		Widget->RegisterComponent();
 
 #if WITH_EDITOR
 		Widget->ConnectEditorData();
diff --git a/Engine/Source/Runtime/UMG/Private/WidgetBlueprintLibrary.cpp b/Engine/Source/Runtime/UMG/Private/WidgetBlueprintLibrary.cpp
--- a/Engine/Source/Runtime/UMG/Private/WidgetBlueprintLibrary.cpp
+++ b/Engine/Source/Runtime/UMG/Private/WidgetBlueprintLibrary.cpp
@@ -14,7 +14,7 @@ UWidgetBlueprintLibrary::UWidgetBlueprintLibrary(const FPostConstructInitializeP
 {
 }
 
-AUserWidget* UWidgetBlueprintLibrary::Create(UObject* WorldContextObject, TSubclassOf<AUserWidget> WidgetType)
+UUserWidget* UWidgetBlueprintLibrary::Create(UObject* WorldContextObject, TSubclassOf<UUserWidget> WidgetType)
 {
 	if ( WidgetType->HasAnyClassFlags(CLASS_Abstract) )
 	{
@@ -22,11 +22,12 @@ AUserWidget* UWidgetBlueprintLibrary::Create(UObject* WorldContextObject, TSubcl
 		return NULL;
 	}
 
-	//TODO ConstructObject after AUserWidget is no longer an actor.
+	//TODO ConstructObject after UUserWidget is no longer an actor.
 
 	UWorld* World = GEngine->GetWorldFromContextObject(WorldContextObject);
-	AUserWidget* NewWidget = World->SpawnActor<AUserWidget>(WidgetType);
-	
+	UUserWidget* NewWidget = ConstructObject<UUserWidget>(WidgetType, World->GetCurrentLevel());
+	//UUserWidget* NewWidget = World->SpawnActor<UUserWidget>(WidgetType);
+		
 	return NewWidget;
 }
 
diff --git a/Engine/Source/Runtime/UMG/Public/Blueprint/UserWidget.h b/Engine/Source/Runtime/UMG/Public/Blueprint/UserWidget.h
--- a/Engine/Source/Runtime/UMG/Public/Blueprint/UserWidget.h
+++ b/Engine/Source/Runtime/UMG/Public/Blueprint/UserWidget.h
@@ -4,10 +4,12 @@
 
 #include "UserWidget.generated.h"
 
+//TODO UMG If you want to host a widget that's full screen there may need to be a SWindow equivalent that you spawn it into.
+
 DECLARE_DYNAMIC_MULTICAST_DELEGATE_OneParam(FOnVisibilityChangedEvent, ESlateVisibility::Type, Visibility);
 
-UCLASS(Abstract, hideCategories=(Object, Actor, Replication, Rendering, Input))
-class UMG_API AUserWidget : public AActor
+UCLASS(Abstract, editinlinenew, hideCategories=(Object, Actor, Replication, Rendering, Input))
+class UMG_API UUserWidget : public UObject
 {
 	GENERATED_UCLASS_BODY()
 
@@ -75,12 +77,13 @@ class UMG_API AUserWidget : public AActor
 	UPROPERTY(EditAnywhere, Category=Appearance)
 	TEnumAsByte<EVerticalAlignment> VerticalAlignment;
 
-	// AActor interface
-	virtual void OnConstruction(const FTransform& Transform) OVERRIDE;
-	virtual void Destroyed() OVERRIDE;
-	virtual void RerunConstructionScripts() OVERRIDE;
-	virtual void TickActor(float DeltaTime, enum ELevelTick TickType, FActorTickFunction& ThisTickFunction) OVERRIDE;
-	// End of AActor interface
+	UPROPERTY()
+	TArray<USlateWrapperComponent*> Components;
+	
+	//UObject interface
+	virtual void PostInitProperties() OVERRIDE;
+	virtual class UWorld* GetWorld() const OVERRIDE;
+	// End of UObject interface
 
 	/*  */
 	UFUNCTION(BlueprintCallable, Category="Appearance")
@@ -99,6 +102,7 @@ class UMG_API AUserWidget : public AActor
 	USlateWrapperComponent* GetWidgetHandle(TSharedRef<SWidget> InWidget);
 
 	TSharedRef<SWidget> GetRootWidget();
+	USlateWrapperComponent* GetRootWidgetComponent();
 
 	TSharedPtr<SWidget> GetWidgetFromName(const FString& Name) const;
 	USlateWrapperComponent* GetHandleFromName(const FString& Name) const;
diff --git a/Engine/Source/Runtime/UMG/Public/Blueprint/WidgetActor.h b/Engine/Source/Runtime/UMG/Public/Blueprint/WidgetActor.h
new file mode 100644
--- /dev/null
+++ b/Engine/Source/Runtime/UMG/Public/Blueprint/WidgetActor.h
@@ -0,0 +1,18 @@
+// Copyright 1998-2014 Epic Games, Inc. All Rights Reserved.
+
+#pragma once
+
+#include "WidgetActor.generated.h"
+
+UCLASS()
+class UMG_API AWidgetActor : public AActor
+{
+	GENERATED_UCLASS_BODY()
+
+	// AActor interface
+	//virtual void OnConstruction(const FTransform& Transform);
+	//virtual void Destroyed();
+	//virtual void RerunConstructionScripts();
+	virtual void TickActor(float DeltaTime, enum ELevelTick TickType, FActorTickFunction& ThisTickFunction);
+	// End of AActor interface
+};
diff --git a/Engine/Source/Runtime/UMG/Public/Blueprint/WidgetBlueprint.h b/Engine/Source/Runtime/UMG/Public/Blueprint/WidgetBlueprint.h
--- a/Engine/Source/Runtime/UMG/Public/Blueprint/WidgetBlueprint.h
+++ b/Engine/Source/Runtime/UMG/Public/Blueprint/WidgetBlueprint.h
@@ -33,7 +33,7 @@ struct UMG_API FDelegateEditorBinding
 	FDelegateRuntimeBinding ToRuntimeBinding(class UWidgetBlueprint* Blueprint) const;
 };
 
-UCLASS(dependson=(UBlueprint, AUserWidget, UWidgetBlueprintGeneratedClass), BlueprintType)
+UCLASS(dependson=(UBlueprint, UUserWidget, UWidgetBlueprintGeneratedClass), BlueprintType)
 class UMG_API UWidgetBlueprint : public UBlueprint
 {
 	GENERATED_UCLASS_BODY()
diff --git a/Engine/Source/Runtime/UMG/Public/Blueprint/WidgetBlueprintGeneratedClass.h b/Engine/Source/Runtime/UMG/Public/Blueprint/WidgetBlueprintGeneratedClass.h
--- a/Engine/Source/Runtime/UMG/Public/Blueprint/WidgetBlueprintGeneratedClass.h
+++ b/Engine/Source/Runtime/UMG/Public/Blueprint/WidgetBlueprintGeneratedClass.h
@@ -41,7 +41,7 @@ class UMG_API UWidgetBlueprintGeneratedClass : public UBlueprintGeneratedClass
 	/** This is transient data calculated at link time. */
 	TArray<UStructProperty*> WidgetNodeProperties;
 
-	virtual void CreateComponentsForActor(AActor* Actor) const OVERRIDE;
-
 	virtual void Link(FArchive& Ar, bool bRelinkExistingProperties) OVERRIDE;
+
+	void InitializeWidget(class UUserWidget* Actor) const;
 };
\ No newline at end of file
diff --git a/Engine/Source/Runtime/UMG/Public/Blueprint/WidgetBlueprintLibrary.h b/Engine/Source/Runtime/UMG/Public/Blueprint/WidgetBlueprintLibrary.h
--- a/Engine/Source/Runtime/UMG/Public/Blueprint/WidgetBlueprintLibrary.h
+++ b/Engine/Source/Runtime/UMG/Public/Blueprint/WidgetBlueprintLibrary.h
@@ -10,5 +10,5 @@ class UWidgetBlueprintLibrary : public UBlueprintFunctionLibrary
 
 	/** Creates a widget */
 	UFUNCTION(BlueprintCallable, meta=( HidePin="WorldContextObject", DefaultToSelf="WorldContextObject", UnsafeDuringActorConstruction = "true", FriendlyName = "Create Widget" ), Category="User Interface|Widget")
-	static class AUserWidget* Create(UObject* WorldContextObject, TSubclassOf<class AUserWidget> WidgetType);
+	static class UUserWidget* Create(UObject* WorldContextObject, TSubclassOf<class UUserWidget> WidgetType);
 };
\ No newline at end of file
diff --git a/Engine/Source/Runtime/UMG/Public/Blueprint/WidgetNode_Base.h b/Engine/Source/Runtime/UMG/Public/Blueprint/WidgetNode_Base.h
--- a/Engine/Source/Runtime/UMG/Public/Blueprint/WidgetNode_Base.h
+++ b/Engine/Source/Runtime/UMG/Public/Blueprint/WidgetNode_Base.h
@@ -14,7 +14,7 @@ struct FExposedWidgetValueHandler
 		UPROPERTY()
 		FName BoundFunction;
 
-	void Execute(AUserWidget* Widget) const
+	void Execute(UUserWidget* Widget) const
 	{
 		if ( BoundFunction != NAME_None )
 		{
@@ -26,7 +26,7 @@ struct FExposedWidgetValueHandler
 		}
 	}
 
-	UFunction* GetFunction(AUserWidget* Widget)
+	UFunction* GetFunction(UUserWidget* Widget)
 	{
 		if ( BoundFunction != NAME_None )
 		{
diff --git a/Engine/Source/Runtime/UMG/Public/Components/SlateWrapperComponent.h b/Engine/Source/Runtime/UMG/Public/Components/SlateWrapperComponent.h
--- a/Engine/Source/Runtime/UMG/Public/Components/SlateWrapperComponent.h
+++ b/Engine/Source/Runtime/UMG/Public/Components/SlateWrapperComponent.h
@@ -10,7 +10,7 @@
 
 /** This is the base class for all Slate widget wrapper components */
 UCLASS(Abstract, hideCategories=(Activation, "Components|Activation", ComponentReplication, LOD, Rendering))
-class UMG_API USlateWrapperComponent : public UActorComponent
+class UMG_API USlateWrapperComponent : public UObject
 {
 	GENERATED_UCLASS_BODY()
 
@@ -59,6 +59,9 @@ class UMG_API USlateWrapperComponent : public UActorComponent
 	UPROPERTY(EditDefaultsOnly, Category=Behavior)
 	TEnumAsByte<EMouseCursor::Type> Cursor;
 
+	UPROPERTY(Transient)
+	bool bCreatedByConstructionScript;
+
 	UFUNCTION(BlueprintCallable, Category="Widget")
 	bool GetIsEnabled() const;
 
@@ -79,8 +82,7 @@ class UMG_API USlateWrapperComponent : public UActorComponent
 	//void SetCursor(const TAttribute< TOptional< TEnumAsByte<EMouseCursor::Type> > >& InCursor);
 
 	// UActorComponent interface
-	virtual void OnRegister() OVERRIDE;
-	virtual void OnUnregister() OVERRIDE;
+	virtual void OnUnregister();
 	// End of UActorComponent interface
 
 	TSharedRef<SWidget> GetWidget();
diff --git a/Engine/Source/Runtime/UMG/Public/Components/WidgetHost.h b/Engine/Source/Runtime/UMG/Public/Components/WidgetHost.h
new file mode 100644
--- /dev/null
+++ b/Engine/Source/Runtime/UMG/Public/Components/WidgetHost.h
@@ -0,0 +1,37 @@
+// Copyright 1998-2014 Epic Games, Inc. All Rights Reserved.
+
+#pragma once
+
+#include "WidgetHost.generated.h"
+
+UCLASS(meta=(BlueprintSpawnableComponent), ClassGroup=UserInterface)
+class UMG_API UWidgetHost : public USlateNonLeafWidgetComponent
+{
+	GENERATED_UCLASS_BODY()
+
+	UPROPERTY(EditDefaultsOnly, EditInline, Category=Appearance)
+	class UUserWidget* Widget;
+
+	// USlateNonLeafWidgetComponent
+	virtual int32 GetChildrenCount() const OVERRIDE;
+	virtual USlateWrapperComponent* GetChildAt(int32 Index) const OVERRIDE;
+	//virtual bool AddChild(USlateWrapperComponent* Child, FVector2D Position) OVERRIDE;
+	//virtual bool RemoveChild(USlateWrapperComponent* Child) OVERRIDE;
+	//virtual void ReplaceChildAt(int32 Index, USlateWrapperComponent* Child) OVERRIDE;
+	// End USlateNonLeafWidgetComponent
+
+#if WITH_EDITOR
+	// UObject interface
+	virtual void PostEditChangeProperty(struct FPropertyChangedEvent& PropertyChangedEvent) OVERRIDE;
+	// End of UObject interface
+
+	// USlateWrapperComponent interface
+	virtual void ConnectEditorData() OVERRIDE;
+	// End USlateWrapperComponent interface
+#endif
+
+protected:
+	// USlateWrapperComponent interface
+	virtual TSharedRef<SWidget> RebuildWidget() OVERRIDE;
+	// End of USlateWrapperComponent interface
+};
diff --git a/Engine/Source/Runtime/UMG/Public/UMG.h b/Engine/Source/Runtime/UMG/Public/UMG.h
--- a/Engine/Source/Runtime/UMG/Public/UMG.h
+++ b/Engine/Source/Runtime/UMG/Public/UMG.h
@@ -32,6 +32,8 @@
 
 #include "ContentWidget.h"
 
+#include "WidgetHost.h"
+
 #include "BoxPanelComponent.h"
 
 #include "CanvasPanelSlot.h"
@@ -74,7 +76,8 @@
 
 #include "WidgetBlueprint.h"
 #include "WidgetBlueprintGeneratedClass.h"
- 
+
+#include "WidgetActor.h"
 
 /* Interfaces
  *****************************************************************************/
