UMG - Removing the widget host, making the widgets and blueprint widget share a baseclass soon to allow for direct nesting.  Fixing some issues with spawning fullscreen widgets and releasing references correctly.  Adding lots of implementable methods that will allow blueprint authors to receieve and handle the events slate sends, this required updating a few slate events to be compatiable as UStructs.

diff --git a/Engine/Source/Editor/UMGEditor/Private/SUMGDesigner.cpp b/Engine/Source/Editor/UMGEditor/Private/SUMGDesigner.cpp
--- a/Engine/Source/Editor/UMGEditor/Private/SUMGDesigner.cpp
+++ b/Engine/Source/Editor/UMGEditor/Private/SUMGDesigner.cpp
@@ -99,6 +99,7 @@ static bool LocateWidgetsUnderCursor_Helper(FArrangedWidget& Candidate, FVector2
 
 void SUMGDesigner::Construct(const FArguments& InArgs, TSharedPtr<FWidgetBlueprintEditor> InBlueprintEditor)
 {
+	PreviewWidgetActor = NULL;
 	BlueprintEditor = InBlueprintEditor;
 	CurrentHandle = DH_NONE;
 	SelectedTemplate = NULL;
@@ -499,29 +500,34 @@ FCursorReply SUMGDesigner::OnCursorQuery(const FGeometry& MyGeometry, const FPoi
 void SUMGDesigner::Tick(const FGeometry& AllottedGeometry, const double InCurrentTime, const float InDeltaTime)
 {
 	UUserWidget* WidgetActor = BlueprintEditor.Pin()->GetPreview();
-	if ( WidgetActor )
-	{
-		TSharedRef<SWidget> CurrentWidget = WidgetActor->GetRootWidget();
 
-		if ( CurrentWidget != PreviewWidget.Pin() )
+	if ( WidgetActor != PreviewWidgetActor )
+	{
+		PreviewWidgetActor = WidgetActor;
+		if ( PreviewWidgetActor )
 		{
-			PreviewWidget = CurrentWidget;
-			PreviewSurface->SetContent(CurrentWidget);
+			TSharedRef<SWidget> CurrentWidget = PreviewWidgetActor->MakeFullScreenWidget();
+
+			if ( CurrentWidget != PreviewWidget.Pin() )
+			{
+				PreviewWidget = CurrentWidget;
+				PreviewSurface->SetContent(CurrentWidget);
+			}
 		}
-	}
-	else
-	{
-		ChildSlot
-		[
-			SNew(SHorizontalBox)
-			+SHorizontalBox::Slot()
-			.HAlign(HAlign_Center)
-			.VAlign(VAlign_Center)
+		else
+		{
+			ChildSlot
 			[
-				SNew(STextBlock)
-				.Text(LOCTEXT("NoWrappedWidget", "No Widget Preview"))
-			]
-		];
+				SNew(SHorizontalBox)
+				+ SHorizontalBox::Slot()
+				.HAlign(HAlign_Center)
+				.VAlign(VAlign_Center)
+				[
+					SNew(STextBlock)
+					.Text(LOCTEXT("NoWrappedWidget", "No Widget Preview"))
+				]
+			];
+		}
 	}
 
 	// Update the selected widget to match the selected template.
diff --git a/Engine/Source/Editor/UMGEditor/Private/SUMGDesigner.h b/Engine/Source/Editor/UMGEditor/Private/SUMGDesigner.h
--- a/Engine/Source/Editor/UMGEditor/Private/SUMGDesigner.h
+++ b/Engine/Source/Editor/UMGEditor/Private/SUMGDesigner.h
@@ -72,6 +72,7 @@ private:
 	TWeakPtr<FWidgetBlueprintEditor> BlueprintEditor;
 
 	TWeakPtr<SWidget> PreviewWidget;
+	UUserWidget* PreviewWidgetActor;
 
 	USlateWrapperComponent* SelectedTemplate;
 	TWeakPtr<SWidget> SelectedWidget;
diff --git a/Engine/Source/Editor/UnrealEd/Public/Toolkits/IToolkit.h b/Engine/Source/Editor/UnrealEd/Public/Toolkits/IToolkit.h
--- a/Engine/Source/Editor/UnrealEd/Public/Toolkits/IToolkit.h
+++ b/Engine/Source/Editor/UnrealEd/Public/Toolkits/IToolkit.h
@@ -97,7 +97,7 @@ public:
 	 *
 	 * @return true if an action was processed
 	 */
-	virtual bool ProcessCommandBindings( const class FKeyboardEvent& InKeyboardEvent ) const = 0;
+	virtual bool ProcessCommandBindings( const struct FKeyboardEvent& InKeyboardEvent ) const = 0;
 
 	/** Call this function to bring all of this toolkit's tabs to the foreground in their respective stacks.  Also causes the toolkit's host window to be foregrounded, too! */
 	virtual void BringToolkitToFront() = 0;
diff --git a/Engine/Source/Runtime/SlateCore/Public/Input/Events.h b/Engine/Source/Runtime/SlateCore/Public/Input/Events.h
--- a/Engine/Source/Runtime/SlateCore/Public/Input/Events.h
+++ b/Engine/Source/Runtime/SlateCore/Public/Input/Events.h
@@ -6,12 +6,14 @@
 
 #pragma once
 
-class SWindow;
+#include "Events.generated.h"
 
+class SWindow;
 
 /**
  * Context for keyboard focus change
  */
+UENUM()
 namespace EKeyboardFocusCause
 {
 	enum Type
@@ -34,17 +36,27 @@ namespace EKeyboardFocusCause
 		/** Keyboard focus was set in response to the owning window being activated */
 		WindowActivate,
 	};
-};
+}
 
 
 /**
  * FKeyboardFocusEvent is used when notifying widgets about keyboard focus changes
  * It is passed to event handlers dealing with keyboard focus
  */
-class FKeyboardFocusEvent
+USTRUCT()
+struct FKeyboardFocusEvent
 {
+	GENERATED_USTRUCT_BODY()
+
 public:
 
+	/**
+	 * UStruct Constructor.  Not meant for normal usage.
+	 */
+	FKeyboardFocusEvent()
+		: Cause(EKeyboardFocusCause::SetDirectly)
+	{ }
+
 	/**
 	 * Constructor.  Events are immutable once constructed.
 	 *
@@ -75,10 +87,22 @@ private:
 /**
  * Base class for all mouse and keyboard events.
  */
-class FInputEvent
+USTRUCT()
+struct FInputEvent
 {
+	GENERATED_USTRUCT_BODY()
+
 public:
 
+	/**
+	 * UStruct Constructor.  Not meant for normal usage.
+	 */
+	FInputEvent()
+		: ModifierKeys(FModifierKeysState(false, false, false, false, false, false))
+		, bIsRepeat(false)
+		, EventPath(nullptr)
+	{ }
+
 	/**
 	 * Constructor.  Events are immutable once constructed.
 	 *
@@ -232,10 +256,22 @@ protected:
  * FKeyboardEvent describes a keyboard action (key pressed or released.)
  * It is passed to event handlers dealing with keyboard input.
  */
-class FKeyboardEvent
+USTRUCT()
+struct FKeyboardEvent
 	: public FInputEvent
 {
+	GENERATED_USTRUCT_BODY()
+
 public:
+	/**
+	 * UStruct Constructor.  Not meant for normal usage.
+	 */
+	FKeyboardEvent()
+		: FInputEvent(FModifierKeysState(false, false, false, false, false, false), false)
+		, Key(EKeys::SpaceBar)
+		, CharacterCode(0)
+	{
+	}
 
 	/**
 	 * Constructor.  Events are immutable once constructed.
@@ -287,10 +323,22 @@ private:
 /**
  * FCharacterEvent describes a keyboard action where the utf-16 code is given.  Used for OnKeyChar messages
  */
-class FCharacterEvent
+USTRUCT()
+struct FCharacterEvent
 	: public FInputEvent
 {
+	GENERATED_USTRUCT_BODY()
+
 public:
+	/**
+	 * UStruct Constructor.  Not meant for normal usage.
+	 */
+	FCharacterEvent()
+		: FInputEvent(FModifierKeysState(false, false, false, false, false, false), false)
+		, Character(0)
+	{
+	}
+
 	FCharacterEvent( const TCHAR InCharacter, const FModifierKeysState& InModifierKeys, const bool bInIsRepeat )
 		: FInputEvent(InModifierKeys, bInIsRepeat)
 		, Character(InCharacter)
@@ -345,11 +393,30 @@ public:
  * FPointerEvent describes a mouse or touch action (e.g. Press, Release, Move, etc).
  * It is passed to event handlers dealing with pointer-based input.
  */
-class FPointerEvent
+USTRUCT()
+struct FPointerEvent
 	: public FInputEvent
 {
+	GENERATED_USTRUCT_BODY()
 public:
 
+	/**
+	 * UStruct Constructor.  Not meant for normal usage.
+	 */
+	FPointerEvent()
+		: ScreenSpacePosition(FVector2D(0, 0))
+		, LastScreenSpacePosition(FVector2D(0, 0))
+		, CursorDelta(FVector2D(0, 0))
+		, PressedButtons(FTouchKeySet::EmptySet)
+		, EffectingButton(EKeys::LeftMouseButton)
+		, UserIndex(0)
+		, PointerIndex(0)
+		, TouchpadIndex(0)
+		, bIsTouchEvent(false)
+		, GestureType(EGestureEvent::None)
+		, WheelOrGestureDelta(0.0f, 0)
+	{ }
+
 	/** Events are immutable once constructed. */
 	FPointerEvent(
 		const FVector2D& InScreenSpacePosition,
@@ -409,7 +476,7 @@ public:
 		, LastScreenSpacePosition(InLastScreenSpacePosition)
 		, CursorDelta(LastScreenSpacePosition - ScreenSpacePosition)
 		, PressedButtons(bPressLeftMouseButton ? FTouchKeySet::StandardSet : FTouchKeySet::EmptySet)
-		, EffectingButton (EKeys::LeftMouseButton)
+		, EffectingButton(EKeys::LeftMouseButton)
 		, UserIndex(InUserIndex)
 		, PointerIndex(InPointerIndex)
 		, TouchpadIndex(InTouchpadIndex)
@@ -482,7 +549,7 @@ private:
 	FVector2D ScreenSpacePosition;
 	FVector2D LastScreenSpacePosition;
 	FVector2D CursorDelta;
-	const TSet<FKey>& PressedButtons;
+	TSet<FKey> PressedButtons;
 	FKey EffectingButton;
 	uint32 UserIndex;
 	uint32 PointerIndex;
@@ -497,11 +564,22 @@ private:
  * FControllerEvent describes a controller action (e.g. Button Press, Release, Analog stick move, etc).
  * It is passed to event handlers dealing with controller input.
  */
-class FControllerEvent
+USTRUCT()
+struct FControllerEvent
 	: public FInputEvent
 {
+	GENERATED_USTRUCT_BODY()
 public:
 
+	/**
+	 * UStruct Constructor.  Not meant for normal usage.
+	 */
+	FControllerEvent()
+		: EffectingButton(EKeys::Gamepad_RightTrigger)
+		, UserIndex(0)
+		, AnalogValue(0)
+	{ }
+
 	FControllerEvent( FKey InEffectingButton, int32 InUserIndex, float InAnalogValue, bool bIsRepeat )
 		: FInputEvent(FModifierKeysState(false, false, false, false, false, false), bIsRepeat)
 		, EffectingButton(InEffectingButton)
@@ -537,10 +615,24 @@ private:
  * FMotionEvent describes a touch pad action (press, move, lift)
  * It is passed to event handlers dealing with touch input.
  */
-class FMotionEvent
+USTRUCT()
+struct FMotionEvent
 	: public FInputEvent
 {
+	GENERATED_USTRUCT_BODY()
+
 public:
+	/**
+	* UStruct Constructor.  Not meant for normal usage.
+	*/
+	FMotionEvent()
+		: UserIndex(0)
+		, Tilt(FVector(0, 0, 0))
+		, RotationRate(FVector(0, 0, 0))
+		, Gravity(FVector(0, 0, 0))
+		, Acceleration(FVector(0, 0, 0))
+	{ }
+
 	FMotionEvent(
 		uint32 InUserIndex,
 		const FVector& InTilt, 
diff --git a/Engine/Source/Runtime/SlateCore/Public/Layout/Geometry.h b/Engine/Source/Runtime/SlateCore/Public/Layout/Geometry.h
--- a/Engine/Source/Runtime/SlateCore/Public/Layout/Geometry.h
+++ b/Engine/Source/Runtime/SlateCore/Public/Layout/Geometry.h
@@ -6,6 +6,8 @@
 
 #pragma once
 
+#include "Geometry.generated.h"
+
 
 class FArrangedWidget;
 class SWidget;
@@ -20,8 +22,10 @@ class SWidget;
  * A Geometry's parent is generally thought to be the Geometry of the
  * the corresponding parent widget.
  */
-class SLATECORE_API FGeometry
+USTRUCT()
+struct SLATECORE_API FGeometry
 {
+	GENERATED_USTRUCT_BODY()
 
 public:
 
diff --git a/Engine/Source/Runtime/UMG/Private/Components/WidgetHost.cpp b/Engine/Source/Runtime/UMG/Private/Components/WidgetHost.cpp
deleted file mode 100644
diff --git a/Engine/Source/Runtime/UMG/Private/Slate/SObjectWidget.cpp b/Engine/Source/Runtime/UMG/Private/Slate/SObjectWidget.cpp
new file mode 100644
--- /dev/null
+++ b/Engine/Source/Runtime/UMG/Private/Slate/SObjectWidget.cpp
@@ -0,0 +1,167 @@
+// Copyright 1998-2014 Epic Games, Inc. All Rights Reserved.
+
+#include "UMGPrivatePCH.h"
+
+void SObjectWidget::Construct(const FArguments& InArgs, UUserWidget* InWidgetObject)
+{
+	WidgetObject = InWidgetObject;
+
+	ChildSlot
+	[
+		InArgs._Content.Widget
+	];
+}
+
+void SObjectWidget::AddReferencedObjects(FReferenceCollector& Collector)
+{
+	Collector.AddReferencedObject(WidgetObject);
+}
+
+FReply SObjectWidget::OnKeyboardFocusReceived(const FGeometry& MyGeometry, const FKeyboardFocusEvent& InKeyboardFocusEvent)
+{
+	return WidgetObject->OnKeyboardFocusReceived(MyGeometry, InKeyboardFocusEvent).ToReply();
+}
+
+void SObjectWidget::OnKeyboardFocusLost(const FKeyboardFocusEvent& InKeyboardFocusEvent)
+{
+	WidgetObject->OnKeyboardFocusLost(InKeyboardFocusEvent);
+}
+
+void SObjectWidget::OnKeyboardFocusChanging(const FWeakWidgetPath& PreviousFocusPath, const FWidgetPath& NewWidgetPath)
+{
+	// TODO UMG
+}
+
+FReply SObjectWidget::OnKeyChar(const FGeometry& MyGeometry, const FCharacterEvent& InCharacterEvent)
+{
+	return WidgetObject->OnKeyChar(MyGeometry, InCharacterEvent).ToReply();
+}
+
+FReply SObjectWidget::OnPreviewKeyDown(const FGeometry& MyGeometry, const FKeyboardEvent& InKeyboardEvent)
+{
+	return WidgetObject->OnPreviewKeyDown(MyGeometry, InKeyboardEvent).ToReply();
+}
+
+FReply SObjectWidget::OnKeyDown(const FGeometry& MyGeometry, const FKeyboardEvent& InKeyboardEvent)
+{
+	return WidgetObject->OnKeyDown(MyGeometry, InKeyboardEvent).ToReply();
+}
+
+FReply SObjectWidget::OnKeyUp(const FGeometry& MyGeometry, const FKeyboardEvent& InKeyboardEvent)
+{
+	return WidgetObject->OnKeyUp(MyGeometry, InKeyboardEvent).ToReply();
+}
+
+FReply SObjectWidget::OnMouseButtonDown(const FGeometry& MyGeometry, const FPointerEvent& MouseEvent)
+{
+	return WidgetObject->OnMouseButtonDown(MyGeometry, MouseEvent).ToReply();
+}
+
+FReply SObjectWidget::OnPreviewMouseButtonDown(const FGeometry& MyGeometry, const FPointerEvent& MouseEvent)
+{
+	return WidgetObject->OnPreviewMouseButtonDown(MyGeometry, MouseEvent).ToReply();
+}
+
+FReply SObjectWidget::OnMouseButtonUp(const FGeometry& MyGeometry, const FPointerEvent& MouseEvent)
+{
+	return WidgetObject->OnMouseButtonUp(MyGeometry, MouseEvent).ToReply();
+}
+
+FReply SObjectWidget::OnMouseMove(const FGeometry& MyGeometry, const FPointerEvent& MouseEvent)
+{
+	return WidgetObject->OnMouseMove(MyGeometry, MouseEvent).ToReply();
+}
+
+void SObjectWidget::OnMouseEnter(const FGeometry& MyGeometry, const FPointerEvent& MouseEvent)
+{
+	WidgetObject->OnMouseEnter(MyGeometry, MouseEvent);
+}
+
+void SObjectWidget::OnMouseLeave(const FPointerEvent& MouseEvent)
+{
+	WidgetObject->OnMouseLeave(MouseEvent);
+}
+
+FReply SObjectWidget::OnMouseWheel(const FGeometry& MyGeometry, const FPointerEvent& MouseEvent)
+{
+	return WidgetObject->OnMouseWheel(MyGeometry, MouseEvent).ToReply();
+}
+
+FCursorReply SObjectWidget::OnCursorQuery(const FGeometry& MyGeometry, const FPointerEvent& CursorEvent) const
+{
+	// TODO UMG
+	return FCursorReply::Unhandled();
+}
+
+FReply SObjectWidget::OnMouseButtonDoubleClick(const FGeometry& MyGeometry, const FPointerEvent& MouseEvent)
+{
+	return WidgetObject->OnMouseButtonDoubleClick(MyGeometry, MouseEvent).ToReply();
+}
+
+FReply SObjectWidget::OnDragDetected(const FGeometry& MyGeometry, const FPointerEvent& MouseEvent)
+{
+	// TODO UMG
+	return FReply::Unhandled();
+}
+
+void SObjectWidget::OnDragEnter(const FGeometry& MyGeometry, const FDragDropEvent& DragDropEvent)
+{
+	// TODO UMG
+}
+
+void SObjectWidget::OnDragLeave(const FDragDropEvent& DragDropEvent)
+{
+	// TODO UMG
+}
+
+FReply SObjectWidget::OnDragOver(const FGeometry& MyGeometry, const FDragDropEvent& DragDropEvent)
+{
+	// TODO UMG
+	return FReply::Unhandled();
+}
+
+FReply SObjectWidget::OnDrop(const FGeometry& MyGeometry, const FDragDropEvent& DragDropEvent)
+{
+	// TODO UMG
+	return FReply::Unhandled();
+}
+
+FReply SObjectWidget::OnControllerButtonPressed(const FGeometry& MyGeometry, const FControllerEvent& ControllerEvent)
+{
+	return WidgetObject->OnControllerButtonPressed(MyGeometry, ControllerEvent).ToReply();
+}
+
+FReply SObjectWidget::OnControllerButtonReleased(const FGeometry& MyGeometry, const FControllerEvent& ControllerEvent)
+{
+	return WidgetObject->OnControllerButtonReleased(MyGeometry, ControllerEvent).ToReply();
+}
+
+FReply SObjectWidget::OnControllerAnalogValueChanged(const FGeometry& MyGeometry, const FControllerEvent& ControllerEvent)
+{
+	return WidgetObject->OnControllerAnalogValueChanged(MyGeometry, ControllerEvent).ToReply();
+}
+
+FReply SObjectWidget::OnTouchGesture(const FGeometry& MyGeometry, const FPointerEvent& GestureEvent)
+{
+	return WidgetObject->OnTouchGesture(MyGeometry, GestureEvent).ToReply();
+}
+
+FReply SObjectWidget::OnTouchStarted(const FGeometry& MyGeometry, const FPointerEvent& InTouchEvent)
+{
+	return WidgetObject->OnTouchStarted(MyGeometry, InTouchEvent).ToReply();
+}
+
+FReply SObjectWidget::OnTouchMoved(const FGeometry& MyGeometry, const FPointerEvent& InTouchEvent)
+{
+	return WidgetObject->OnTouchMoved(MyGeometry, InTouchEvent).ToReply();
+}
+
+FReply SObjectWidget::OnTouchEnded(const FGeometry& MyGeometry, const FPointerEvent& InTouchEvent)
+{
+	return WidgetObject->OnTouchEnded(MyGeometry, InTouchEvent).ToReply();
+}
+
+FReply SObjectWidget::OnMotionDetected(const FGeometry& MyGeometry, const FMotionEvent& InMotionEvent)
+{
+	return WidgetObject->OnMotionDetected(MyGeometry, InMotionEvent).ToReply();
+}
diff --git a/Engine/Source/Runtime/UMG/Private/UserWidget.cpp b/Engine/Source/Runtime/UMG/Private/UserWidget.cpp
--- a/Engine/Source/Runtime/UMG/Private/UserWidget.cpp
+++ b/Engine/Source/Runtime/UMG/Private/UserWidget.cpp
@@ -160,41 +160,15 @@ USlateWrapperComponent* UUserWidget::GetWidgetHandle(TSharedRef<SWidget> InWidge
 void UUserWidget::RebuildWrapperWidget()
 {
 	WidgetToComponent.Reset();
-	RootWidget = SNew(SSpacer);
 	
 	// Add the first component to the root of the widget surface.
 	if ( Components.Num() > 0 )
 	{
-		if ( bAbsoluteLayout )
-		{
-			RootWidget =
-				SNew(SCanvas)
-
-				+ SCanvas::Slot()
-				.Position(AbsolutePosition)
-				.Size(AbsoluteSize)
-				.VAlign(VerticalAlignment)
-				.HAlign(HorizontalAlignment)
-				[
-					Components[0]->GetWidget()
-				];
-		}
-		else
-		{
-			TSharedRef<SVerticalBox> VerticalBox = SNew(SVerticalBox);
-
-			auto& NewSlot = VerticalBox->AddSlot()
-				.Padding(Padding)
-				.HAlign(HorizontalAlignment)
-				.VAlign(VerticalAlignment)
-				[
-					Components[0]->GetWidget()
-				];
-
-			NewSlot.SizeParam = USlateWrapperComponent::ConvertSerializedSizeParamToRuntime(Size);
-
-			RootWidget = VerticalBox;
-		}
+		UserRootWidget = Components[0]->GetWidget();
+	}
+	else
+	{
+		UserRootWidget = SNew(SSpacer);
 	}
 
 	// Place all of our top-level children Slate wrapped components into the overlay
@@ -210,16 +184,6 @@ void UUserWidget::RebuildWrapperWidget()
 	UWorld* World = GetWorld();
 	if ( World && World->IsGameWorld() )
 	{
-		TSharedRef<SViewportWidgetHost> WidgetHost = SNew(SViewportWidgetHost, (bool)bModal)
-		[
-			RootWidget.ToSharedRef()
-		];
-
-		RootWidget = WidgetHost;
-
-		UGameViewportClient* Viewport = World->GetGameViewport();
-		Viewport->AddViewportWidgetContent(RootWidget.ToSharedRef());
-
 		if ( Visiblity == ESlateVisibility::Visible )
 		{
 			Show();
@@ -257,9 +221,45 @@ USlateWrapperComponent* UUserWidget::GetHandleFromName(const FString& Name) cons
 	return NULL;
 }
 
-TSharedRef<SWidget> UUserWidget::GetRootWidget()
+TSharedRef<SWidget> UUserWidget::MakeWidget()
+{
+	return SNew(SObjectWidget, this)
+	[
+		UserRootWidget.ToSharedRef()
+	];
+}
+
+TSharedRef<SWidget> UUserWidget::MakeFullScreenWidget()
 {
-	return RootWidget.ToSharedRef();
+	if ( bAbsoluteLayout )
+	{
+		return SNew(SCanvas)
+
+			+ SCanvas::Slot()
+			.Position(AbsolutePosition)
+			.Size(AbsoluteSize)
+			.VAlign(VerticalAlignment)
+			.HAlign(HorizontalAlignment)
+			[
+				MakeWidget()
+			];
+	}
+	else
+	{
+		TSharedRef<SVerticalBox> VerticalBox = SNew(SVerticalBox);
+
+		auto& NewSlot = VerticalBox->AddSlot()
+			.Padding(Padding)
+			.HAlign(HorizontalAlignment)
+			.VAlign(VerticalAlignment)
+			[
+				MakeWidget()
+			];
+
+		NewSlot.SizeParam = USlateWrapperComponent::ConvertSerializedSizeParamToRuntime(Size);
+
+		return VerticalBox;
+	}
 }
 
 USlateWrapperComponent* UUserWidget::GetRootWidgetComponent()
@@ -274,50 +274,78 @@ USlateWrapperComponent* UUserWidget::GetRootWidgetComponent()
 
 void UUserWidget::Show()
 {
-	RootWidget->SetVisibility(EVisibility::Visible);
-	OnVisibilityChanged.Broadcast(ESlateVisibility::Visible);
-
-	// If this is a game world add the widget to the current worlds viewport.
-	UWorld* World = GetWorld();
-	if ( World && World->IsGameWorld() )
+	if ( !FullScreenWidget.IsValid() )
 	{
-		UGameViewportClient* Viewport = World->GetGameViewport();
-		TWeakPtr<SViewport> GameViewportWidget = Viewport->GetGameViewport()->GetViewportWidget();
-		if ( GameViewportWidget.IsValid() )
+		TSharedRef<SWidget> RootWidget = MakeFullScreenWidget();
+
+		TSharedRef<SViewportWidgetHost> WidgetHost = SNew(SViewportWidgetHost, (bool)bModal)
+			[
+				RootWidget
+			];
+
+		FullScreenWidget = WidgetHost;
+
+		//WidgetHost->SetVisibility(EVisibility::Visible);
+		//OnVisibilityChanged.Broadcast(ESlateVisibility::Visible);
+
+		// If this is a game world add the widget to the current worlds viewport.
+		UWorld* World = GetWorld();
+		if ( World && World->IsGameWorld() )
 		{
-			GameViewportWidget.Pin()->SetWidgetToFocusOnActivate(RootWidget);
-			FSlateApplication::Get().SetKeyboardFocus(RootWidget);
+			UGameViewportClient* Viewport = World->GetGameViewport();
+			Viewport->AddViewportWidgetContent(WidgetHost);
+
+			TWeakPtr<SViewport> GameViewportWidget = Viewport->GetGameViewport()->GetViewportWidget();
+			if ( GameViewportWidget.IsValid() )
+			{
+				GameViewportWidget.Pin()->SetWidgetToFocusOnActivate(RootWidget);
+				FSlateApplication::Get().SetKeyboardFocus(RootWidget);
+			}
 		}
 	}
 }
 
 void UUserWidget::Hide()
 {
-	RootWidget->SetVisibility(EVisibility::Hidden);
-	OnVisibilityChanged.Broadcast(ESlateVisibility::Hidden);
-
-	// If this is a game world add the widget to the current worlds viewport.
-	UWorld* World = GetWorld();
-	if ( World && World->IsGameWorld() )
+	if ( FullScreenWidget.IsValid() )
 	{
-		UGameViewportClient* Viewport = World->GetGameViewport();
-		TWeakPtr<SViewport> GameViewportWidget = Viewport->GetGameViewport()->GetViewportWidget();
-		if ( GameViewportWidget.IsValid() )
+		TSharedPtr<SWidget> RootWidget = FullScreenWidget.Pin();
+
+		//RootWidget->SetVisibility(EVisibility::Hidden);
+		//OnVisibilityChanged.Broadcast(ESlateVisibility::Hidden);
+
+		// If this is a game world add the widget to the current worlds viewport.
+		UWorld* World = GetWorld();
+		if ( World && World->IsGameWorld() )
 		{
-			//TODO UMG this isn't what should manage focus, a higher level window controller, probably the viewport needs to understand
-			// the Widget stack, and the dialog stack.
-			GameViewportWidget.Pin()->ClearWidgetToFocusOnActivate();
-			FSlateApplication::Get().SetKeyboardFocus(TSharedPtr<SWidget>());
+			UGameViewportClient* Viewport = World->GetGameViewport();
+			Viewport->RemoveViewportWidgetContent(RootWidget.ToSharedRef());
+
+			TWeakPtr<SViewport> GameViewportWidget = Viewport->GetGameViewport()->GetViewportWidget();
+			if ( GameViewportWidget.IsValid() )
+			{
+				//TODO UMG this isn't what should manage focus, a higher level window controller, probably the viewport needs to understand
+				// the Widget stack, and the dialog stack.
+				GameViewportWidget.Pin()->ClearWidgetToFocusOnActivate();
+				FSlateApplication::Get().SetKeyboardFocus(TSharedPtr<SWidget>());
+			}
 		}
 	}
 }
 
 bool UUserWidget::GetIsVisible()
 {
-	return RootWidget->GetVisibility().IsVisible();
+	return FullScreenWidget.IsValid();
 }
 
 TEnumAsByte<ESlateVisibility::Type> UUserWidget::GetVisiblity()
 {
-	return USlateWrapperComponent::ConvertRuntimeToSerializedVisiblity(RootWidget->GetVisibility());
+	if ( FullScreenWidget.IsValid() )
+	{
+		TSharedPtr<SWidget> RootWidget = FullScreenWidget.Pin();
+
+		return USlateWrapperComponent::ConvertRuntimeToSerializedVisiblity(RootWidget->GetVisibility());
+	}
+
+	return ESlateVisibility::Collapsed;
 }
diff --git a/Engine/Source/Runtime/UMG/Public/Blueprint/UserWidget.h b/Engine/Source/Runtime/UMG/Public/Blueprint/UserWidget.h
--- a/Engine/Source/Runtime/UMG/Public/Blueprint/UserWidget.h
+++ b/Engine/Source/Runtime/UMG/Public/Blueprint/UserWidget.h
@@ -2,6 +2,8 @@
 
 #pragma once
 
+#include "Geometry.h"
+
 #include "UserWidget.generated.h"
 
 //TODO UMG If you want to host a widget that's full screen there may need to be a SWindow equivalent that you spawn it into.
@@ -99,18 +101,87 @@ class UMG_API UUserWidget : public UObject
 	UFUNCTION(BlueprintPure, Category="Appearance")
 	TEnumAsByte<ESlateVisibility::Type> GetVisiblity();
 
+	UFUNCTION(BlueprintImplementableEvent, Category="User Interface")
+	FSReply OnKeyboardFocusReceived(FGeometry MyGeometry, FKeyboardFocusEvent InKeyboardFocusEvent);
+	UFUNCTION(BlueprintImplementableEvent, Category="User Interface")
+	void OnKeyboardFocusLost(FKeyboardFocusEvent InKeyboardFocusEvent);
+	//UFUNCTION(BlueprintImplementableEvent, Category="User Interface")
+	//void OnKeyboardFocusChanging(FWeakWidgetPath PreviousFocusPath, FWidgetPath NewWidgetPath);
+
+	UFUNCTION(BlueprintImplementableEvent, Category="User Interface")
+	FSReply OnKeyChar(FGeometry MyGeometry, FCharacterEvent InCharacterEvent);
+	UFUNCTION(BlueprintImplementableEvent, Category="User Interface")
+	FSReply OnPreviewKeyDown(FGeometry MyGeometry, FKeyboardEvent InKeyboardEvent);
+	UFUNCTION(BlueprintImplementableEvent, Category="User Interface")
+	FSReply OnKeyDown(FGeometry MyGeometry, FKeyboardEvent InKeyboardEvent);
+	UFUNCTION(BlueprintImplementableEvent, Category="User Interface")
+	FSReply OnKeyUp(FGeometry MyGeometry, FKeyboardEvent InKeyboardEvent);
+
+	UFUNCTION(BlueprintImplementableEvent, Category="User Interface")
+	FSReply OnMouseButtonDown(FGeometry MyGeometry, const FPointerEvent& MouseEvent);
+	UFUNCTION(BlueprintImplementableEvent, Category="User Interface")
+	FSReply OnPreviewMouseButtonDown(FGeometry MyGeometry, const FPointerEvent& MouseEvent);
+	UFUNCTION(BlueprintImplementableEvent, Category="User Interface")
+	FSReply OnMouseButtonUp(FGeometry MyGeometry, const FPointerEvent& MouseEvent);
+	UFUNCTION(BlueprintImplementableEvent, Category="User Interface")
+	FSReply OnMouseMove(FGeometry MyGeometry, const FPointerEvent& MouseEvent);
+	UFUNCTION(BlueprintImplementableEvent, Category="User Interface")
+	void OnMouseEnter(FGeometry MyGeometry, const FPointerEvent& MouseEvent);
+	UFUNCTION(BlueprintImplementableEvent, Category="User Interface")
+	void OnMouseLeave(const FPointerEvent& MouseEvent);
+	UFUNCTION(BlueprintImplementableEvent, Category="User Interface")
+	FSReply OnMouseWheel(FGeometry MyGeometry, const FPointerEvent& MouseEvent);
+	//UFUNCTION(BlueprintImplementableEvent, Category="User Interface")
+	//FCursorReply OnCursorQuery(FGeometry MyGeometry, const FPointerEvent& CursorEvent) const;
+	UFUNCTION(BlueprintImplementableEvent, Category="User Interface")
+	FSReply OnMouseButtonDoubleClick(FGeometry InMyGeometry, const FPointerEvent& InMouseEvent);
+
+	//UFUNCTION(BlueprintImplementableEvent, Category="User Interface")
+	//FSReply OnDragDetected(FGeometry MyGeometry, const FPointerEvent& MouseEvent);
+	//UFUNCTION(BlueprintImplementableEvent, Category="User Interface")
+	//void OnDragEnter(FGeometry MyGeometry, FDragDropEvent DragDropEvent);
+	//UFUNCTION(BlueprintImplementableEvent, Category="User Interface")
+	//void OnDragLeave(FDragDropEvent DragDropEvent);
+	//UFUNCTION(BlueprintImplementableEvent, Category="User Interface")
+	//FSReply OnDragOver(FGeometry MyGeometry, FDragDropEvent DragDropEvent);
+	//UFUNCTION(BlueprintImplementableEvent, Category="User Interface")
+	//FSReply OnDrop(FGeometry MyGeometry, FDragDropEvent DragDropEvent);
+
+	UFUNCTION(BlueprintImplementableEvent, Category="User Interface")
+	FSReply OnControllerButtonPressed(FGeometry MyGeometry, FControllerEvent ControllerEvent);
+	UFUNCTION(BlueprintImplementableEvent, Category="User Interface")
+	FSReply OnControllerButtonReleased(FGeometry MyGeometry, FControllerEvent ControllerEvent);
+	UFUNCTION(BlueprintImplementableEvent, Category="User Interface")
+	FSReply OnControllerAnalogValueChanged(FGeometry MyGeometry, FControllerEvent ControllerEvent);
+
+	UFUNCTION(BlueprintImplementableEvent, Category="User Interface")
+	FSReply OnTouchGesture(FGeometry MyGeometry, const FPointerEvent& GestureEvent);
+	UFUNCTION(BlueprintImplementableEvent, Category="User Interface")
+	FSReply OnTouchStarted(FGeometry MyGeometry, const FPointerEvent& InTouchEvent);
+	UFUNCTION(BlueprintImplementableEvent, Category="User Interface")
+	FSReply OnTouchMoved(FGeometry MyGeometry, const FPointerEvent& InTouchEvent);
+	UFUNCTION(BlueprintImplementableEvent, Category="User Interface")
+	FSReply OnTouchEnded(FGeometry MyGeometry, const FPointerEvent& InTouchEvent);
+	UFUNCTION(BlueprintImplementableEvent, Category="User Interface")
+	FSReply OnMotionDetected(FGeometry MyGeometry, FMotionEvent InMotionEvent);
+
+
 	USlateWrapperComponent* GetWidgetHandle(TSharedRef<SWidget> InWidget);
 
-	TSharedRef<SWidget> GetRootWidget();
+	TSharedRef<SWidget> MakeWidget();
+	TSharedRef<SWidget> MakeFullScreenWidget();
+
 	USlateWrapperComponent* GetRootWidgetComponent();
 
 	TSharedPtr<SWidget> GetWidgetFromName(const FString& Name) const;
 	USlateWrapperComponent* GetHandleFromName(const FString& Name) const;
 
 private:
-	TSharedPtr<SWidget> RootWidget;
+	TSharedPtr<SWidget> UserRootWidget;
 	TMap< TWeakPtr<SWidget>, USlateWrapperComponent* > WidgetToComponent;
 
+	TWeakPtr<SWidget> FullScreenWidget;
+
 private:
 	void RebuildWrapperWidget();
 };
diff --git a/Engine/Source/Runtime/UMG/Public/Components/WidgetHost.h b/Engine/Source/Runtime/UMG/Public/Components/WidgetHost.h
deleted file mode 100644
diff --git a/Engine/Source/Runtime/UMG/Public/Slate/SObjectWidget.h b/Engine/Source/Runtime/UMG/Public/Slate/SObjectWidget.h
new file mode 100644
--- /dev/null
+++ b/Engine/Source/Runtime/UMG/Public/Slate/SObjectWidget.h
@@ -0,0 +1,53 @@
+// Copyright 1998-2014 Epic Games, Inc. All Rights Reserved.
+
+#pragma once
+
+class SObjectWidget : public SCompoundWidget, public FGCObject
+{
+	SLATE_BEGIN_ARGS(SObjectWidget)
+	{ }
+		SLATE_DEFAULT_SLOT(FArguments, Content)
+	SLATE_END_ARGS()
+
+	void Construct(const FArguments& InArgs, class UUserWidget* InWidgetObject);
+
+	virtual void AddReferencedObjects(FReferenceCollector& Collector) OVERRIDE;
+
+	virtual FReply OnKeyboardFocusReceived(const FGeometry& MyGeometry, const FKeyboardFocusEvent& InKeyboardFocusEvent) OVERRIDE;
+	virtual void OnKeyboardFocusLost(const FKeyboardFocusEvent& InKeyboardFocusEvent) OVERRIDE;
+	virtual void OnKeyboardFocusChanging(const FWeakWidgetPath& PreviousFocusPath, const FWidgetPath& NewWidgetPath) OVERRIDE;
+	
+	virtual FReply OnKeyChar(const FGeometry& MyGeometry, const FCharacterEvent& InCharacterEvent) OVERRIDE;
+	virtual FReply OnPreviewKeyDown(const FGeometry& MyGeometry, const FKeyboardEvent& InKeyboardEvent) OVERRIDE;
+	virtual FReply OnKeyDown(const FGeometry& MyGeometry, const FKeyboardEvent& InKeyboardEvent) OVERRIDE;
+	virtual FReply OnKeyUp(const FGeometry& MyGeometry, const FKeyboardEvent& InKeyboardEvent) OVERRIDE;
+	
+	virtual FReply OnMouseButtonDown(const FGeometry& MyGeometry, const FPointerEvent& MouseEvent) OVERRIDE;
+	virtual FReply OnPreviewMouseButtonDown(const FGeometry& MyGeometry, const FPointerEvent& MouseEvent) OVERRIDE;
+	virtual FReply OnMouseButtonUp(const FGeometry& MyGeometry, const FPointerEvent& MouseEvent) OVERRIDE;
+	virtual FReply OnMouseMove(const FGeometry& MyGeometry, const FPointerEvent& MouseEvent) OVERRIDE;
+	virtual void OnMouseEnter(const FGeometry& MyGeometry, const FPointerEvent& MouseEvent) OVERRIDE;
+	virtual void OnMouseLeave(const FPointerEvent& MouseEvent) OVERRIDE;
+	virtual FReply OnMouseWheel(const FGeometry& MyGeometry, const FPointerEvent& MouseEvent) OVERRIDE;
+	virtual FCursorReply OnCursorQuery(const FGeometry& MyGeometry, const FPointerEvent& CursorEvent) const OVERRIDE;
+	virtual FReply OnMouseButtonDoubleClick(const FGeometry& InMyGeometry, const FPointerEvent& InMouseEvent) OVERRIDE;
+	
+	virtual FReply OnDragDetected(const FGeometry& MyGeometry, const FPointerEvent& MouseEvent) OVERRIDE;
+	virtual void OnDragEnter(const FGeometry& MyGeometry, const FDragDropEvent& DragDropEvent) OVERRIDE;
+	virtual void OnDragLeave(const FDragDropEvent& DragDropEvent) OVERRIDE;
+	virtual FReply OnDragOver(const FGeometry& MyGeometry, const FDragDropEvent& DragDropEvent) OVERRIDE;
+	virtual FReply OnDrop(const FGeometry& MyGeometry, const FDragDropEvent& DragDropEvent) OVERRIDE;
+	
+	virtual FReply OnControllerButtonPressed(const FGeometry& MyGeometry, const FControllerEvent& ControllerEvent) OVERRIDE;
+	virtual FReply OnControllerButtonReleased(const FGeometry& MyGeometry, const FControllerEvent& ControllerEvent) OVERRIDE;
+	virtual FReply OnControllerAnalogValueChanged(const FGeometry& MyGeometry, const FControllerEvent& ControllerEvent) OVERRIDE;
+	
+	virtual FReply OnTouchGesture(const FGeometry& MyGeometry, const FPointerEvent& GestureEvent) OVERRIDE;
+	virtual FReply OnTouchStarted(const FGeometry& MyGeometry, const FPointerEvent& InTouchEvent) OVERRIDE;
+	virtual FReply OnTouchMoved(const FGeometry& MyGeometry, const FPointerEvent& InTouchEvent) OVERRIDE;
+	virtual FReply OnTouchEnded(const FGeometry& MyGeometry, const FPointerEvent& InTouchEvent) OVERRIDE;
+	virtual FReply OnMotionDetected(const FGeometry& MyGeometry, const FMotionEvent& InMotionEvent) OVERRIDE;
+
+private:
+	class UUserWidget* WidgetObject;
+};
\ No newline at end of file
diff --git a/Engine/Source/Runtime/UMG/Public/UMG.h b/Engine/Source/Runtime/UMG/Public/UMG.h
--- a/Engine/Source/Runtime/UMG/Public/UMG.h
+++ b/Engine/Source/Runtime/UMG/Public/UMG.h
@@ -32,8 +32,6 @@
 
 #include "ContentWidget.h"
 
-#include "WidgetHost.h"
-
 #include "BoxPanelComponent.h"
 
 #include "CanvasPanelSlot.h"
@@ -64,6 +62,10 @@
 #include "ThrobberComponent.h"
 #include "CircularThrobberComponent.h"
 
+/* Slate
+ *****************************************************************************/
+ 
+ #include "SObjectWidget.h"
 
 /* Blueprint
  *****************************************************************************/
