UMG - More renaming changes to baseclasses, and fixing build.

diff --git a/Engine/Source/Runtime/UMG/Private/Slate/SObjectWidget.cpp b/Engine/Source/Runtime/UMG/Private/Slate/SObjectWidget.cpp
--- a/Engine/Source/Runtime/UMG/Private/Slate/SObjectWidget.cpp
+++ b/Engine/Source/Runtime/UMG/Private/Slate/SObjectWidget.cpp
@@ -17,6 +17,11 @@ void SObjectWidget::AddReferencedObjects(FReferenceCollector& Collector)
 	Collector.AddReferencedObject(WidgetObject);
 }
 
+//int32 SObjectWidget::OnPaint(const FGeometry& AllottedGeometry, const FSlateRect& MyClippingRect, FSlateWindowElementList& OutDrawElements, int32 LayerId, const FWidgetStyle& InWidgetStyle, bool bParentEnabled) const
+//{
+//
+//}
+
 FReply SObjectWidget::OnKeyboardFocusReceived(const FGeometry& MyGeometry, const FKeyboardFocusEvent& InKeyboardFocusEvent)
 {
 	return WidgetObject->OnKeyboardFocusReceived(MyGeometry, InKeyboardFocusEvent).ToReply();
diff --git a/Engine/Source/Runtime/UMG/Public/Components/CircularThrobberComponent.h b/Engine/Source/Runtime/UMG/Public/Components/CircularThrobberComponent.h
--- a/Engine/Source/Runtime/UMG/Public/Components/CircularThrobberComponent.h
+++ b/Engine/Source/Runtime/UMG/Public/Components/CircularThrobberComponent.h
@@ -6,7 +6,7 @@
 
 /**  A throbber widget that orients images in a spinning circle. */
 UCLASS(meta=(BlueprintSpawnableComponent), ClassGroup=UserInterface)
-class UMG_API UCircularThrobberComponent : public USlateLeafWidgetComponent
+class UMG_API UCircularThrobberComponent : public UWidget
 {
 	GENERATED_UCLASS_BODY()
 
diff --git a/Engine/Source/Runtime/UMG/Public/Components/ImageComponent.h b/Engine/Source/Runtime/UMG/Public/Components/ImageComponent.h
--- a/Engine/Source/Runtime/UMG/Public/Components/ImageComponent.h
+++ b/Engine/Source/Runtime/UMG/Public/Components/ImageComponent.h
@@ -6,7 +6,7 @@
 
 /** Image widget */
 UCLASS(meta=(BlueprintSpawnableComponent), ClassGroup=UserInterface)
-class UMG_API UImageComponent : public USlateLeafWidgetComponent
+class UMG_API UImageComponent : public UWidget
 {
 	GENERATED_UCLASS_BODY()
 
diff --git a/Engine/Source/Runtime/UMG/Public/Components/PanelWidget.h b/Engine/Source/Runtime/UMG/Public/Components/PanelWidget.h
new file mode 100644
--- /dev/null
+++ b/Engine/Source/Runtime/UMG/Public/Components/PanelWidget.h
@@ -0,0 +1,28 @@
+// Copyright 1998-2014 Epic Games, Inc. All Rights Reserved.
+
+#pragma once
+
+#include "PanelWidget.generated.h"
+
+UCLASS(Abstract)
+class UMG_API UPanelWidget : public UWidget
+{
+	GENERATED_UCLASS_BODY()
+
+	virtual int32 GetChildrenCount() const { return 0; }
+	virtual UWidget* GetChildAt(int32 Index) const { return NULL; }
+
+	virtual void ReplaceChildAt(int32 Index, UWidget* Child)
+	{
+	}
+
+	virtual bool AddChild(UWidget* Child, FVector2D Position)
+	{
+		return false;
+	}
+
+	virtual bool RemoveChild(UWidget* Child)
+	{
+		return false;
+	}
+};
diff --git a/Engine/Source/Runtime/UMG/Public/Components/SlateNonLeafWidgetComponent.h b/Engine/Source/Runtime/UMG/Public/Components/SlateNonLeafWidgetComponent.h
deleted file mode 100644
diff --git a/Engine/Source/Runtime/UMG/Public/Components/SlateWrapperComponent.h b/Engine/Source/Runtime/UMG/Public/Components/SlateWrapperComponent.h
deleted file mode 100644
diff --git a/Engine/Source/Runtime/UMG/Public/Components/Slider.h b/Engine/Source/Runtime/UMG/Public/Components/Slider.h
--- a/Engine/Source/Runtime/UMG/Public/Components/Slider.h
+++ b/Engine/Source/Runtime/UMG/Public/Components/Slider.h
@@ -10,7 +10,7 @@ DECLARE_DYNAMIC_MULTICAST_DELEGATE_OneParam(FOnFloatValueChangedEvent, float, Va
 
 /** Slider widget */
 UCLASS(meta=(BlueprintSpawnableComponent), ClassGroup=UserInterface)
-class UMG_API USlider : public USlateLeafWidgetComponent
+class UMG_API USlider : public UWidget
 {
 	GENERATED_UCLASS_BODY()
 
diff --git a/Engine/Source/Runtime/UMG/Public/Components/SpacerComponent.h b/Engine/Source/Runtime/UMG/Public/Components/SpacerComponent.h
--- a/Engine/Source/Runtime/UMG/Public/Components/SpacerComponent.h
+++ b/Engine/Source/Runtime/UMG/Public/Components/SpacerComponent.h
@@ -6,7 +6,7 @@
 
 /** A spacer widget; it does not have a visual representation, and just provides padding between other widgets. */
 UCLASS(meta=(BlueprintSpawnableComponent), ClassGroup=UserInterface)
-class UMG_API USpacerComponent : public USlateLeafWidgetComponent
+class UMG_API USpacerComponent : public UWidget
 {
 	GENERATED_UCLASS_BODY()
 
diff --git a/Engine/Source/Runtime/UMG/Public/Components/TextBlockComponent.h b/Engine/Source/Runtime/UMG/Public/Components/TextBlockComponent.h
--- a/Engine/Source/Runtime/UMG/Public/Components/TextBlockComponent.h
+++ b/Engine/Source/Runtime/UMG/Public/Components/TextBlockComponent.h
@@ -5,7 +5,7 @@
 #include "TextBlockComponent.generated.h"
 
 UCLASS(meta=(BlueprintSpawnableComponent), ClassGroup=UserInterface)
-class UMG_API UTextBlockComponent : public USlateLeafWidgetComponent
+class UMG_API UTextBlockComponent : public UWidget
 {
 	GENERATED_UCLASS_BODY()
 
diff --git a/Engine/Source/Runtime/UMG/Public/Components/ThrobberComponent.h b/Engine/Source/Runtime/UMG/Public/Components/ThrobberComponent.h
--- a/Engine/Source/Runtime/UMG/Public/Components/ThrobberComponent.h
+++ b/Engine/Source/Runtime/UMG/Public/Components/ThrobberComponent.h
@@ -6,7 +6,7 @@
 
 /** A Throbber widget that shows several zooming circles in a row. */
 UCLASS(meta=(BlueprintSpawnableComponent), ClassGroup=UserInterface)
-class UMG_API UThrobberComponent : public USlateLeafWidgetComponent
+class UMG_API UThrobberComponent : public UWidget
 {
 	GENERATED_UCLASS_BODY()
 
diff --git a/Engine/Source/Runtime/UMG/Public/Components/Widget.h b/Engine/Source/Runtime/UMG/Public/Components/Widget.h
new file mode 100644
--- /dev/null
+++ b/Engine/Source/Runtime/UMG/Public/Components/Widget.h
@@ -0,0 +1,117 @@
+// Copyright 1998-2014 Epic Games, Inc. All Rights Reserved.
+
+#pragma once
+
+#include "SlateWrapperTypes.h"
+#include "Widget.generated.h"
+
+#define OPTIONAL_BINDING(ReturnType, MemberName) MemberName ## Delegate.IsBound() ? TAttribute< ReturnType >::Create(MemberName ## Delegate.GetUObject(), MemberName ## Delegate.GetFunctionName()) : TAttribute< ReturnType >(MemberName)
+#define OPTIONAL_BINDING_CONVERT(ReturnType, MemberName, ConvertedType, ConversionFunction) MemberName ## Delegate.IsBound() ? TAttribute< ConvertedType >::Create(TAttribute< ConvertedType >::FGetter::CreateUObject(this, &ThisClass::ConversionFunction, TAttribute< ReturnType >::Create(MemberName ## Delegate.GetUObject(), MemberName ## Delegate.GetFunctionName()))) : ConversionFunction(TAttribute< ReturnType >(MemberName))
+
+/** This is the base class for all Slate widget wrapper components */
+UCLASS(Abstract, hideCategories=(Activation, "Components|Activation", ComponentReplication, LOD, Rendering))
+class UMG_API UWidget : public UObject
+{
+	GENERATED_UCLASS_BODY()
+
+	// Common Bindings
+	DECLARE_DYNAMIC_DELEGATE_RetVal(bool, FGetBool);
+	DECLARE_DYNAMIC_DELEGATE_RetVal(FText, FGetText);
+	DECLARE_DYNAMIC_DELEGATE_RetVal(FLinearColor, FGetSlateColor);
+	DECLARE_DYNAMIC_DELEGATE_RetVal(FLinearColor, FGetLinearColor);
+	DECLARE_DYNAMIC_DELEGATE_RetVal(ESlateVisibility::Type, FGetSlateVisibility);
+
+	// Events
+	DECLARE_DYNAMIC_DELEGATE_RetVal(FSReply, FOnReply);
+
+	UPROPERTY(EditAnywhere, Category=Variable)
+	FString Name;
+
+	UPROPERTY(EditAnywhere, Category=Variable)
+	bool bIsVariable;
+
+#if WITH_EDITOR
+	UPROPERTY(Transient, EditInline, EditAnywhere, BlueprintReadWrite, Category=Layout, meta=( ShowOnlyInnerProperties ))
+	UPanelSlot* Slot;
+#endif
+
+	/** Sets whether this widget can be modified interactively by the user */
+	UPROPERTY(EditDefaultsOnly, Category=Behavior)
+	bool bIsEnabled;
+
+	UPROPERTY()
+	FGetBool bIsEnabledDelegate;
+
+	//TODO UMG ToolTipWidget
+
+	UPROPERTY(EditDefaultsOnly, Category=Behavior)
+	FText ToolTipText;
+
+	UPROPERTY()
+	FGetText ToolTipTextDelegate;
+
+	UPROPERTY(EditDefaultsOnly, Category=Behavior)
+	TEnumAsByte<ESlateVisibility::Type> Visiblity;
+
+	UPROPERTY()
+	FGetSlateVisibility VisiblityDelegate;
+
+	UPROPERTY(EditDefaultsOnly, Category=Behavior)
+	TEnumAsByte<EMouseCursor::Type> Cursor;
+
+	UPROPERTY(Transient)
+	bool bCreatedByConstructionScript;
+
+	UFUNCTION(BlueprintCallable, Category="Widget")
+	bool GetIsEnabled() const;
+
+	UFUNCTION(BlueprintCallable, Category="Widget")
+	void SetIsEnabled(bool bInIsEnabled);
+
+	UFUNCTION(BlueprintCallable, Category="Widget")
+	void SetToolTipText(const FText& InToolTipText);
+
+	UFUNCTION(BlueprintCallable, Category="Widget")
+	TEnumAsByte<ESlateVisibility::Type> GetVisibility();
+
+	UFUNCTION(BlueprintCallable, Category="Widget")
+	void SetVisibility(TEnumAsByte<ESlateVisibility::Type> InVisibility);
+
+	//TOptional< TEnumAsByte<EMouseCursor::Type> > GetCursor() const;
+
+	//void SetCursor(const TAttribute< TOptional< TEnumAsByte<EMouseCursor::Type> > >& InCursor);
+
+	// UActorComponent interface
+	virtual void OnUnregister();
+	// End of UActorComponent interface
+
+	TSharedRef<SWidget> GetWidget();
+
+#if WITH_EDITOR
+	virtual void ConnectEditorData() { }
+#endif
+
+#if WITH_EDITOR
+	// UObject interface
+	virtual void PostEditChangeProperty(struct FPropertyChangedEvent& PropertyChangedEvent) OVERRIDE;
+	// End of UObject interface
+#endif
+
+	// Utility methods
+	//@TODO UMG: Should move elsewhere
+	static EVisibility ConvertSerializedVisibilityToRuntime(TEnumAsByte<ESlateVisibility::Type> Input);
+	static TEnumAsByte<ESlateVisibility::Type> ConvertRuntimeToSerializedVisiblity(const EVisibility& Input);
+
+	static FSizeParam ConvertSerializedSizeParamToRuntime(const FSlateChildSize& Input);
+
+protected:
+	virtual TSharedRef<SWidget> RebuildWidget();
+
+	EVisibility ConvertVisibility(TAttribute<ESlateVisibility::Type> SerializedType) const
+	{
+		return ConvertSerializedVisibilityToRuntime(SerializedType.Get());
+	}
+
+protected:
+	TSharedPtr<SWidget> MyWidget;
+};
diff --git a/Engine/Source/Runtime/UMG/Public/UMG.h b/Engine/Source/Runtime/UMG/Public/UMG.h
--- a/Engine/Source/Runtime/UMG/Public/UMG.h
+++ b/Engine/Source/Runtime/UMG/Public/UMG.h
@@ -25,11 +25,9 @@
 
 #include "PanelSlot.h"
 
-#include "SlateWrapperComponent.h"
-
-#include "SlateLeafWidgetComponent.h"
-#include "SlateNonLeafWidgetComponent.h"
+#include "Widget.h"
 
+#include "PanelWidget.h"
 #include "ContentWidget.h"
 
 #include "BoxPanelComponent.h"
