UMG - Adding the overlay panel.

diff --git a/Engine/Source/Runtime/UMG/Private/Components/Overlay.cpp b/Engine/Source/Runtime/UMG/Private/Components/Overlay.cpp
new file mode 100644
--- /dev/null
+++ b/Engine/Source/Runtime/UMG/Private/Components/Overlay.cpp
@@ -0,0 +1,128 @@
+// Copyright 1998-2014 Epic Games, Inc. All Rights Reserved.
+
+#include "UMGPrivatePCH.h"
+
+/////////////////////////////////////////////////////
+// UOverlay
+
+UOverlay::UOverlay(const FPostConstructInitializeProperties& PCIP)
+	: Super(PCIP)
+{
+	bIsVariable = false;
+}
+
+int32 UOverlay::GetChildrenCount() const
+{
+	return Slots.Num();
+}
+
+UWidget* UOverlay::GetChildAt(int32 Index) const
+{
+	return Slots[Index]->Content;
+}
+
+bool UOverlay::AddChild(UWidget* Child, FVector2D Position)
+{
+	AddSlot(Child);
+	return true;
+}
+
+bool UOverlay::RemoveChild(UWidget* Child)
+{
+	for ( int32 SlotIndex = 0; SlotIndex < Slots.Num(); ++SlotIndex )
+	{
+		UOverlaySlot* Slot = Slots[SlotIndex];
+
+		if ( Slot->Content == Child )
+		{
+			Slots.RemoveAt(SlotIndex);
+			return true;
+		}
+	}
+
+	return false;
+}
+
+void UOverlay::ReplaceChildAt(int32 Index, UWidget* Content)
+{
+	UOverlaySlot* Slot = Slots[Index];
+	Slot->Content = Content;
+
+#if WITH_EDITOR
+	Content->Slot = Slot;
+#endif
+}
+
+TSharedRef<SWidget> UOverlay::RebuildWidget()
+{
+	TSharedRef<SOverlay> NewCanvas = SNew(SOverlay);
+	MyVerticalBox = NewCanvas;
+
+	TSharedPtr<SOverlay> Canvas = MyVerticalBox.Pin();
+	if ( Canvas.IsValid() )
+	{
+		Canvas->ClearChildren();
+
+		for ( int32 SlotIndex = 0; SlotIndex < Slots.Num(); ++SlotIndex )
+		{
+			UOverlaySlot* Slot = Slots[SlotIndex];
+			if ( Slot == NULL )
+			{
+				Slots[SlotIndex] = Slot = ConstructObject<UOverlaySlot>(UOverlaySlot::StaticClass(), this);
+			}
+
+			auto& NewSlot = Canvas->AddSlot()
+				.Padding(Slot->Padding)
+				.HAlign(Slot->HorizontalAlignment)
+				.VAlign(Slot->VerticalAlignment)
+				[
+					Slot->Content == NULL ? SNullWidget::NullWidget : Slot->Content->GetWidget()
+				];
+		}
+	}
+
+	return NewCanvas;
+}
+
+UOverlaySlot* UOverlay::AddSlot(UWidget* Content)
+{
+	UOverlaySlot* Slot = ConstructObject<UOverlaySlot>(UOverlaySlot::StaticClass(), this);
+	Slot->Content = Content;
+
+#if WITH_EDITOR
+	Content->Slot = Slot;
+#endif
+	
+	Slots.Add(Slot);
+
+	return Slot;
+}
+
+#if WITH_EDITOR
+void UOverlay::ConnectEditorData()
+{
+	for ( UOverlaySlot* Slot : Slots )
+	{
+		if ( Slot->Content )
+		{
+			Slot->Content->Slot = Slot;
+		}
+		//TODO UMG Should we auto delete empty slots?
+	}
+}
+
+void UOverlay::PostEditChangeProperty(struct FPropertyChangedEvent& PropertyChangedEvent)
+{
+	// Ensure the slots have unique names
+	int32 SlotNumbering = 1;
+
+	TSet<FName> UniqueSlotNames;
+	for (int32 SlotIndex = 0; SlotIndex < Slots.Num(); ++SlotIndex)
+	{
+		if ( Slots[SlotIndex] == NULL )
+		{
+			Slots[SlotIndex] = ConstructObject<UOverlaySlot>(UOverlaySlot::StaticClass(), this);
+		}
+	}
+}
+#endif
diff --git a/Engine/Source/Runtime/UMG/Private/Components/OverlaySlot.cpp b/Engine/Source/Runtime/UMG/Private/Components/OverlaySlot.cpp
new file mode 100644
--- /dev/null
+++ b/Engine/Source/Runtime/UMG/Private/Components/OverlaySlot.cpp
@@ -0,0 +1,13 @@
+// Copyright 1998-2014 Epic Games, Inc. All Rights Reserved.
+
+#include "UMGPrivatePCH.h"
+
+/////////////////////////////////////////////////////
+// UOverlaySlot
+
+UOverlaySlot::UOverlaySlot(const FPostConstructInitializeProperties& PCIP)
+	: Super(PCIP)
+{
+	HorizontalAlignment = HAlign_Left;
+	VerticalAlignment = VAlign_Top;
+}
diff --git a/Engine/Source/Runtime/UMG/Public/Components/Overlay.h b/Engine/Source/Runtime/UMG/Public/Components/Overlay.h
new file mode 100644
--- /dev/null
+++ b/Engine/Source/Runtime/UMG/Public/Components/Overlay.h
@@ -0,0 +1,44 @@
+// Copyright 1998-2014 Epic Games, Inc. All Rights Reserved.
+
+#pragma once
+
+#include "Overlay.generated.h"
+
+UCLASS(meta=(BlueprintSpawnableComponent), ClassGroup=UserInterface)
+class UMG_API UOverlay : public UPanelWidget
+{
+	GENERATED_UCLASS_BODY()
+
+	/** The items placed on the canvas */
+	UPROPERTY(EditAnywhere, EditInline, Category=Slots)
+	TArray<UOverlaySlot*> Slots;
+
+	UOverlaySlot* AddSlot(UWidget* Content);
+
+	// UPanelWidget
+	virtual int32 GetChildrenCount() const OVERRIDE;
+	virtual UWidget* GetChildAt(int32 Index) const OVERRIDE;
+	virtual bool AddChild(UWidget* Child, FVector2D Position) OVERRIDE;
+	virtual bool RemoveChild(UWidget* Child) OVERRIDE;
+	virtual void ReplaceChildAt(int32 Index, UWidget* Child) OVERRIDE;
+	// End UPanelWidget
+
+#if WITH_EDITOR
+	// UObject interface
+	virtual void PostEditChangeProperty(struct FPropertyChangedEvent& PropertyChangedEvent) OVERRIDE;
+	// End of UObject interface
+
+	// UWidget interface
+	virtual void ConnectEditorData() OVERRIDE;
+	// End UWidget interface
+#endif
+
+protected:
+
+	TWeakPtr<class SOverlay> MyVerticalBox;
+
+protected:
+	// UWidget interface
+	virtual TSharedRef<SWidget> RebuildWidget() OVERRIDE;
+	// End of UWidget interface
+};
diff --git a/Engine/Source/Runtime/UMG/Public/Components/OverlaySlot.h b/Engine/Source/Runtime/UMG/Public/Components/OverlaySlot.h
new file mode 100644
--- /dev/null
+++ b/Engine/Source/Runtime/UMG/Public/Components/OverlaySlot.h
@@ -0,0 +1,22 @@
+// Copyright 1998-2014 Epic Games, Inc. All Rights Reserved.
+
+#pragma once
+
+#include "SlateWrapperTypes.h"
+
+#include "OverlaySlot.generated.h"
+
+UCLASS()
+class UMG_API UOverlaySlot : public UPanelSlot
+{
+	GENERATED_UCLASS_BODY()
+	
+	UPROPERTY(EditAnywhere, Category=Appearance)
+	FMargin Padding;
+
+	UPROPERTY(EditAnywhere, Category=Appearance)
+	TEnumAsByte<EHorizontalAlignment> HorizontalAlignment;
+
+	UPROPERTY(EditAnywhere, Category=Appearance)
+	TEnumAsByte<EVerticalAlignment> VerticalAlignment;
+};
diff --git a/Engine/Source/Runtime/UMG/Public/UMG.h b/Engine/Source/Runtime/UMG/Public/UMG.h
--- a/Engine/Source/Runtime/UMG/Public/UMG.h
+++ b/Engine/Source/Runtime/UMG/Public/UMG.h
@@ -44,6 +44,9 @@
 #include "UniformGridSlot.h"
 #include "UniformGridPanel.h"
 
+#include "OverlaySlot.h"
+#include "Overlay.h"
+
 #include "TextBlockComponent.h"
 #include "BorderComponent.h"
 #include "ButtonComponent.h"
