UMG - Prototyping a UMG control that could support the list view.  Binding the generated widget to the data source is the most painful part.

diff --git a/Engine/Source/Runtime/UMG/Private/Components/ListView.cpp b/Engine/Source/Runtime/UMG/Private/Components/ListView.cpp
new file mode 100644
--- /dev/null
+++ b/Engine/Source/Runtime/UMG/Private/Components/ListView.cpp
@@ -0,0 +1,61 @@
+// Copyright 1998-2014 Epic Games, Inc. All Rights Reserved.
+
+#include "UMGPrivatePCH.h"
+
+#define LOCTEXT_NAMESPACE "UMG"
+
+/////////////////////////////////////////////////////
+// UListView
+
+UListView::UListView(const FPostConstructInitializeProperties& PCIP)
+	: Super(PCIP)
+{
+	bIsVariable = true;
+
+	ItemHeight = 20.0f;
+	SelectionMode = ESelectionMode::Single;
+}
+
+TSharedRef<SWidget> UListView::RebuildWidget()
+{
+	return SNew(SListView< UObject* >)
+		.SelectionMode(SelectionMode)
+		.ListItemsSource(&Items)
+		.ItemHeight(ItemHeight)
+		.OnGenerateRow(BIND_UOBJECT_DELEGATE(SListView< UObject* >::FOnGenerateRow, HandleOnGenerateRow))
+		//.OnSelectionChanged(this, &SSocketManager::SocketSelectionChanged_Execute)
+		//.OnContextMenuOpening(this, &SSocketManager::OnContextMenuOpening)
+		//.OnItemScrolledIntoView(this, &SSocketManager::OnItemScrolledIntoView)
+		//	.HeaderRow
+		//	(
+		//		SNew(SHeaderRow)
+		//		.Visibility(EVisibility::Collapsed)
+		//		+ SHeaderRow::Column(TEXT("Socket"))
+		//	);
+		;
+}
+
+TSharedRef<ITableRow> UListView::HandleOnGenerateRow(UObject* Item, const TSharedRef< STableViewBase >& OwnerTable) const
+{
+	if ( OnGenerateRowEvent.IsBound() )
+	{
+		UWidget* Widget = OnGenerateRowEvent.Execute(Item);
+		if ( Widget != NULL )
+		{
+			return SNew(STableRow< UObject* >, OwnerTable)
+			[
+				Widget->GetWidget()
+			];
+		}
+	}
+
+	// If a row wasn't generated just create the default one, a simple text block of the item's name.
+	return SNew(STableRow< UObject* >, OwnerTable)
+		[
+			SNew(STextBlock).Text(Item ? FText::FromString(Item->GetName()) : LOCTEXT("null", "null"))
+		];
+}
+
+/////////////////////////////////////////////////////
+
+#undef LOCTEXT_NAMESPACE
diff --git a/Engine/Source/Runtime/UMG/Public/Components/ListView.h b/Engine/Source/Runtime/UMG/Public/Components/ListView.h
new file mode 100644
--- /dev/null
+++ b/Engine/Source/Runtime/UMG/Public/Components/ListView.h
@@ -0,0 +1,45 @@
+// Copyright 1998-2014 Epic Games, Inc. All Rights Reserved.
+
+#pragma once
+
+#include "ListView.generated.h"
+
+UCLASS(meta=(BlueprintSpawnableComponent), ClassGroup=UserInterface)
+class UMG_API UListView : public UWidget
+{
+	GENERATED_UCLASS_BODY()
+
+	DECLARE_DYNAMIC_DELEGATE_RetVal_OneParam(UWidget*, FOnGenerateRowUObject, UObject*, Item);
+		
+public:
+
+	UPROPERTY(EditDefaultsOnly, Category=Content)
+	float ItemHeight;
+
+	UPROPERTY(EditAnywhere, BlueprintReadWrite, Category=Content)
+	TArray<UObject*> Items;
+
+	UPROPERTY(EditDefaultsOnly, Category=Content)
+	TEnumAsByte<ESelectionMode::Type> SelectionMode;
+
+	/** Called when the button is clicked */
+	UPROPERTY(EditDefaultsOnly, Category=Events)
+	FOnGenerateRowUObject OnGenerateRowEvent;
+
+protected:
+	TSharedPtr< SListView<UObject*> > ListWidget() const
+	{
+		if ( MyWidget.IsValid() )
+		{
+			return StaticCastSharedRef< SListView<UObject*> >(MyWidget.ToSharedRef());
+		}
+
+		return TSharedPtr< SListView<UObject*> >();
+	}
+
+	TSharedRef<ITableRow> HandleOnGenerateRow(UObject* Item, const TSharedRef< STableViewBase >& OwnerTable) const;
+
+	// UWidget interface
+	virtual TSharedRef<SWidget> RebuildWidget() OVERRIDE;
+	// End of UWidget interface
+};
diff --git a/Engine/Source/Runtime/UMG/Public/UMG.h b/Engine/Source/Runtime/UMG/Public/UMG.h
--- a/Engine/Source/Runtime/UMG/Public/UMG.h
+++ b/Engine/Source/Runtime/UMG/Public/UMG.h
@@ -63,6 +63,8 @@
 #include "ThrobberComponent.h"
 #include "CircularThrobberComponent.h"
 
+#include "ListView.h"
+
 /* Slate
  *****************************************************************************/
  
