UMG - Experimenting with a way to expose the onpaint override of widgets to blueprints.  Adding TileView, and a baseclass UObject for tables.

diff --git a/Engine/Source/Runtime/UMG/Private/Components/ListView.cpp b/Engine/Source/Runtime/UMG/Private/Components/ListView.cpp
--- a/Engine/Source/Runtime/UMG/Private/Components/ListView.cpp
+++ b/Engine/Source/Runtime/UMG/Private/Components/ListView.cpp
@@ -12,7 +12,7 @@ UListView::UListView(const FPostConstructInitializeProperties& PCIP)
 {
 	bIsVariable = true;
 
-	ItemHeight = 20.0f;
+	ItemHeight = 16.0f;
 	SelectionMode = ESelectionMode::Single;
 }
 
@@ -37,6 +37,7 @@ TSharedRef<SWidget> UListView::RebuildWidget()
 
 TSharedRef<ITableRow> UListView::HandleOnGenerateRow(UObject* Item, const TSharedRef< STableViewBase >& OwnerTable) const
 {
+	// Call the user's delegate to see if they want to generate a custom widget bound to the data source.
 	if ( OnGenerateRowEvent.IsBound() )
 	{
 		UWidget* Widget = OnGenerateRowEvent.Execute(Item);
diff --git a/Engine/Source/Runtime/UMG/Private/Components/TableViewBase.cpp b/Engine/Source/Runtime/UMG/Private/Components/TableViewBase.cpp
new file mode 100644
--- /dev/null
+++ b/Engine/Source/Runtime/UMG/Private/Components/TableViewBase.cpp
@@ -0,0 +1,18 @@
+// Copyright 1998-2014 Epic Games, Inc. All Rights Reserved.
+
+#include "UMGPrivatePCH.h"
+
+#define LOCTEXT_NAMESPACE "UMG"
+
+/////////////////////////////////////////////////////
+// UTableViewBase
+
+UTableViewBase::UTableViewBase(const FPostConstructInitializeProperties& PCIP)
+	: Super(PCIP)
+{
+	bIsVariable = true;
+}
+
+/////////////////////////////////////////////////////
+
+#undef LOCTEXT_NAMESPACE
diff --git a/Engine/Source/Runtime/UMG/Private/Components/TileView.cpp b/Engine/Source/Runtime/UMG/Private/Components/TileView.cpp
new file mode 100644
--- /dev/null
+++ b/Engine/Source/Runtime/UMG/Private/Components/TileView.cpp
@@ -0,0 +1,64 @@
+// Copyright 1998-2014 Epic Games, Inc. All Rights Reserved.
+
+#include "UMGPrivatePCH.h"
+
+#define LOCTEXT_NAMESPACE "UMG"
+
+/////////////////////////////////////////////////////
+// UTileView
+
+UTileView::UTileView(const FPostConstructInitializeProperties& PCIP)
+	: Super(PCIP)
+{
+	bIsVariable = true;
+
+	ItemWidth = 128.0f;
+	ItemHeight = 128.0f;
+	SelectionMode = ESelectionMode::Single;
+}
+
+TSharedRef<SWidget> UTileView::RebuildWidget()
+{
+	return SNew(STileView< UObject* >)
+		.SelectionMode(SelectionMode)
+		.ListItemsSource(&Items)
+		.ItemWidth(ItemWidth)
+		.ItemHeight(ItemHeight)
+		.OnGenerateTile(BIND_UOBJECT_DELEGATE(STileView< UObject* >::FOnGenerateRow, HandleOnGenerateTile))
+		//.OnSelectionChanged(this, &SSocketManager::SocketSelectionChanged_Execute)
+		//.OnContextMenuOpening(this, &SSocketManager::OnContextMenuOpening)
+		//.OnItemScrolledIntoView(this, &SSocketManager::OnItemScrolledIntoView)
+		//	.HeaderRow
+		//	(
+		//		SNew(SHeaderRow)
+		//		.Visibility(EVisibility::Collapsed)
+		//		+ SHeaderRow::Column(TEXT("Socket"))
+		//	);
+		;
+}
+
+TSharedRef<ITableRow> UTileView::HandleOnGenerateTile(UObject* Item, const TSharedRef< STableViewBase >& OwnerTable) const
+{
+	// Call the user's delegate to see if they want to generate a custom widget bound to the data source.
+	if ( OnGenerateTileEvent.IsBound() )
+	{
+		UWidget* Widget = OnGenerateTileEvent.Execute(Item);
+		if ( Widget != NULL )
+		{
+			return SNew(STableRow< UObject* >, OwnerTable)
+			[
+				Widget->GetWidget()
+			];
+		}
+	}
+
+	// If a row wasn't generated just create the default one, a simple text block of the item's name.
+	return SNew(STableRow< UObject* >, OwnerTable)
+		[
+			SNew(STextBlock).Text(Item ? FText::FromString(Item->GetName()) : LOCTEXT("null", "null"))
+		];
+}
+
+/////////////////////////////////////////////////////
+
+#undef LOCTEXT_NAMESPACE
diff --git a/Engine/Source/Runtime/UMG/Private/Slate/SObjectWidget.cpp b/Engine/Source/Runtime/UMG/Private/Slate/SObjectWidget.cpp
--- a/Engine/Source/Runtime/UMG/Private/Slate/SObjectWidget.cpp
+++ b/Engine/Source/Runtime/UMG/Private/Slate/SObjectWidget.cpp
@@ -17,10 +17,16 @@ void SObjectWidget::AddReferencedObjects(FReferenceCollector& Collector)
 	Collector.AddReferencedObject(WidgetObject);
 }
 
-//int32 SObjectWidget::OnPaint(const FGeometry& AllottedGeometry, const FSlateRect& MyClippingRect, FSlateWindowElementList& OutDrawElements, int32 LayerId, const FWidgetStyle& InWidgetStyle, bool bParentEnabled) const
-//{
-//
-//}
+int32 SObjectWidget::OnPaint(const FGeometry& AllottedGeometry, const FSlateRect& MyClippingRect, FSlateWindowElementList& OutDrawElements, int32 LayerId, const FWidgetStyle& InWidgetStyle, bool bParentEnabled) const
+{
+	int32 MaxLayer = SCompoundWidget::OnPaint(AllottedGeometry, MyClippingRect, OutDrawElements, LayerId, InWidgetStyle, bParentEnabled);
+
+	//const FGeometry& AllottedGeometry, const FSlateRect& MyClippingRect, FSlateWindowElementList& OutDrawElements, int32 LayerId, const FWidgetStyle& InWidgetStyle, bool bParentEnabled
+	FPaintContext Context(AllottedGeometry, MyClippingRect, OutDrawElements, LayerId, InWidgetStyle, bParentEnabled);
+	WidgetObject->OnPaint(Context);
+
+	return FMath::Max(MaxLayer, Context.MaxLayer);
+}
 
 FReply SObjectWidget::OnKeyboardFocusReceived(const FGeometry& MyGeometry, const FKeyboardFocusEvent& InKeyboardFocusEvent)
 {
diff --git a/Engine/Source/Runtime/UMG/Private/WidgetBlueprintLibrary.cpp b/Engine/Source/Runtime/UMG/Private/WidgetBlueprintLibrary.cpp
--- a/Engine/Source/Runtime/UMG/Private/WidgetBlueprintLibrary.cpp
+++ b/Engine/Source/Runtime/UMG/Private/WidgetBlueprintLibrary.cpp
@@ -31,4 +31,58 @@ UUserWidget* UWidgetBlueprintLibrary::Create(UObject* WorldContextObject, TSubcl
 	return NewWidget;
 }
 
+void UWidgetBlueprintLibrary::DrawBox(FPaintContext& Context, FVector2D Position, FVector2D Size, USlateBrushAsset* Brush, FLinearColor Tint)
+{
+	Context.MaxLayer++;
+
+	if ( Brush )
+	{
+		FSlateDrawElement::MakeBox(
+			Context.OutDrawElements,
+			Context.MaxLayer,
+			Context.AllottedGeometry.ToPaintGeometry(Position, Size),
+			&Brush->Brush,
+			Context.MyClippingRect,
+			ESlateDrawEffect::None,
+			Tint);
+	}
+}
+
+void UWidgetBlueprintLibrary::DrawLine(UPARAM(ref) FPaintContext& Context, FVector2D PositionA, FVector2D PositionB, float Thickness, FLinearColor Tint, bool bAntiAlias)
+{
+	Context.MaxLayer++;
+
+	TArray<FVector2D> Points;
+	Points.Add(PositionA);
+	Points.Add(PositionB);
+
+	FSlateDrawElement::MakeLines(
+		Context.OutDrawElements,
+		Context.MaxLayer,
+		Context.AllottedGeometry.ToPaintGeometry(),
+		Points,
+		Context.MyClippingRect,
+		ESlateDrawEffect::None,
+		Tint,
+		bAntiAlias);
+}
+
+void UWidgetBlueprintLibrary::DrawText(UPARAM(ref) FPaintContext& Context, const FString& InString, FVector2D Position, FLinearColor Tint)
+{
+	Context.MaxLayer++;
+
+	//TODO UMG Create a font asset usable as a UFont or as a slate font asset.
+	FSlateFontInfo FontInfo = FCoreStyle::Get().GetWidgetStyle<FTextBlockStyle>("NormalText").Font;
+	
+	FSlateDrawElement::MakeText(
+		Context.OutDrawElements,
+		Context.MaxLayer,
+		Context.AllottedGeometry.ToPaintGeometry(),
+		InString,
+		FontInfo,
+		Context.MyClippingRect,
+		ESlateDrawEffect::None,
+		Tint);
+}
+
 #undef LOCTEXT_NAMESPACE
\ No newline at end of file
diff --git a/Engine/Source/Runtime/UMG/Public/Blueprint/UserWidget.h b/Engine/Source/Runtime/UMG/Public/Blueprint/UserWidget.h
--- a/Engine/Source/Runtime/UMG/Public/Blueprint/UserWidget.h
+++ b/Engine/Source/Runtime/UMG/Public/Blueprint/UserWidget.h
@@ -6,6 +6,63 @@
 
 #include "UserWidget.generated.h"
 
+static FGeometry NullGeometry;
+static FSlateRect NullRect;
+static FSlateWindowElementList NullElementList;
+static FWidgetStyle NullStyle;
+
+USTRUCT()
+struct UMG_API FPaintContext
+{
+	GENERATED_USTRUCT_BODY()
+
+public:
+
+	/** Don't ever use this constructor.  Needed for code generation. */
+	FPaintContext()
+		: AllottedGeometry(NullGeometry)
+		, MyClippingRect(NullRect)
+		, OutDrawElements(NullElementList)
+		, LayerId(0)
+		, InWidgetStyle(NullStyle)
+		, bParentEnabled(true)
+		, MaxLayer(0)
+	{ }
+
+	FPaintContext(const FGeometry& AllottedGeometry, const FSlateRect& MyClippingRect, FSlateWindowElementList& OutDrawElements, int32 LayerId, const FWidgetStyle& InWidgetStyle, bool bParentEnabled)
+		: AllottedGeometry(AllottedGeometry)
+		, MyClippingRect(MyClippingRect)
+		, OutDrawElements(OutDrawElements)
+		, LayerId(LayerId)
+		, InWidgetStyle(InWidgetStyle)
+		, bParentEnabled(bParentEnabled)
+		, MaxLayer(LayerId)
+	{
+	}
+
+	/** We override the assignment operator to allow generated code to compile with the const ref member. */
+	void operator=( const FPaintContext& Other )
+	{
+		const_cast<FGeometry&>( AllottedGeometry ) = Other.AllottedGeometry;
+		const_cast<FSlateRect&>( MyClippingRect ) = Other.MyClippingRect;
+		OutDrawElements = Other.OutDrawElements;
+		LayerId = Other.LayerId;
+		const_cast<FWidgetStyle&>( InWidgetStyle ) = Other.InWidgetStyle;
+		bParentEnabled = Other.bParentEnabled;
+	}
+
+public:
+
+	const FGeometry& AllottedGeometry;
+	const FSlateRect& MyClippingRect;
+	FSlateWindowElementList& OutDrawElements;
+	int32 LayerId;
+	const FWidgetStyle& InWidgetStyle;
+	bool bParentEnabled;
+
+	int32 MaxLayer;
+};
+
 //TODO UMG If you want to host a widget that's full screen there may need to be a SWindow equivalent that you spawn it into.
 
 DECLARE_DYNAMIC_MULTICAST_DELEGATE_OneParam(FOnVisibilityChangedEvent, ESlateVisibility::Type, Visibility);
@@ -98,6 +155,9 @@ class UMG_API UUserWidget : public UWidget
 	UFUNCTION(BlueprintPure, Category="Appearance")
 	TEnumAsByte<ESlateVisibility::Type> GetVisiblity();
 
+	UFUNCTION(BlueprintImplementableEvent, Category="User Interface")
+	void OnPaint(UPARAM(ref) FPaintContext& Context) const;
+
 	UFUNCTION(BlueprintImplementableEvent, Category="User Interface")
 	FSReply OnKeyboardFocusReceived(FGeometry MyGeometry, FKeyboardFocusEvent InKeyboardFocusEvent);
 	UFUNCTION(BlueprintImplementableEvent, Category="User Interface")
diff --git a/Engine/Source/Runtime/UMG/Public/Blueprint/WidgetBlueprintLibrary.h b/Engine/Source/Runtime/UMG/Public/Blueprint/WidgetBlueprintLibrary.h
--- a/Engine/Source/Runtime/UMG/Public/Blueprint/WidgetBlueprintLibrary.h
+++ b/Engine/Source/Runtime/UMG/Public/Blueprint/WidgetBlueprintLibrary.h
@@ -11,4 +11,25 @@ class UWidgetBlueprintLibrary : public UBlueprintFunctionLibrary
 	/** Creates a widget */
 	UFUNCTION(BlueprintCallable, meta=( HidePin="WorldContextObject", DefaultToSelf="WorldContextObject", UnsafeDuringActorConstruction = "true", FriendlyName = "Create Widget" ), Category="User Interface|Widget")
 	static class UUserWidget* Create(UObject* WorldContextObject, TSubclassOf<class UUserWidget> WidgetType);
-};
\ No newline at end of file
+
+	/** Paint Functions */
+	UFUNCTION(BlueprintCallable, Category="Painting")
+	static void DrawBox(UPARAM(ref) FPaintContext& Context, FVector2D Position, FVector2D Size, USlateBrushAsset* Brush, FLinearColor Tint = FLinearColor::White);
+
+	/**
+	 * Draws a line.
+	 *
+	 * @param PositionA		Starting position of the line in local space.
+	 * @param PositionB		Ending position of the line in local space.
+	 * @param Thickness				How many pixels thick this line should be.
+	 * @param Tint			Color to render the line.
+	 */
+	UFUNCTION(BlueprintCallable, meta=( AdvancedDisplay = "5" ), Category="Painting" )
+	static void DrawLine(UPARAM(ref) FPaintContext& Context, FVector2D PositionA, FVector2D PositionB, float Thickness = 1.0f, FLinearColor Tint = FLinearColor::White, bool bAntiAlias = true);
+
+	// TODO UMG DrawLines
+
+	UFUNCTION(BlueprintCallable, Category="Painting")
+	static void DrawText(UPARAM(ref) FPaintContext& Context, const FString& InString, FVector2D Position, FLinearColor Tint = FLinearColor::White);
+
+};
diff --git a/Engine/Source/Runtime/UMG/Public/Components/ListView.h b/Engine/Source/Runtime/UMG/Public/Components/ListView.h
--- a/Engine/Source/Runtime/UMG/Public/Components/ListView.h
+++ b/Engine/Source/Runtime/UMG/Public/Components/ListView.h
@@ -4,13 +4,11 @@
 
 #include "ListView.generated.h"
 
-UCLASS(meta=(BlueprintSpawnableComponent), ClassGroup=UserInterface)
-class UMG_API UListView : public UWidget
+UCLASS(meta=( BlueprintSpawnableComponent ), ClassGroup=UserInterface)
+class UMG_API UListView : public UTableViewBase
 {
 	GENERATED_UCLASS_BODY()
 
-	DECLARE_DYNAMIC_DELEGATE_RetVal_OneParam(UWidget*, FOnGenerateRowUObject, UObject*, Item);
-		
 public:
 
 	UPROPERTY(EditDefaultsOnly, Category=Content)
diff --git a/Engine/Source/Runtime/UMG/Public/Components/TableViewBase.h b/Engine/Source/Runtime/UMG/Public/Components/TableViewBase.h
new file mode 100644
--- /dev/null
+++ b/Engine/Source/Runtime/UMG/Public/Components/TableViewBase.h
@@ -0,0 +1,14 @@
+// Copyright 1998-2014 Epic Games, Inc. All Rights Reserved.
+
+#pragma once
+
+#include "TableViewBase.generated.h"
+
+UCLASS(Abstract, meta=(BlueprintSpawnableComponent), ClassGroup=UserInterface)
+class UMG_API UTableViewBase : public UWidget
+{
+	GENERATED_UCLASS_BODY()
+
+	DECLARE_DYNAMIC_DELEGATE_RetVal_OneParam(UWidget*, FOnGenerateRowUObject, UObject*, Item);
+	
+};
diff --git a/Engine/Source/Runtime/UMG/Public/Components/TileView.h b/Engine/Source/Runtime/UMG/Public/Components/TileView.h
new file mode 100644
--- /dev/null
+++ b/Engine/Source/Runtime/UMG/Public/Components/TileView.h
@@ -0,0 +1,46 @@
+// Copyright 1998-2014 Epic Games, Inc. All Rights Reserved.
+
+#pragma once
+
+#include "TileView.generated.h"
+
+UCLASS(meta=(BlueprintSpawnableComponent), ClassGroup=UserInterface)
+class UMG_API UTileView : public UTableViewBase
+{
+	GENERATED_UCLASS_BODY()
+
+public:
+
+	UPROPERTY(EditDefaultsOnly, Category=Content)
+	float ItemWidth;
+
+	UPROPERTY(EditDefaultsOnly, Category=Content)
+	float ItemHeight;
+
+	UPROPERTY(EditAnywhere, BlueprintReadWrite, Category=Content)
+	TArray<UObject*> Items;
+
+	UPROPERTY(EditDefaultsOnly, Category=Content)
+	TEnumAsByte<ESelectionMode::Type> SelectionMode;
+
+	/** Called when the button is clicked */
+	UPROPERTY(EditDefaultsOnly, Category=Events)
+	FOnGenerateRowUObject OnGenerateTileEvent;
+
+protected:
+	TSharedPtr< STileView<UObject*> > ListWidget() const
+	{
+		if ( MyWidget.IsValid() )
+		{
+			return StaticCastSharedRef< STileView<UObject*> >(MyWidget.ToSharedRef());
+		}
+
+		return TSharedPtr< STileView<UObject*> >();
+	}
+
+	TSharedRef<ITableRow> HandleOnGenerateTile(UObject* Item, const TSharedRef< STableViewBase >& OwnerTable) const;
+
+	// UWidget interface
+	virtual TSharedRef<SWidget> RebuildWidget() OVERRIDE;
+	// End of UWidget interface
+};
diff --git a/Engine/Source/Runtime/UMG/Public/Slate/SObjectWidget.h b/Engine/Source/Runtime/UMG/Public/Slate/SObjectWidget.h
--- a/Engine/Source/Runtime/UMG/Public/Slate/SObjectWidget.h
+++ b/Engine/Source/Runtime/UMG/Public/Slate/SObjectWidget.h
@@ -13,6 +13,8 @@ class SObjectWidget : public SCompoundWidget, public FGCObject
 
 	virtual void AddReferencedObjects(FReferenceCollector& Collector) OVERRIDE;
 
+	virtual int32 OnPaint(const FGeometry& AllottedGeometry, const FSlateRect& MyClippingRect, FSlateWindowElementList& OutDrawElements, int32 LayerId, const FWidgetStyle& InWidgetStyle, bool bParentEnabled) const OVERRIDE;
+
 	virtual FReply OnKeyboardFocusReceived(const FGeometry& MyGeometry, const FKeyboardFocusEvent& InKeyboardFocusEvent) OVERRIDE;
 	virtual void OnKeyboardFocusLost(const FKeyboardFocusEvent& InKeyboardFocusEvent) OVERRIDE;
 	virtual void OnKeyboardFocusChanging(const FWeakWidgetPath& PreviousFocusPath, const FWidgetPath& NewWidgetPath) OVERRIDE;
diff --git a/Engine/Source/Runtime/UMG/Public/UMG.h b/Engine/Source/Runtime/UMG/Public/UMG.h
--- a/Engine/Source/Runtime/UMG/Public/UMG.h
+++ b/Engine/Source/Runtime/UMG/Public/UMG.h
@@ -63,7 +63,9 @@
 #include "ThrobberComponent.h"
 #include "CircularThrobberComponent.h"
 
+#include "TableViewBase.h"
 #include "ListView.h"
+#include "TileView.h"
 
 /* Slate
  *****************************************************************************/
