UMG - Breaking backwards compatibility with a massive rename pass.  Adding more documentation.

diff --git a/Engine/Source/Editor/UMGEditor/Private/Extensions/CanvasSlotExtension.cpp b/Engine/Source/Editor/UMGEditor/Private/Extensions/CanvasSlotExtension.cpp
--- a/Engine/Source/Editor/UMGEditor/Private/Extensions/CanvasSlotExtension.cpp
+++ b/Engine/Source/Editor/UMGEditor/Private/Extensions/CanvasSlotExtension.cpp
@@ -84,7 +84,7 @@ void FCanvasSlotExtension::MoveByAmount(UWidget* Widget, FVector2D Delta)
 	if ( !Delta.IsZero() )
 	{
 		UCanvasPanelSlot* CanvasSlot = Cast<UCanvasPanelSlot>(Widget->Slot);
-		UCanvasPanelComponent* Parent = Cast<UCanvasPanelComponent>(CanvasSlot->Parent);
+		UCanvasPanel* Parent = Cast<UCanvasPanel>(CanvasSlot->Parent);
 
 		CanvasSlot->SetPosition(CanvasSlot->Position + Delta);
 	}
diff --git a/Engine/Source/Editor/UMGEditor/Private/Extensions/HorizontalSlotExtension.cpp b/Engine/Source/Editor/UMGEditor/Private/Extensions/HorizontalSlotExtension.cpp
--- a/Engine/Source/Editor/UMGEditor/Private/Extensions/HorizontalSlotExtension.cpp
+++ b/Engine/Source/Editor/UMGEditor/Private/Extensions/HorizontalSlotExtension.cpp
@@ -63,7 +63,7 @@ FReply FHorizontalSlotExtension::HandleShift(int32 ShiftAmount)
 void FHorizontalSlotExtension::ShiftHorizontal(UWidget* Widget, int32 ShiftAmount)
 {
 	UHorizontalBoxSlot* Slot = Cast<UHorizontalBoxSlot>(Widget->Slot);
-	UHorizontalBoxComponent* Parent = Cast<UHorizontalBoxComponent>(Slot->Parent);
+	UHorizontalBox* Parent = Cast<UHorizontalBox>(Slot->Parent);
 
 	int32 CurrentIndex = Parent->GetChildIndex(Widget);
 	Parent->Slots.RemoveAt(CurrentIndex);
diff --git a/Engine/Source/Editor/UMGEditor/Private/Extensions/VerticalSlotExtension.cpp b/Engine/Source/Editor/UMGEditor/Private/Extensions/VerticalSlotExtension.cpp
--- a/Engine/Source/Editor/UMGEditor/Private/Extensions/VerticalSlotExtension.cpp
+++ b/Engine/Source/Editor/UMGEditor/Private/Extensions/VerticalSlotExtension.cpp
@@ -77,7 +77,7 @@ FReply FVerticalSlotExtension::HandleDownPressed()
 void FVerticalSlotExtension::MoveUp(UWidget* Widget)
 {
 	UVerticalBoxSlot* VerticalSlot = Cast<UVerticalBoxSlot>(Widget->Slot);
-	UVerticalBoxComponent* Parent = Cast<UVerticalBoxComponent>(VerticalSlot->Parent);
+	UVerticalBox* Parent = Cast<UVerticalBox>(VerticalSlot->Parent);
 
 	int32 CurrentIndex = Parent->GetChildIndex(Widget);
 	Parent->Slots.RemoveAt(CurrentIndex);
@@ -87,7 +87,7 @@ void FVerticalSlotExtension::MoveUp(UWidget* Widget)
 void FVerticalSlotExtension::MoveDown(UWidget* Widget)
 {
 	UVerticalBoxSlot* VerticalSlot = Cast<UVerticalBoxSlot>(Widget->Slot);
-	UVerticalBoxComponent* Parent = Cast<UVerticalBoxComponent>(VerticalSlot->Parent);
+	UVerticalBox* Parent = Cast<UVerticalBox>(VerticalSlot->Parent);
 
 	int32 CurrentIndex = Parent->GetChildIndex(Widget);
 	Parent->Slots.RemoveAt(CurrentIndex);
diff --git a/Engine/Source/Editor/UMGEditor/Private/SUMGEditorTree.cpp b/Engine/Source/Editor/UMGEditor/Private/SUMGEditorTree.cpp
--- a/Engine/Source/Editor/UMGEditor/Private/SUMGEditorTree.cpp
+++ b/Engine/Source/Editor/UMGEditor/Private/SUMGEditorTree.cpp
@@ -235,23 +235,23 @@ FReply SUMGEditorTree::CreateTestUI()
 
 	if ( WidgetTemplates.Num() > 0 )
 	{
-		UVerticalBoxComponent* Vertical = CastChecked<UVerticalBoxComponent>(WidgetTemplates[2]);
+		UVerticalBox* Vertical = CastChecked<UVerticalBox>(WidgetTemplates[2]);
 
-		UButtonComponent* NewButton = BP->WidgetTree->ConstructWidget<UButtonComponent>(UButtonComponent::StaticClass());
+		UButton* NewButton = BP->WidgetTree->ConstructWidget<UButton>(UButton::StaticClass());
 
 		Vertical->AddChild(NewButton, FVector2D(0,0));
 	}
 	else
 	{
-		UCanvasPanelComponent* Canvas = BP->WidgetTree->ConstructWidget<UCanvasPanelComponent>(UCanvasPanelComponent::StaticClass());
-		UBorderComponent* Border = BP->WidgetTree->ConstructWidget<UBorderComponent>(UBorderComponent::StaticClass());
-		UVerticalBoxComponent* Vertical = BP->WidgetTree->ConstructWidget<UVerticalBoxComponent>(UVerticalBoxComponent::StaticClass());
-		UButtonComponent* Button1 = BP->WidgetTree->ConstructWidget<UButtonComponent>(UButtonComponent::StaticClass());
-		UButtonComponent* Button2 = BP->WidgetTree->ConstructWidget<UButtonComponent>(UButtonComponent::StaticClass());
-		UButtonComponent* Button3 = BP->WidgetTree->ConstructWidget<UButtonComponent>(UButtonComponent::StaticClass());
-		UButtonComponent* Button4 = BP->WidgetTree->ConstructWidget<UButtonComponent>(UButtonComponent::StaticClass());
-
-		UTextBlockComponent* Text1 = BP->WidgetTree->ConstructWidget<UTextBlockComponent>(UTextBlockComponent::StaticClass());
+		UCanvasPanel* Canvas = BP->WidgetTree->ConstructWidget<UCanvasPanel>(UCanvasPanel::StaticClass());
+		UBorder* Border = BP->WidgetTree->ConstructWidget<UBorder>(UBorder::StaticClass());
+		UVerticalBox* Vertical = BP->WidgetTree->ConstructWidget<UVerticalBox>(UVerticalBox::StaticClass());
+		UButton* Button1 = BP->WidgetTree->ConstructWidget<UButton>(UButton::StaticClass());
+		UButton* Button2 = BP->WidgetTree->ConstructWidget<UButton>(UButton::StaticClass());
+		UButton* Button3 = BP->WidgetTree->ConstructWidget<UButton>(UButton::StaticClass());
+		UButton* Button4 = BP->WidgetTree->ConstructWidget<UButton>(UButton::StaticClass());
+
+		UTextBlock* Text1 = BP->WidgetTree->ConstructWidget<UTextBlock>(UTextBlock::StaticClass());
 		Text1->Text = FText::FromString(TEXT("Button!"));
 		Button4->SetContent(Text1);
 
diff --git a/Engine/Source/Editor/UMGEditor/Private/Templates/WidgetTemplateButton.cpp b/Engine/Source/Editor/UMGEditor/Private/Templates/WidgetTemplateButton.cpp
--- a/Engine/Source/Editor/UMGEditor/Private/Templates/WidgetTemplateButton.cpp
+++ b/Engine/Source/Editor/UMGEditor/Private/Templates/WidgetTemplateButton.cpp
@@ -20,8 +20,8 @@ FText FWidgetTemplateButton::GetCategory()
 
 UWidget* FWidgetTemplateButton::Create(UWidgetTree* Tree)
 {
-	UButtonComponent* Button = Tree->ConstructWidget<UButtonComponent>(UButtonComponent::StaticClass());
-	UTextBlockComponent* Text = Tree->ConstructWidget<UTextBlockComponent>(UTextBlockComponent::StaticClass());
+	UButton* Button = Tree->ConstructWidget<UButton>(UButton::StaticClass());
+	UTextBlock* Text = Tree->ConstructWidget<UTextBlock>(UTextBlock::StaticClass());
 
 	Text->Text = FText::FromString(TEXT("Button Text"));
 	Text->Font.Size = 10;
diff --git a/Engine/Source/Editor/UMGEditor/Private/Templates/WidgetTemplateCheckBox.cpp b/Engine/Source/Editor/UMGEditor/Private/Templates/WidgetTemplateCheckBox.cpp
--- a/Engine/Source/Editor/UMGEditor/Private/Templates/WidgetTemplateCheckBox.cpp
+++ b/Engine/Source/Editor/UMGEditor/Private/Templates/WidgetTemplateCheckBox.cpp
@@ -20,9 +20,9 @@ FText FWidgetTemplateCheckBox::GetCategory()
 
 UWidget* FWidgetTemplateCheckBox::Create(UWidgetTree* Tree)
 {
-	UHorizontalBoxComponent* Horizontal = Tree->ConstructWidget<UHorizontalBoxComponent>(UHorizontalBoxComponent::StaticClass());
-	UCheckBoxComponent* Checkbox = Tree->ConstructWidget<UCheckBoxComponent>(UCheckBoxComponent::StaticClass());
-	UTextBlockComponent* Text = Tree->ConstructWidget<UTextBlockComponent>(UTextBlockComponent::StaticClass());
+	UHorizontalBox* Horizontal = Tree->ConstructWidget<UHorizontalBox>(UHorizontalBox::StaticClass());
+	UCheckBox* Checkbox = Tree->ConstructWidget<UCheckBox>(UCheckBox::StaticClass());
+	UTextBlock* Text = Tree->ConstructWidget<UTextBlock>(UTextBlock::StaticClass());
 	Text->Text = LOCTEXT("CheckboxText", "Checkbox Text");
 	Text->Font.Size = 10;
 
diff --git a/Engine/Source/Runtime/UMG/Private/Components/Border.cpp b/Engine/Source/Runtime/UMG/Private/Components/Border.cpp
new file mode 100644
--- /dev/null
+++ b/Engine/Source/Runtime/UMG/Private/Components/Border.cpp
@@ -0,0 +1,97 @@
+// Copyright 1998-2014 Epic Games, Inc. All Rights Reserved.
+
+#include "UMGPrivatePCH.h"
+
+#define LOCTEXT_NAMESPACE "UMG"
+
+/////////////////////////////////////////////////////
+// UBorder
+
+UBorder::UBorder(const FPostConstructInitializeProperties& PCIP)
+	: Super(PCIP)
+{
+	bIsVariable = false;
+
+	SBorder::FArguments BorderDefaults;
+
+	HorizontalAlignment = BorderDefaults._HAlign;
+	VerticalAlignment = BorderDefaults._VAlign;
+	ContentPadding = BorderDefaults._Padding.Get();
+	//BorderImage = *BorderDefaults._BorderImage.Get();
+	DesiredSizeScale = BorderDefaults._DesiredSizeScale.Get();
+	ContentScale = BorderDefaults._ContentScale.Get();
+	ContentColorAndOpacity = FLinearColor::White;
+	BorderColorAndOpacity = FLinearColor::White;
+	ForegroundColor = FLinearColor::Black;
+	bShowEffectWhenDisabled = BorderDefaults._ShowEffectWhenDisabled.Get();
+}
+
+TSharedRef<SWidget> UBorder::RebuildWidget()
+{
+	TSharedRef<SBorder> NewBorder = SNew(SBorder)
+		.HAlign(HorizontalAlignment)
+		.VAlign(VerticalAlignment)
+		.Padding(BIND_UOBJECT_ATTRIBUTE(FMargin, GetContentPadding))
+		// SLATE_EVENT( FPointerEventHandler, OnMouseButtonDown )
+		// SLATE_EVENT( FPointerEventHandler, OnMouseButtonUp )
+		// SLATE_EVENT( FPointerEventHandler, OnMouseMove )
+		// SLATE_EVENT( FPointerEventHandler, OnMouseDoubleClick )
+		.BorderImage(BIND_UOBJECT_ATTRIBUTE(const FSlateBrush*, GetBorderBrush))
+		.ContentScale(ContentScale)
+		.DesiredSizeScale(DesiredSizeScale)
+		.ColorAndOpacity(BIND_UOBJECT_ATTRIBUTE(FLinearColor, GetContentColor))
+		.BorderBackgroundColor(BIND_UOBJECT_ATTRIBUTE(FSlateColor, GetBorderColor))
+		.ForegroundColor(BIND_UOBJECT_ATTRIBUTE(FSlateColor, GetForegroundColor))
+		.ShowEffectWhenDisabled(bShowEffectWhenDisabled);
+
+	NewBorder->SetContent(Content ? Content->GetWidget() : SNullWidget::NullWidget);
+
+	MyBorder = NewBorder;
+	return NewBorder;
+}
+
+void UBorder::SetContent(UWidget* Content)
+{
+	Super::SetContent(Content);
+
+	TSharedPtr<SBorder> Border = MyBorder.Pin();
+	if ( Border.IsValid() )
+	{
+		Border->SetContent(Content ? Content->GetWidget() : SNullWidget::NullWidget);
+	}
+}
+
+FMargin UBorder::GetContentPadding() const
+{
+	return ContentPadding;
+}
+
+FLinearColor UBorder::GetContentColor() const
+{
+	return ContentColorAndOpacity;
+}
+
+const FSlateBrush* UBorder::GetBorderBrush() const
+{
+	if ( BorderBrush == NULL )
+	{
+		SBorder::FArguments BorderDefaults;
+		return BorderDefaults._BorderImage.Get();
+	}
+
+	return &BorderBrush->Brush;
+}
+
+FSlateColor UBorder::GetBorderColor() const
+{
+	return BorderColorAndOpacity;
+}
+
+FSlateColor UBorder::GetForegroundColor() const
+{
+	return ForegroundColor;
+}
+
+/////////////////////////////////////////////////////
+
+#undef LOCTEXT_NAMESPACE
diff --git a/Engine/Source/Runtime/UMG/Private/Components/BorderComponent.cpp b/Engine/Source/Runtime/UMG/Private/Components/BorderComponent.cpp
deleted file mode 100644
diff --git a/Engine/Source/Runtime/UMG/Private/Components/Button.cpp b/Engine/Source/Runtime/UMG/Private/Components/Button.cpp
new file mode 100644
--- /dev/null
+++ b/Engine/Source/Runtime/UMG/Private/Components/Button.cpp
@@ -0,0 +1,121 @@
+// Copyright 1998-2014 Epic Games, Inc. All Rights Reserved.
+
+#include "UMGPrivatePCH.h"
+
+#define LOCTEXT_NAMESPACE "UMG"
+
+/////////////////////////////////////////////////////
+// UButton
+
+UButton::UButton(const FPostConstructInitializeProperties& PCIP)
+	: Super(PCIP)
+{
+	SButton::FArguments ButtonDefaults;
+
+	Style = NULL;
+	HorizontalAlignment = ButtonDefaults._HAlign;
+	VerticalAlignment = ButtonDefaults._VAlign;
+	ContentPadding = ButtonDefaults._ContentPadding.Get();
+
+	DesiredSizeScale = ButtonDefaults._DesiredSizeScale.Get();
+	ContentScale = ButtonDefaults._ContentScale.Get();
+
+	ButtonColorAndOpacity = FLinearColor::White;
+	ForegroundColor = FLinearColor::Black;
+}
+
+TSharedRef<SWidget> UButton::RebuildWidget()
+{
+	SButton::FArguments ButtonDefaults;
+
+	const FButtonStyle* StylePtr = ( Style != NULL ) ? Style->GetStyle<FButtonStyle>() : NULL;
+	if (StylePtr == NULL)
+	{
+		StylePtr = ButtonDefaults._ButtonStyle;
+	}
+
+	TOptional<FSlateSound> OptionalPressedSound;
+	if ( PressedSound.GetResourceObject() )
+	{
+		OptionalPressedSound = PressedSound;
+	}
+
+	TOptional<FSlateSound> OptionalHoveredSound;
+	if ( HoveredSound.GetResourceObject() )
+	{
+		OptionalHoveredSound = HoveredSound;
+	}
+
+	TSharedRef<SButton> NewButton = SNew(SButton)
+		.ButtonStyle(StylePtr)
+		.HAlign(HorizontalAlignment)
+		.VAlign(VerticalAlignment)
+		.ContentPadding(ContentPadding)
+		.DesiredSizeScale(DesiredSizeScale)
+		.ContentScale(ContentScale)
+		.ButtonColorAndOpacity(ButtonColorAndOpacity)
+		.ForegroundColor(ForegroundColor)
+		.PressedSoundOverride(OptionalPressedSound)
+		.HoveredSoundOverride(OptionalHoveredSound)
+		.OnClicked(BIND_UOBJECT_DELEGATE(FOnClicked, HandleOnClicked))
+		;
+
+	NewButton->SetContent(Content ? Content->GetWidget() : SNullWidget::NullWidget);
+	
+	return NewButton;
+}
+
+void UButton::SetContent(UWidget* InContent)
+{
+	Super::SetContent(InContent);
+
+	if ( ButtonWidget().IsValid() )
+	{
+		ButtonWidget()->SetContent(InContent ? InContent->GetWidget() : SNullWidget::NullWidget);
+	}
+}
+
+FLinearColor UButton::GetButtonColorAndOpacity()
+{
+	return ButtonWidget()->GetBorderBackgroundColor().GetSpecifiedColor();
+}
+
+void UButton::SetButtonColorAndOpacity(FLinearColor InButtonColorAndOpacity)
+{
+	ButtonWidget()->SetBorderBackgroundColor(InButtonColorAndOpacity);
+}
+
+FLinearColor UButton::GetForegroundColor()
+{
+	return ButtonWidget()->GetForegroundColor().GetSpecifiedColor();
+}
+
+void UButton::SetForegroundColor(FLinearColor InForegroundColor)
+{
+	ButtonWidget()->SetForegroundColor(InForegroundColor);
+}
+
+FMargin UButton::GetContentPadding()
+{
+	return ButtonWidget()->GetContentPadding();
+}
+
+void UButton::SetContentPadding(FMargin InContentPadding)
+{
+	ButtonWidget()->SetContentPadding(InContentPadding);
+}
+
+FReply UButton::HandleOnClicked()
+{
+	if ( OnClickedEvent.IsBound() )
+	{
+		return OnClickedEvent.Execute().ToReply();
+	}
+	
+	return FReply::Unhandled();
+}
+
+
+/////////////////////////////////////////////////////
+
+#undef LOCTEXT_NAMESPACE
diff --git a/Engine/Source/Runtime/UMG/Private/Components/ButtonComponent.cpp b/Engine/Source/Runtime/UMG/Private/Components/ButtonComponent.cpp
deleted file mode 100644
diff --git a/Engine/Source/Runtime/UMG/Private/Components/CanvasPanel.cpp b/Engine/Source/Runtime/UMG/Private/Components/CanvasPanel.cpp
new file mode 100644
--- /dev/null
+++ b/Engine/Source/Runtime/UMG/Private/Components/CanvasPanel.cpp
@@ -0,0 +1,128 @@
+// Copyright 1998-2014 Epic Games, Inc. All Rights Reserved.
+
+#include "UMGPrivatePCH.h"
+
+/////////////////////////////////////////////////////
+// SFixedSizeCanvas
+
+class SFixedSizeCanvas : public SCanvas
+{
+	SLATE_BEGIN_ARGS( SFixedSizeCanvas )
+		{
+			_Visibility = EVisibility::SelfHitTestInvisible;
+		}
+	SLATE_END_ARGS()
+
+	FVector2D CanvasSize;
+
+	void Construct(const FArguments& InArgs, FVector2D InSize)
+	{
+		SCanvas::FArguments ParentArgs;
+		SCanvas::Construct(ParentArgs);
+
+		CanvasSize = InSize;
+	}
+
+	// SWidget interface
+	virtual FVector2D ComputeDesiredSize() const OVERRIDE
+	{
+		return CanvasSize;
+	}
+	// End of SWidget interface
+};
+
+/////////////////////////////////////////////////////
+// UCanvasPanel
+
+UCanvasPanel::UCanvasPanel(const FPostConstructInitializeProperties& PCIP)
+	: Super(PCIP)
+{
+	bIsVariable = false;
+	DesiredCanvasSize = FVector2D(128.0f, 128.0f);
+}
+
+int32 UCanvasPanel::GetChildrenCount() const
+{
+	return Slots.Num();
+}
+
+UWidget* UCanvasPanel::GetChildAt(int32 Index) const
+{
+	return Slots[Index]->Content;
+}
+
+TSharedRef<SWidget> UCanvasPanel::RebuildWidget()
+{
+	TSharedRef<SFixedSizeCanvas> NewCanvas = SNew(SFixedSizeCanvas, DesiredCanvasSize);
+	MyCanvas = NewCanvas;
+
+	TSharedPtr<SFixedSizeCanvas> Canvas = MyCanvas.Pin();
+	if ( Canvas.IsValid() )
+	{
+		Canvas->ClearChildren();
+
+		for ( int32 SlotIndex = 0; SlotIndex < Slots.Num(); ++SlotIndex )
+		{
+			UCanvasPanelSlot* Slot = Slots[SlotIndex];
+			if ( Slot == NULL )
+			{
+				Slots[SlotIndex] = Slot = ConstructObject<UCanvasPanelSlot>(UCanvasPanelSlot::StaticClass(), this);
+			}
+
+			Slot->BuildSlot(Canvas.ToSharedRef());
+		}
+	}
+
+	return NewCanvas;
+}
+
+UCanvasPanelSlot* UCanvasPanel::AddSlot(UWidget* Content)
+{
+	UCanvasPanelSlot* Slot = ConstructObject<UCanvasPanelSlot>(UCanvasPanelSlot::StaticClass(), this);
+	Slot->Content = Content;
+	Slot->Parent = this;
+
+#if WITH_EDITOR
+	Content->Slot = Slot;
+#endif
+	
+	Slots.Add(Slot);
+
+	return Slot;
+}
+
+bool UCanvasPanel::AddChild(UWidget* Child, FVector2D Position)
+{
+	UCanvasPanelSlot* Slot = AddSlot(Child);
+	Slot->Position = Position;
+	Slot->Size = FVector2D(100, 25);
+
+	return true;
+}
+
+#if WITH_EDITOR
+
+void UCanvasPanel::ConnectEditorData()
+{
+	for ( UCanvasPanelSlot* Slot : Slots )
+	{
+		Slot->Parent = this;
+		Slot->Content->Slot = Slot;
+	}
+}
+
+void UCanvasPanel::PostEditChangeProperty(struct FPropertyChangedEvent& PropertyChangedEvent)
+{
+	// Ensure the slots have unique names
+	int32 SlotNumbering = 1;
+
+	TSet<FName> UniqueSlotNames;
+	for (int32 SlotIndex = 0; SlotIndex < Slots.Num(); ++SlotIndex)
+	{
+		if ( Slots[SlotIndex] == NULL )
+		{
+			Slots[SlotIndex] = ConstructObject<UCanvasPanelSlot>(UCanvasPanelSlot::StaticClass(), this);
+		}
+	}
+}
+#endif
\ No newline at end of file
diff --git a/Engine/Source/Runtime/UMG/Private/Components/CanvasPanelComponent.cpp b/Engine/Source/Runtime/UMG/Private/Components/CanvasPanelComponent.cpp
deleted file mode 100644
diff --git a/Engine/Source/Runtime/UMG/Private/Components/CheckBox.cpp b/Engine/Source/Runtime/UMG/Private/Components/CheckBox.cpp
new file mode 100644
--- /dev/null
+++ b/Engine/Source/Runtime/UMG/Private/Components/CheckBox.cpp
@@ -0,0 +1,90 @@
+// Copyright 1998-2014 Epic Games, Inc. All Rights Reserved.
+
+#include "UMGPrivatePCH.h"
+
+#define LOCTEXT_NAMESPACE "UMG"
+
+/////////////////////////////////////////////////////
+// UCheckBox
+
+UCheckBox::UCheckBox(const FPostConstructInitializeProperties& PCIP)
+	: Super(PCIP)
+{
+	SCheckBox::FArguments CheckBoxDefaults;
+
+	bIsChecked = false;
+
+	HorizontalAlignment = CheckBoxDefaults._HAlign;
+	Padding = CheckBoxDefaults._Padding.Get();
+
+	ForegroundColor = FLinearColor::White;
+	BorderBackgroundColor = FLinearColor::White;
+}
+
+TSharedRef<SWidget> UCheckBox::RebuildWidget()
+{
+	const FCheckBoxStyle* StylePtr = ( Style != NULL ) ? Style->GetStyle<FCheckBoxStyle>() : NULL;
+	if ( StylePtr == NULL )
+	{
+		SCheckBox::FArguments Defaults;
+		StylePtr = Defaults._Style;
+	}
+
+	TOptional<FSlateSound> OptionalCheckedSound;
+	if ( CheckedSound.GetResourceObject() )
+	{
+		OptionalCheckedSound = CheckedSound;
+	}
+
+	TOptional<FSlateSound> OptionalUncheckedSound;
+	if ( UncheckedSound.GetResourceObject() )
+	{
+		OptionalUncheckedSound = UncheckedSound;
+	}
+
+	TOptional<FSlateSound> OptionalHoveredSound;
+	if ( HoveredSound.GetResourceObject() )
+	{
+		OptionalHoveredSound = HoveredSound;
+	}
+
+	return SNew(SCheckBox)
+		.Style(StylePtr)
+		.OnCheckStateChanged( BIND_UOBJECT_DELEGATE(FOnCheckStateChanged, SlateOnCheckStateChangedCallback) )
+		.IsChecked( BIND_UOBJECT_ATTRIBUTE(ESlateCheckBoxState::Type, GetCheckState) )
+		.HAlign( HorizontalAlignment )
+		.Padding( Padding )
+		.ForegroundColor( ForegroundColor )
+		.BorderBackgroundColor( BorderBackgroundColor )
+		.CheckedSoundOverride(OptionalCheckedSound)
+		.UncheckedSoundOverride(OptionalUncheckedSound)
+		.HoveredSoundOverride(OptionalHoveredSound)
+		;
+// 		SLATE_ATTRIBUTE( bool, ReadOnly )
+// 		SLATE_ARGUMENT( bool, IsFocusable )
+// 		SLATE_EVENT( FOnGetContent, OnGetMenuContent )
+}
+
+ESlateCheckBoxState::Type UCheckBox::GetCheckState() const
+{
+	return bIsChecked ? ESlateCheckBoxState::Checked : ESlateCheckBoxState::Unchecked;
+}
+
+void UCheckBox::SlateOnCheckStateChangedCallback(ESlateCheckBoxState::Type NewState)
+{
+	//@TODO: Choosing to treat Undetermined as Checked
+	const bool bWantsToBeChecked = NewState != ESlateCheckBoxState::Unchecked;
+
+	if (OnCheckStateChanged.IsBound())
+	{
+		OnCheckStateChanged.Broadcast(bWantsToBeChecked);
+	}
+	else
+	{
+		bIsChecked = bWantsToBeChecked;
+	}
+}
+
+/////////////////////////////////////////////////////
+
+#undef LOCTEXT_NAMESPACE
diff --git a/Engine/Source/Runtime/UMG/Private/Components/CheckBoxComponent.cpp b/Engine/Source/Runtime/UMG/Private/Components/CheckBoxComponent.cpp
deleted file mode 100644
diff --git a/Engine/Source/Runtime/UMG/Private/Components/EditableTextBlock.cpp b/Engine/Source/Runtime/UMG/Private/Components/EditableTextBlock.cpp
new file mode 100644
--- /dev/null
+++ b/Engine/Source/Runtime/UMG/Private/Components/EditableTextBlock.cpp
@@ -0,0 +1,64 @@
+// Copyright 1998-2014 Epic Games, Inc. All Rights Reserved.
+
+#include "UMGPrivatePCH.h"
+
+#define LOCTEXT_NAMESPACE "UMG"
+
+/////////////////////////////////////////////////////
+// UEditableTextBlock
+
+UEditableTextBlock::UEditableTextBlock(const FPostConstructInitializeProperties& PCIP)
+	: Super(PCIP)
+{
+	ForegroundColor = FLinearColor::Black;
+	BackgroundColor = FLinearColor::White;
+	ReadOnlyForegroundColor = FLinearColor::Black;
+
+	// HACK Special font initialization hack since there are no font assets yet for slate.
+	Font = FSlateFontInfo(TEXT("Slate/Fonts/Roboto-Bold.ttf"), 12);
+
+	// Grab other defaults from slate arguments.
+	SEditableTextBox::FArguments Defaults;
+	IsReadOnly = Defaults._IsReadOnly.Get();
+	IsPassword = Defaults._IsReadOnly.Get();
+	MinimumDesiredWidth = Defaults._MinDesiredWidth.Get();
+	Padding = Defaults._Padding.Get();
+	IsCaretMovedWhenGainFocus = Defaults._IsCaretMovedWhenGainFocus.Get();
+	SelectAllTextWhenFocused = Defaults._SelectAllTextWhenFocused.Get();
+	RevertTextOnEscape = Defaults._RevertTextOnEscape.Get();
+	ClearKeyboardFocusOnCommit = Defaults._ClearKeyboardFocusOnCommit.Get();
+	SelectAllTextOnCommit = Defaults._SelectAllTextOnCommit.Get();
+}
+
+TSharedRef<SWidget> UEditableTextBlock::RebuildWidget()
+{
+	FString FontPath = FPaths::EngineContentDir() / Font.FontName.ToString();
+
+	return SNew(SEditableTextBox)
+		.Text(Text)
+		.HintText(HintText)
+		.Font(FSlateFontInfo(FontPath, Font.Size))
+		.ForegroundColor(ForegroundColor)
+		.BackgroundColor(BackgroundColor)
+		.ReadOnlyForegroundColor(ReadOnlyForegroundColor)
+		.IsReadOnly(IsReadOnly)
+		.IsPassword(IsPassword)
+		.MinDesiredWidth(MinimumDesiredWidth)
+		.Padding(Padding)
+		.IsCaretMovedWhenGainFocus(IsCaretMovedWhenGainFocus)
+		.SelectAllTextWhenFocused(SelectAllTextWhenFocused)
+		.RevertTextOnEscape(RevertTextOnEscape)
+		.ClearKeyboardFocusOnCommit(ClearKeyboardFocusOnCommit)
+		.SelectAllTextOnCommit(SelectAllTextOnCommit)
+		.OnTextChanged(BIND_UOBJECT_DELEGATE(FOnTextChanged, SlateOnTextChanged))
+		;
+}
+
+void UEditableTextBlock::SlateOnTextChanged(const FText& Text)
+{
+	OnTextChanged.Broadcast(Text);
+}
+
+/////////////////////////////////////////////////////
+
+#undef LOCTEXT_NAMESPACE
diff --git a/Engine/Source/Runtime/UMG/Private/Components/EditableTextBlockComponent.cpp b/Engine/Source/Runtime/UMG/Private/Components/EditableTextBlockComponent.cpp
deleted file mode 100644
diff --git a/Engine/Source/Runtime/UMG/Private/Components/HorizontalBox.cpp b/Engine/Source/Runtime/UMG/Private/Components/HorizontalBox.cpp
new file mode 100644
--- /dev/null
+++ b/Engine/Source/Runtime/UMG/Private/Components/HorizontalBox.cpp
@@ -0,0 +1,152 @@
+// Copyright 1998-2014 Epic Games, Inc. All Rights Reserved.
+
+#include "UMGPrivatePCH.h"
+
+/////////////////////////////////////////////////////
+// UHorizontalBox
+
+UHorizontalBox::UHorizontalBox(const FPostConstructInitializeProperties& PCIP)
+	: Super(PCIP)
+{
+	bIsVariable = false;
+}
+
+int32 UHorizontalBox::GetChildrenCount() const
+{
+	return Slots.Num();
+}
+
+UWidget* UHorizontalBox::GetChildAt(int32 Index) const
+{
+	return Slots[Index]->Content;
+}
+
+int32 UHorizontalBox::GetChildIndex(UWidget* Content) const
+{
+	for ( int32 SlotIndex = 0; SlotIndex < Slots.Num(); ++SlotIndex )
+	{
+		UHorizontalBoxSlot* Slot = Slots[SlotIndex];
+
+		if ( Slot->Content == Content )
+		{
+			return SlotIndex;
+		}
+	}
+
+	return -1;
+}
+
+bool UHorizontalBox::AddChild(UWidget* Child, FVector2D Position)
+{
+	AddSlot(Child);
+	return true;
+}
+
+bool UHorizontalBox::RemoveChild(UWidget* Child)
+{
+	int32 SlotIndex = GetChildIndex(Child);
+	if ( SlotIndex != -1 )
+	{
+		Slots.RemoveAt(SlotIndex);
+		return true;
+	}
+
+	return false;
+}
+
+void UHorizontalBox::ReplaceChildAt(int32 Index, UWidget* Content)
+{
+	UHorizontalBoxSlot* Slot = Slots[Index];
+	Slot->Content = Content;
+
+#if WITH_EDITOR
+	Content->Slot = Slot;
+#endif
+}
+
+void UHorizontalBox::InsertChildAt(int32 Index, UWidget* Content)
+{
+	UHorizontalBoxSlot* Slot = ConstructObject<UHorizontalBoxSlot>(UHorizontalBoxSlot::StaticClass(), this);
+	Slot->Content = Content;
+	Slot->Parent = this;
+
+#if WITH_EDITOR
+	Content->Slot = Slot;
+#endif
+
+	Slots.Insert(Slot, Index);
+}
+
+TSharedRef<SWidget> UHorizontalBox::RebuildWidget()
+{
+	TSharedRef<SHorizontalBox> NewCanvas = SNew(SHorizontalBox);
+	MyHorizontalBox = NewCanvas;
+
+	TSharedPtr<SHorizontalBox> Canvas = MyHorizontalBox.Pin();
+	if ( Canvas.IsValid() )
+	{
+		Canvas->ClearChildren();
+
+		for ( int32 SlotIndex = 0; SlotIndex < Slots.Num(); ++SlotIndex )
+		{
+			UHorizontalBoxSlot* Slot = Slots[SlotIndex];
+			if ( Slot == NULL )
+			{
+				Slots[SlotIndex] = Slot = ConstructObject<UHorizontalBoxSlot>(UHorizontalBoxSlot::StaticClass(), this);
+			}
+
+			auto& NewSlot = Canvas->AddSlot()
+				.Padding(Slot->Padding)
+				.HAlign(Slot->HorizontalAlignment)
+				.VAlign(Slot->VerticalAlignment)
+				[
+					Slot->Content == NULL ? SNullWidget::NullWidget : Slot->Content->GetWidget()
+				];
+
+			NewSlot.SizeParam = UWidget::ConvertSerializedSizeParamToRuntime(Slot->Size);
+		}
+	}
+
+	return NewCanvas;
+}
+
+UHorizontalBoxSlot* UHorizontalBox::AddSlot(UWidget* Content)
+{
+	UHorizontalBoxSlot* Slot = ConstructObject<UHorizontalBoxSlot>(UHorizontalBoxSlot::StaticClass(), this);
+	Slot->Content = Content;
+	Slot->Parent = this;
+
+#if WITH_EDITOR
+	Content->Slot = Slot;
+#endif
+	
+	Slots.Add(Slot);
+
+	return Slot;
+}
+
+#if WITH_EDITOR
+void UHorizontalBox::ConnectEditorData()
+{
+	for ( UHorizontalBoxSlot* Slot : Slots )
+	{
+		Slot->Content->Slot = Slot;
+		Slot->Parent = this;
+	}
+}
+
+void UHorizontalBox::PostEditChangeProperty(struct FPropertyChangedEvent& PropertyChangedEvent)
+{
+	// Ensure the slots have unique names
+	int32 SlotNumbering = 1;
+
+	TSet<FName> UniqueSlotNames;
+	for (int32 SlotIndex = 0; SlotIndex < Slots.Num(); ++SlotIndex)
+	{
+		if ( Slots[SlotIndex] == NULL )
+		{
+			Slots[SlotIndex] = ConstructObject<UHorizontalBoxSlot>(UHorizontalBoxSlot::StaticClass(), this);
+		}
+	}
+}
+#endif
diff --git a/Engine/Source/Runtime/UMG/Private/Components/HorizontalBoxComponent.cpp b/Engine/Source/Runtime/UMG/Private/Components/HorizontalBoxComponent.cpp
deleted file mode 100644
diff --git a/Engine/Source/Runtime/UMG/Private/Components/Image.cpp b/Engine/Source/Runtime/UMG/Private/Components/Image.cpp
new file mode 100644
--- /dev/null
+++ b/Engine/Source/Runtime/UMG/Private/Components/Image.cpp
@@ -0,0 +1,44 @@
+// Copyright 1998-2014 Epic Games, Inc. All Rights Reserved.
+
+#include "UMGPrivatePCH.h"
+
+#define LOCTEXT_NAMESPACE "UMG"
+
+/////////////////////////////////////////////////////
+// UImage
+
+UImage::UImage(const FPostConstructInitializeProperties& PCIP)
+	: Super(PCIP)
+	, Image()
+	, ColorAndOpacity(FLinearColor::White)
+{
+	//@TODO: Would like to default to the checkerboard, but TextureName always wins over TextureObject :(
+	//	Image = *FEditorStyle::GetDefaultBrush();
+}
+
+TSharedRef<SWidget> UImage::RebuildWidget()
+{
+	return SNew(SImage)
+		.Image( BIND_UOBJECT_ATTRIBUTE(const FSlateBrush*, GetImageBrush) )
+		.ColorAndOpacity( BIND_UOBJECT_ATTRIBUTE(FSlateColor, GetColorAndOpacity) );
+}
+
+const FSlateBrush* UImage::GetImageBrush() const
+{
+	if ( Image == NULL )
+	{
+		SImage::FArguments ImageDefaults;
+		return ImageDefaults._Image.Get();
+	}
+
+	return &Image->Brush;
+}
+
+FSlateColor UImage::GetColorAndOpacity() const
+{
+	return ColorAndOpacity;
+}
+
+/////////////////////////////////////////////////////
+
+#undef LOCTEXT_NAMESPACE
diff --git a/Engine/Source/Runtime/UMG/Private/Components/ImageComponent.cpp b/Engine/Source/Runtime/UMG/Private/Components/ImageComponent.cpp
deleted file mode 100644
diff --git a/Engine/Source/Runtime/UMG/Private/Components/ScrollBox.cpp b/Engine/Source/Runtime/UMG/Private/Components/ScrollBox.cpp
new file mode 100644
--- /dev/null
+++ b/Engine/Source/Runtime/UMG/Private/Components/ScrollBox.cpp
@@ -0,0 +1,23 @@
+// Copyright 1998-2014 Epic Games, Inc. All Rights Reserved.
+
+#include "UMGPrivatePCH.h"
+
+#define LOCTEXT_NAMESPACE "UMG"
+
+/////////////////////////////////////////////////////
+// UScrollBox
+
+UScrollBox::UScrollBox(const FPostConstructInitializeProperties& PCIP)
+	: Super(PCIP)
+{
+	bIsVariable = false;
+}
+
+TSharedRef<SWidget> UScrollBox::RebuildWidget()
+{
+	return SNew(SScrollBox);
+}
+
+/////////////////////////////////////////////////////
+
+#undef LOCTEXT_NAMESPACE
diff --git a/Engine/Source/Runtime/UMG/Private/Components/ScrollBoxComponent.cpp b/Engine/Source/Runtime/UMG/Private/Components/ScrollBoxComponent.cpp
deleted file mode 100644
diff --git a/Engine/Source/Runtime/UMG/Private/Components/Spacer.cpp b/Engine/Source/Runtime/UMG/Private/Components/Spacer.cpp
new file mode 100644
--- /dev/null
+++ b/Engine/Source/Runtime/UMG/Private/Components/Spacer.cpp
@@ -0,0 +1,30 @@
+// Copyright 1998-2014 Epic Games, Inc. All Rights Reserved.
+
+#include "UMGPrivatePCH.h"
+
+#define LOCTEXT_NAMESPACE "UMG"
+
+/////////////////////////////////////////////////////
+// USpacer
+
+USpacer::USpacer(const FPostConstructInitializeProperties& PCIP)
+	: Super(PCIP)
+	, Size(1.0f, 1.0f)
+{
+	bIsVariable = false;
+}
+
+TSharedRef<SWidget> USpacer::RebuildWidget()
+{
+	return SNew(SSpacer)
+		.Size( BIND_UOBJECT_ATTRIBUTE(FVector2D, GetSpacerSize) );
+}
+
+FVector2D USpacer::GetSpacerSize() const
+{
+	return Size;
+}
+
+/////////////////////////////////////////////////////
+
+#undef LOCTEXT_NAMESPACE
diff --git a/Engine/Source/Runtime/UMG/Private/Components/SpacerComponent.cpp b/Engine/Source/Runtime/UMG/Private/Components/SpacerComponent.cpp
deleted file mode 100644
diff --git a/Engine/Source/Runtime/UMG/Private/Components/SpinningImageComponent.cpp b/Engine/Source/Runtime/UMG/Private/Components/SpinningImageComponent.cpp
--- a/Engine/Source/Runtime/UMG/Private/Components/SpinningImageComponent.cpp
+++ b/Engine/Source/Runtime/UMG/Private/Components/SpinningImageComponent.cpp
@@ -5,7 +5,7 @@
 #define LOCTEXT_NAMESPACE "UMG"
 
 /////////////////////////////////////////////////////
-// UImageComponent
+// UImage
 
 USpinningImageComponent::USpinningImageComponent(const FPostConstructInitializeProperties& PCIP)
 	: Super(PCIP)
diff --git a/Engine/Source/Runtime/UMG/Private/Components/TextBlock.cpp b/Engine/Source/Runtime/UMG/Private/Components/TextBlock.cpp
new file mode 100644
--- /dev/null
+++ b/Engine/Source/Runtime/UMG/Private/Components/TextBlock.cpp
@@ -0,0 +1,52 @@
+// Copyright 1998-2014 Epic Games, Inc. All Rights Reserved.
+
+#include "UMGPrivatePCH.h"
+
+#define LOCTEXT_NAMESPACE "UMG"
+
+/////////////////////////////////////////////////////
+// UTextBlock
+
+UTextBlock::UTextBlock(const FPostConstructInitializeProperties& PCIP)
+	: Super(PCIP)
+{
+	bIsVariable = false;
+
+	Text = LOCTEXT("TextBlockDefaultValue", "Text Block");
+	ShadowOffset = FVector2D(1.0f, 1.0f);
+	ColorAndOpacity = FLinearColor::White;
+	ShadowColorAndOpacity = FLinearColor::Transparent;
+
+	Style = NULL;
+
+	// HACK Special font initialization hack since there are no font assets yet for slate.
+	Font = FSlateFontInfo(TEXT("Slate/Fonts/Roboto-Bold.ttf"), 24);
+}
+
+TSharedRef<SWidget> UTextBlock::RebuildWidget()
+{
+	FString FontPath = FPaths::EngineContentDir() / Font.FontName.ToString();
+
+	const FTextBlockStyle* StylePtr = ( Style != NULL ) ? Style->GetStyle<FTextBlockStyle>() : NULL;
+	if ( StylePtr == NULL )
+	{
+		STextBlock::FArguments Defaults;
+		StylePtr = Defaults._TextStyle;
+	}
+
+	TAttribute<FText> TextBinding = OPTIONAL_BINDING(FText, Text);
+	TAttribute<FSlateColor> ColorAndOpacityBinding = OPTIONAL_BINDING(FSlateColor, ColorAndOpacity);
+	TAttribute<FLinearColor> ShadowColorAndOpacityBinding = OPTIONAL_BINDING(FLinearColor, ShadowColorAndOpacity);
+	
+	return SNew(STextBlock)
+		.Text( TextBinding )
+		.TextStyle( StylePtr )
+		.Font( FSlateFontInfo(FontPath, Font.Size) )
+		.ColorAndOpacity( ColorAndOpacityBinding )
+		.ShadowOffset( ShadowOffset )
+		.ShadowColorAndOpacity( ShadowColorAndOpacityBinding );
+}
+
+/////////////////////////////////////////////////////
+
+#undef LOCTEXT_NAMESPACE
diff --git a/Engine/Source/Runtime/UMG/Private/Components/TextBlockComponent.cpp b/Engine/Source/Runtime/UMG/Private/Components/TextBlockComponent.cpp
deleted file mode 100644
diff --git a/Engine/Source/Runtime/UMG/Private/Components/VerticalBox.cpp b/Engine/Source/Runtime/UMG/Private/Components/VerticalBox.cpp
new file mode 100644
--- /dev/null
+++ b/Engine/Source/Runtime/UMG/Private/Components/VerticalBox.cpp
@@ -0,0 +1,151 @@
+// Copyright 1998-2014 Epic Games, Inc. All Rights Reserved.
+
+#include "UMGPrivatePCH.h"
+
+/////////////////////////////////////////////////////
+// UVerticalBox
+
+UVerticalBox::UVerticalBox(const FPostConstructInitializeProperties& PCIP)
+	: Super(PCIP)
+{
+	bIsVariable = false;
+}
+
+int32 UVerticalBox::GetChildrenCount() const
+{
+	return Slots.Num();
+}
+
+UWidget* UVerticalBox::GetChildAt(int32 Index) const
+{
+	return Slots[Index]->Content;
+}
+
+int32 UVerticalBox::GetChildIndex(UWidget* Content) const
+{
+	for ( int32 SlotIndex = 0; SlotIndex < Slots.Num(); ++SlotIndex )
+	{
+		UVerticalBoxSlot* Slot = Slots[SlotIndex];
+
+		if ( Slot->Content == Content )
+		{
+			return SlotIndex;
+		}
+	}
+
+	return -1;
+}
+
+bool UVerticalBox::AddChild(UWidget* Child, FVector2D Position)
+{
+	AddSlot(Child);
+	return true;
+}
+
+bool UVerticalBox::RemoveChild(UWidget* Child)
+{
+	int32 SlotIndex = GetChildIndex(Child);
+	if ( SlotIndex != -1 )
+	{
+		Slots.RemoveAt(SlotIndex);
+		return true;
+	}
+
+	return false;
+}
+
+void UVerticalBox::ReplaceChildAt(int32 Index, UWidget* Content)
+{
+	UVerticalBoxSlot* Slot = Slots[Index];
+	Slot->Content = Content;
+
+#if WITH_EDITOR
+	Content->Slot = Slot;
+#endif
+}
+
+void UVerticalBox::InsertChildAt(int32 Index, UWidget* Content)
+{
+	UVerticalBoxSlot* Slot = ConstructObject<UVerticalBoxSlot>(UVerticalBoxSlot::StaticClass(), this);
+	Slot->Content = Content;
+	Slot->Parent = this;
+
+#if WITH_EDITOR
+	Content->Slot = Slot;
+#endif
+
+	Slots.Insert(Slot, Index);
+}
+
+TSharedRef<SWidget> UVerticalBox::RebuildWidget()
+{
+	TSharedRef<SVerticalBox> NewCanvas = SNew(SVerticalBox);
+	MyVerticalBox = NewCanvas;
+
+	TSharedPtr<SVerticalBox> Canvas = MyVerticalBox.Pin();
+	if ( Canvas.IsValid() )
+	{
+		Canvas->ClearChildren();
+
+		for ( int32 SlotIndex = 0; SlotIndex < Slots.Num(); ++SlotIndex )
+		{
+			UVerticalBoxSlot* Slot = Slots[SlotIndex];
+			if ( Slot == NULL )
+			{
+				Slots[SlotIndex] = Slot = ConstructObject<UVerticalBoxSlot>(UVerticalBoxSlot::StaticClass(), this);
+			}
+
+			Slot->Parent = this;
+			Slot->BuildSlot(NewCanvas);
+		}
+	}
+
+	return NewCanvas;
+}
+
+UVerticalBoxSlot* UVerticalBox::AddSlot(UWidget* Content)
+{
+	UVerticalBoxSlot* Slot = ConstructObject<UVerticalBoxSlot>(UVerticalBoxSlot::StaticClass(), this);
+	Slot->Content = Content;
+	Slot->Parent = this;
+
+#if WITH_EDITOR
+	Content->Slot = Slot;
+#endif
+	
+	Slots.Add(Slot);
+
+	return Slot;
+}
+
+#if WITH_EDITOR
+void UVerticalBox::ConnectEditorData()
+{
+	for ( UVerticalBoxSlot* Slot : Slots )
+	{
+		// TODO UMG Find a better way to assign the parent pointers for slots.
+		Slot->Parent = this;
+
+		if ( Slot->Content )
+		{
+			Slot->Content->Slot = Slot;
+		}
+		//TODO UMG Should we auto delete empty slots?
+	}
+}
+
+void UVerticalBox::PostEditChangeProperty(struct FPropertyChangedEvent& PropertyChangedEvent)
+{
+	// Ensure the slots have unique names
+	int32 SlotNumbering = 1;
+
+	TSet<FName> UniqueSlotNames;
+	for (int32 SlotIndex = 0; SlotIndex < Slots.Num(); ++SlotIndex)
+	{
+		if ( Slots[SlotIndex] == NULL )
+		{
+			Slots[SlotIndex] = ConstructObject<UVerticalBoxSlot>(UVerticalBoxSlot::StaticClass(), this);
+		}
+	}
+}
+#endif
diff --git a/Engine/Source/Runtime/UMG/Private/Components/VerticalBoxComponent.cpp b/Engine/Source/Runtime/UMG/Private/Components/VerticalBoxComponent.cpp
deleted file mode 100644
diff --git a/Engine/Source/Runtime/UMG/Private/Components/Widget.cpp b/Engine/Source/Runtime/UMG/Private/Components/Widget.cpp
--- a/Engine/Source/Runtime/UMG/Private/Components/Widget.cpp
+++ b/Engine/Source/Runtime/UMG/Private/Components/Widget.cpp
@@ -42,12 +42,6 @@ void UWidget::SetToolTipText(const FText& InToolTipText)
 	MyWidget->SetToolTipText(ToolTipText);
 }
 
-void UWidget::OnUnregister()
-{
-	MyWidget = NULL;
-//	Super::OnUnregister();
-}
-
 void UWidget::PostEditChangeProperty(struct FPropertyChangedEvent& PropertyChangedEvent)
 {
 	// TODO Rename?
@@ -60,6 +54,7 @@ TSharedRef<SWidget> UWidget::GetWidget()
 		MyWidget = RebuildWidget();
 		MyWidget->SetEnabled(OPTIONAL_BINDING(bool, bIsEnabled));
 		MyWidget->SetVisibility(OPTIONAL_BINDING_CONVERT(ESlateVisibility::Type, Visiblity, EVisibility, ConvertVisibility));
+		//MyWidget->SetCursor(OPTIONAL_BINDING(EMouseCursor)
 
 		if ( !ToolTipText.IsEmpty() )
 		{
diff --git a/Engine/Source/Runtime/UMG/Public/Blueprint/UserWidget.h b/Engine/Source/Runtime/UMG/Public/Blueprint/UserWidget.h
--- a/Engine/Source/Runtime/UMG/Public/Blueprint/UserWidget.h
+++ b/Engine/Source/Runtime/UMG/Public/Blueprint/UserWidget.h
@@ -11,6 +11,10 @@ static FSlateRect NullRect;
 static FSlateWindowElementList NullElementList;
 static FWidgetStyle NullStyle;
 
+/**
+ * The state passed into OnPaint that we can expose as a single painting structure to blueprints to
+ * allow script code to override OnPaint behavior.
+ */
 USTRUCT()
 struct UMG_API FPaintContext
 {
@@ -67,6 +71,9 @@ public:
 
 DECLARE_DYNAMIC_MULTICAST_DELEGATE_OneParam(FOnVisibilityChangedEvent, ESlateVisibility::Type, Visibility);
 
+/**
+ * 
+ */
 UCLASS(Abstract, editinlinenew, BlueprintType, Blueprintable)
 class UMG_API UUserWidget : public UWidget
 {
diff --git a/Engine/Source/Runtime/UMG/Public/Components/Border.h b/Engine/Source/Runtime/UMG/Public/Components/Border.h
new file mode 100644
--- /dev/null
+++ b/Engine/Source/Runtime/UMG/Public/Components/Border.h
@@ -0,0 +1,77 @@
+// Copyright 1998-2014 Epic Games, Inc. All Rights Reserved.
+
+#pragma once
+
+#include "Border.generated.h"
+
+/** A border is a container widget that can contain one child widget, providing an opportunity to surround it with a border image and adjustable padding. */
+
+UCLASS(meta=(BlueprintSpawnableComponent), ClassGroup=UserInterface)
+class UMG_API UBorder : public UContentWidget
+{
+	GENERATED_UCLASS_BODY()
+
+	/** Horizontal positioning of the content within the border */
+	UPROPERTY(EditAnywhere, BlueprintReadOnly, Category=Appearance)
+	TEnumAsByte<EHorizontalAlignment> HorizontalAlignment;
+
+	/** Vertical positioning of the content within the border */
+	UPROPERTY(EditAnywhere, BlueprintReadOnly, Category=Appearance)
+	TEnumAsByte<EVerticalAlignment> VerticalAlignment;
+
+	/** The padding to add around the inner content. */
+	UPROPERTY(EditAnywhere, BlueprintReadWrite, Category=Appearance)
+	FMargin ContentPadding;
+
+	//@TODO: Should these be exposed?
+	// SLATE_EVENT( FPointerEventHandler, OnMouseButtonDown )
+	// SLATE_EVENT( FPointerEventHandler, OnMouseButtonUp )
+	// SLATE_EVENT( FPointerEventHandler, OnMouseMove )
+	// SLATE_EVENT( FPointerEventHandler, OnMouseDoubleClick )
+
+	/** Image to use for the border */
+	UPROPERTY(EditAnywhere, BlueprintReadWrite, Category=Appearance, meta=( DisplayThumbnail = "true" ))
+	USlateBrushAsset* BorderBrush;
+
+	/** The scaling factor for the border */
+	UPROPERTY(EditAnywhere, BlueprintReadOnly, Category=Appearance)
+	FVector2D DesiredSizeScale;
+
+	/** The scaling factor for the border content */
+	UPROPERTY(EditAnywhere, BlueprintReadOnly, Category=Appearance)
+	FVector2D ContentScale;
+
+	/** Color and opacity multiplier of content in the border */
+	UPROPERTY(EditAnywhere, BlueprintReadWrite, Category=Appearance)
+	FLinearColor ContentColorAndOpacity;
+
+	/** Color and opacity of the actual border image */
+	UPROPERTY(EditAnywhere, BlueprintReadWrite, Category=Appearance)
+	FSlateColor BorderColorAndOpacity;
+
+	/** The foreground color of text and some glyphs that appear as the border's content. */
+	UPROPERTY(EditAnywhere, BlueprintReadWrite, Category=Appearance)
+	FSlateColor ForegroundColor;
+
+	/** Whether or not to show the disabled effect when this border is disabled */
+	UPROPERTY(EditAnywhere, BlueprintReadOnly, Category=Appearance)
+	bool bShowEffectWhenDisabled;
+
+	// UContentWidget interface
+	virtual void SetContent(UWidget* Content) OVERRIDE;
+	// End UContentWidget interface
+
+protected:
+	TWeakPtr<class SBorder> MyBorder;
+
+protected:
+	// UWidget interface
+	virtual TSharedRef<SWidget> RebuildWidget() OVERRIDE;
+	// End of UWidget interface
+
+	FMargin GetContentPadding() const;
+	FLinearColor GetContentColor() const;
+	const FSlateBrush* GetBorderBrush() const;
+	FSlateColor GetBorderColor() const;
+	FSlateColor GetForegroundColor() const;
+};
diff --git a/Engine/Source/Runtime/UMG/Public/Components/BorderComponent.h b/Engine/Source/Runtime/UMG/Public/Components/BorderComponent.h
deleted file mode 100644
diff --git a/Engine/Source/Runtime/UMG/Public/Components/Button.h b/Engine/Source/Runtime/UMG/Public/Components/Button.h
new file mode 100644
--- /dev/null
+++ b/Engine/Source/Runtime/UMG/Public/Components/Button.h
@@ -0,0 +1,106 @@
+// Copyright 1998-2014 Epic Games, Inc. All Rights Reserved.
+
+#pragma once
+
+#include "Button.generated.h"
+
+DECLARE_DYNAMIC_MULTICAST_DELEGATE(FOnButtonClickedEvent);
+
+/** Buttons are clickable widgets */
+UCLASS(meta=(BlueprintSpawnableComponent), ClassGroup=UserInterface)
+class UMG_API UButton : public UContentWidget
+{
+	GENERATED_UCLASS_BODY()
+
+public:
+	/** Style of the button */
+	UPROPERTY(EditDefaultsOnly, Category=Style, meta=( DisplayThumbnail = "true" ))
+	USlateWidgetStyleAsset* Style;
+
+	/** Horizontal positioning of the content within the button */
+	UPROPERTY(EditDefaultsOnly, Category=Appearance)
+	TEnumAsByte<EHorizontalAlignment> HorizontalAlignment;
+
+	/** Vertical positioning of the content within the button */
+	UPROPERTY(EditDefaultsOnly, Category=Appearance)
+	TEnumAsByte<EVerticalAlignment> VerticalAlignment;
+
+	/** The padding to add around the button content. */
+	UPROPERTY(EditDefaultsOnly, Category=Appearance)
+	FMargin ContentPadding;
+
+	/** The scaling factor for the button border */
+	UPROPERTY(EditDefaultsOnly, Category=Appearance)
+	FVector2D DesiredSizeScale;
+
+	/** The scaling factor for the button content */
+	UPROPERTY(EditDefaultsOnly, Category=Appearance)
+	FVector2D ContentScale;
+
+	/** The color multiplier for the button images */
+	UPROPERTY(EditDefaultsOnly, Category=Appearance)
+	FLinearColor ButtonColorAndOpacity;
+
+	/** The foreground color of the button */
+	UPROPERTY(EditDefaultsOnly, Category=Appearance)
+	FLinearColor ForegroundColor;
+
+	UPROPERTY(EditDefaultsOnly, Category=Sound)
+	FSlateSound PressedSound;
+
+	UPROPERTY(EditDefaultsOnly, Category=Sound)
+	FSlateSound HoveredSound;
+
+	/** Called when the button is clicked */
+	UPROPERTY(EditDefaultsOnly, Category=Events)
+	FOnReply OnClickedEvent;
+	
+	/*  */
+	UFUNCTION(BlueprintPure, Category="Appearance")
+	FLinearColor GetButtonColorAndOpacity();
+	
+	/*  */
+	UFUNCTION(BlueprintCallable, Category="Appearance")
+	void SetButtonColorAndOpacity(FLinearColor InButtonColorAndOpacity);
+
+	/*  */
+	UFUNCTION(BlueprintPure, Category="Appearance")
+	FLinearColor GetForegroundColor();
+
+	/*  */
+	UFUNCTION(BlueprintCallable, Category="Appearance")
+	void SetForegroundColor(FLinearColor InForegroundColor);
+	
+	/*  */
+	UFUNCTION(BlueprintPure, Category="Appearance")
+	FMargin GetContentPadding();
+
+	/*  */
+	UFUNCTION(BlueprintCallable, Category="Appearance")
+	void SetContentPadding(FMargin InContentPadding);
+
+
+	// UContentWidget interface
+	virtual void SetContent(UWidget* Content) OVERRIDE;
+	// End UContentWidget interface
+
+protected:
+	TSharedPtr<class SButton> ButtonWidget() const
+	{
+		if ( MyWidget.IsValid() )
+		{
+			return StaticCastSharedRef<SButton>(MyWidget.ToSharedRef());
+		}
+
+		return TSharedPtr<class SButton>();
+	}
+
+protected:
+	// UWidget interface
+	virtual TSharedRef<SWidget> RebuildWidget() OVERRIDE;
+	// End of UWidget interface
+
+	FMargin GetContentPadding() const;
+
+	FReply HandleOnClicked();
+};
diff --git a/Engine/Source/Runtime/UMG/Public/Components/ButtonComponent.h b/Engine/Source/Runtime/UMG/Public/Components/ButtonComponent.h
deleted file mode 100644
diff --git a/Engine/Source/Runtime/UMG/Public/Components/CanvasPanel.h b/Engine/Source/Runtime/UMG/Public/Components/CanvasPanel.h
new file mode 100644
--- /dev/null
+++ b/Engine/Source/Runtime/UMG/Public/Components/CanvasPanel.h
@@ -0,0 +1,46 @@
+// Copyright 1998-2014 Epic Games, Inc. All Rights Reserved.
+
+#pragma once
+
+#include "CanvasPanel.generated.h"
+
+UCLASS(meta=(BlueprintSpawnableComponent), ClassGroup=UserInterface)
+class UMG_API UCanvasPanel : public UPanelWidget
+{
+	GENERATED_UCLASS_BODY()
+
+	/** The items placed on the canvas */
+	UPROPERTY(EditAnywhere, EditInline, Category=Slots)
+	TArray<UCanvasPanelSlot*> Slots;
+
+	/** The desired size of the canvas */
+	UPROPERTY(EditAnywhere, Category=Appearance)
+	FVector2D DesiredCanvasSize;
+
+	UCanvasPanelSlot* AddSlot(UWidget* Content);
+	virtual bool AddChild(UWidget* Child, FVector2D Position) OVERRIDE;
+
+	// UPanelWidget
+	virtual int32 GetChildrenCount() const OVERRIDE;
+	virtual UWidget* GetChildAt(int32 Index) const OVERRIDE;
+	// End UPanelWidget
+
+#if WITH_EDITOR
+	// UObject interface
+	virtual void PostEditChangeProperty(struct FPropertyChangedEvent& PropertyChangedEvent) OVERRIDE;
+	// End of UObject interface
+
+	// UWidget interface
+	virtual void ConnectEditorData() OVERRIDE;
+	// End UWidget interface
+#endif
+
+protected:
+
+	TWeakPtr<class SFixedSizeCanvas> MyCanvas;
+
+protected:
+	// UWidget interface
+	virtual TSharedRef<SWidget> RebuildWidget() OVERRIDE;
+	// End of UWidget interface
+};
diff --git a/Engine/Source/Runtime/UMG/Public/Components/CanvasPanelComponent.h b/Engine/Source/Runtime/UMG/Public/Components/CanvasPanelComponent.h
deleted file mode 100644
diff --git a/Engine/Source/Runtime/UMG/Public/Components/CheckBox.h b/Engine/Source/Runtime/UMG/Public/Components/CheckBox.h
new file mode 100644
--- /dev/null
+++ b/Engine/Source/Runtime/UMG/Public/Components/CheckBox.h
@@ -0,0 +1,65 @@
+// Copyright 1998-2014 Epic Games, Inc. All Rights Reserved.
+
+#pragma once
+
+#include "CheckBox.generated.h"
+
+DECLARE_DYNAMIC_MULTICAST_DELEGATE_OneParam( FOnCheckBoxComponentStateChanged, bool, bIsChecked );
+
+/** Check box widget */
+UCLASS(meta=(BlueprintSpawnableComponent), ClassGroup=UserInterface)
+class UMG_API UCheckBox : public UWidget
+{
+	GENERATED_UCLASS_BODY()
+
+protected:
+	/** Style of the check box */
+	UPROPERTY(EditAnywhere, BlueprintReadOnly, Category=Style)
+	USlateWidgetStyleAsset* Style;
+
+	/** Called when the checked state has changed */
+	UPROPERTY(BlueprintAssignable)
+	FOnCheckBoxComponentStateChanged OnCheckStateChanged;
+
+	/** Whether the check box is currently in a checked state */
+	UPROPERTY(EditAnywhere, BlueprintReadWrite, Category=Appearance)
+	bool bIsChecked;
+
+	/** How the content of the toggle button should align within the given space */
+	UPROPERTY(EditAnywhere, BlueprintReadOnly, Category=Appearance)
+	TEnumAsByte<EHorizontalAlignment> HorizontalAlignment;
+
+	/** Spacing between the check box image and its content */
+	UPROPERTY(EditAnywhere, BlueprintReadOnly, Category=Appearance)
+	FMargin Padding;
+
+	/** Foreground color for the checkbox's content and parts */
+	UPROPERTY(EditAnywhere, BlueprintReadOnly, Category=Appearance)
+	FSlateColor ForegroundColor;
+
+	/** The color of the background border */
+	UPROPERTY(EditAnywhere, BlueprintReadOnly, Category=Appearance)
+	FSlateColor BorderBackgroundColor;
+
+	//SLATE_ATTRIBUTE( bool, ReadOnly )
+	//SLATE_ARGUMENT( bool, IsFocusable )
+
+	//SLATE_EVENT( FOnGetContent, OnGetMenuContent )
+
+	UPROPERTY(EditAnywhere, Category=Sound)
+	FSlateSound CheckedSound;
+
+	UPROPERTY(EditAnywhere, Category=Sound)
+	FSlateSound UncheckedSound;
+
+	UPROPERTY(EditAnywhere, Category=Sound)
+	FSlateSound HoveredSound;
+
+protected:
+	// UWidget interface
+	virtual TSharedRef<SWidget> RebuildWidget() OVERRIDE;
+	// End of UWidget interface
+
+	ESlateCheckBoxState::Type GetCheckState() const;
+	void SlateOnCheckStateChangedCallback(ESlateCheckBoxState::Type NewState);
+};
diff --git a/Engine/Source/Runtime/UMG/Public/Components/CheckBoxComponent.h b/Engine/Source/Runtime/UMG/Public/Components/CheckBoxComponent.h
deleted file mode 100644
diff --git a/Engine/Source/Runtime/UMG/Public/Components/EditableTextBlock.h b/Engine/Source/Runtime/UMG/Public/Components/EditableTextBlock.h
new file mode 100644
--- /dev/null
+++ b/Engine/Source/Runtime/UMG/Public/Components/EditableTextBlock.h
@@ -0,0 +1,95 @@
+// Copyright 1998-2014 Epic Games, Inc. All Rights Reserved.
+
+#pragma once
+
+#include "EditableTextBlock.generated.h"
+
+DECLARE_DYNAMIC_MULTICAST_DELEGATE_OneParam(FOnEditableTextBlock_TextChangedEvent, const FText&, Text);
+//DECLARE_DYNAMIC_MULTICAST_DELEGATE_TwoParams(FOnEditableTextBlock_TextCommittedEvent, const FText&, Text, ETextCommit::Type, Type);
+
+/** Editable text box widget */
+UCLASS(meta=(BlueprintSpawnableComponent), ClassGroup=UserInterface)
+class UMG_API UEditableTextBlock : public UWidget
+{
+	GENERATED_UCLASS_BODY()
+
+protected:
+
+	/** The text content for this editable text box widget */
+	UPROPERTY(EditAnywhere, BlueprintReadWrite, Category=Content)
+	FText Text;
+
+	/** Hint text that appears when there is no text in the text box */
+	UPROPERTY(EditAnywhere, BlueprintReadWrite, Category=Content)
+	FText HintText;
+
+	/** Font color and opacity (overrides Style) */
+	UPROPERTY(EditAnywhere, Category=Appearance)
+	FSlateFontInfo Font;
+
+	/** Text color and opacity (overrides Style) */
+	UPROPERTY(EditAnywhere, BlueprintReadWrite, Category=Appearance)
+	FSlateColor ForegroundColor;
+
+	/** The color of the background/border around the editable text (overrides Style) */
+	UPROPERTY(EditAnywhere, BlueprintReadWrite, Category=Appearance)
+	FSlateColor BackgroundColor;
+
+	/** Text color and opacity when read-only (overrides Style) */
+	UPROPERTY(EditAnywhere, BlueprintReadWrite, Category=Appearance)
+	FSlateColor ReadOnlyForegroundColor;
+
+	/** Sets whether this text box can actually be modified interactively by the user */
+	UPROPERTY(EditAnywhere, BlueprintReadWrite, Category=Appearance)
+	bool IsReadOnly;
+
+	/** Sets whether this text box is for storing a password */
+	UPROPERTY(EditAnywhere, BlueprintReadWrite, Category=Appearance)
+	bool IsPassword;
+
+	/** Minimum width that a text block should be */
+	UPROPERTY(EditAnywhere, BlueprintReadWrite, Category=Appearance)
+	float MinimumDesiredWidth;
+
+	/** Padding between the box/border and the text widget inside (overrides Style) */
+	UPROPERTY(EditAnywhere, BlueprintReadWrite, Category=Appearance)
+	FMargin Padding;
+
+	/** Workaround as we loose focus when the auto completion closes. */
+	UPROPERTY(EditAnywhere, BlueprintReadWrite, Category=Behavior)
+	bool IsCaretMovedWhenGainFocus;
+
+	/** Whether to select all text when the user clicks to give focus on the widget */
+	UPROPERTY(EditAnywhere, BlueprintReadWrite, Category=Behavior)
+	bool SelectAllTextWhenFocused;
+
+	/** Whether to allow the user to back out of changes when they press the escape key */
+	UPROPERTY(EditAnywhere, BlueprintReadWrite, Category=Behavior)
+	bool RevertTextOnEscape;
+
+	/** Whether to clear keyboard focus when pressing enter to commit changes */
+	UPROPERTY(EditAnywhere, BlueprintReadWrite, Category=Behavior)
+	bool ClearKeyboardFocusOnCommit;
+
+	/** Whether to select all text when pressing enter to commit changes */
+	UPROPERTY(EditAnywhere, BlueprintReadWrite, Category=Behavior)
+	bool SelectAllTextOnCommit;
+
+	/** Called whenever the text is changed interactively by the user */
+	UPROPERTY(BlueprintAssignable)
+	FOnEditableTextBlock_TextChangedEvent OnTextChanged;
+
+	/** Called whenever the text is committed.  This happens when the user presses enter or the text box loses focus. *///
+	//UPROPERTY(BlueprintAssignable)
+	//FOnTextCommittedEvent OnTextCommitted;
+
+	/** Provide a alternative mechanism for error reporting. */
+	//SLATE_ARGUMENT(TSharedPtr<class IErrorReportingWidget>, ErrorReporting)
+
+protected:
+	// UWidget interface
+	virtual TSharedRef<SWidget> RebuildWidget() OVERRIDE;
+	// End of UWidget
+
+	void SlateOnTextChanged(const FText& Text);
+};
diff --git a/Engine/Source/Runtime/UMG/Public/Components/EditableTextBlockComponent.h b/Engine/Source/Runtime/UMG/Public/Components/EditableTextBlockComponent.h
deleted file mode 100644
diff --git a/Engine/Source/Runtime/UMG/Public/Components/HorizontalBox.h b/Engine/Source/Runtime/UMG/Public/Components/HorizontalBox.h
new file mode 100644
--- /dev/null
+++ b/Engine/Source/Runtime/UMG/Public/Components/HorizontalBox.h
@@ -0,0 +1,46 @@
+// Copyright 1998-2014 Epic Games, Inc. All Rights Reserved.
+
+#pragma once
+
+#include "HorizontalBox.generated.h"
+
+UCLASS(meta=(BlueprintSpawnableComponent), ClassGroup=UserInterface)
+class UMG_API UHorizontalBox : public UPanelWidget
+{
+	GENERATED_UCLASS_BODY()
+
+	/** The items placed on the canvas */
+	UPROPERTY(EditAnywhere, EditInline, Category=Slots)
+	TArray<UHorizontalBoxSlot*> Slots;
+
+	UHorizontalBoxSlot* AddSlot(UWidget* Content);
+
+	// UPanelWidget
+	virtual int32 GetChildrenCount() const OVERRIDE;
+	virtual UWidget* GetChildAt(int32 Index) const OVERRIDE;
+	virtual int32 GetChildIndex(UWidget* Content) const OVERRIDE;
+	virtual bool AddChild(UWidget* Child, FVector2D Position) OVERRIDE;
+	virtual bool RemoveChild(UWidget* Child) OVERRIDE;
+	virtual void ReplaceChildAt(int32 Index, UWidget* Child) OVERRIDE;
+	virtual void InsertChildAt(int32 Index, UWidget* Child) OVERRIDE;
+	// End UPanelWidget
+
+#if WITH_EDITOR
+	// UObject interface
+	virtual void PostEditChangeProperty(struct FPropertyChangedEvent& PropertyChangedEvent) OVERRIDE;
+	// End of UObject interface
+
+	// UWidget interface
+	virtual void ConnectEditorData() OVERRIDE;
+	// End UWidget interface
+#endif
+
+protected:
+
+	TWeakPtr<class SHorizontalBox> MyHorizontalBox;
+
+protected:
+	// UWidget interface
+	virtual TSharedRef<SWidget> RebuildWidget() OVERRIDE;
+	// End of UWidget interface
+};
diff --git a/Engine/Source/Runtime/UMG/Public/Components/HorizontalBoxComponent.h b/Engine/Source/Runtime/UMG/Public/Components/HorizontalBoxComponent.h
deleted file mode 100644
diff --git a/Engine/Source/Runtime/UMG/Public/Components/Image.h b/Engine/Source/Runtime/UMG/Public/Components/Image.h
new file mode 100644
--- /dev/null
+++ b/Engine/Source/Runtime/UMG/Public/Components/Image.h
@@ -0,0 +1,28 @@
+// Copyright 1998-2014 Epic Games, Inc. All Rights Reserved.
+
+#pragma once
+
+#include "Image.generated.h"
+
+/** Image widget */
+UCLASS(meta=(BlueprintSpawnableComponent), ClassGroup=UserInterface)
+class UMG_API UImage : public UWidget
+{
+	GENERATED_UCLASS_BODY()
+
+	/** Image to draw */
+	UPROPERTY(EditAnywhere, BlueprintReadWrite, Category=Appearance, meta=( DisplayThumbnail = "true" ))
+	USlateBrushAsset* Image;
+
+	/** Color and opacity */
+	UPROPERTY(EditAnywhere, BlueprintReadWrite, Category=Appearance)
+	FSlateColor ColorAndOpacity;
+
+protected:
+	// UWidget interface
+	virtual TSharedRef<SWidget> RebuildWidget() OVERRIDE;
+	// End of UWidget interface
+
+	const FSlateBrush* GetImageBrush() const;
+	FSlateColor GetColorAndOpacity() const;
+};
diff --git a/Engine/Source/Runtime/UMG/Public/Components/ImageComponent.h b/Engine/Source/Runtime/UMG/Public/Components/ImageComponent.h
deleted file mode 100644
diff --git a/Engine/Source/Runtime/UMG/Public/Components/ScrollBox.h b/Engine/Source/Runtime/UMG/Public/Components/ScrollBox.h
new file mode 100644
--- /dev/null
+++ b/Engine/Source/Runtime/UMG/Public/Components/ScrollBox.h
@@ -0,0 +1,18 @@
+// Copyright 1998-2014 Epic Games, Inc. All Rights Reserved.
+
+#pragma once
+
+#include "ScrollBox.generated.h"
+
+UCLASS(meta=(BlueprintSpawnableComponent), ClassGroup=UserInterface)
+class UMG_API UScrollBox : public UWidget
+{
+	GENERATED_UCLASS_BODY()
+
+protected:
+
+protected:
+	// UWidget interface
+	virtual TSharedRef<SWidget> RebuildWidget() OVERRIDE;
+	// End of UWidget interface
+};
diff --git a/Engine/Source/Runtime/UMG/Public/Components/ScrollBoxComponent.h b/Engine/Source/Runtime/UMG/Public/Components/ScrollBoxComponent.h
deleted file mode 100644
diff --git a/Engine/Source/Runtime/UMG/Public/Components/Spacer.h b/Engine/Source/Runtime/UMG/Public/Components/Spacer.h
new file mode 100644
--- /dev/null
+++ b/Engine/Source/Runtime/UMG/Public/Components/Spacer.h
@@ -0,0 +1,23 @@
+// Copyright 1998-2014 Epic Games, Inc. All Rights Reserved.
+
+#pragma once
+
+#include "Spacer.generated.h"
+
+/** A spacer widget; it does not have a visual representation, and just provides padding between other widgets. */
+UCLASS(meta=(BlueprintSpawnableComponent), ClassGroup=UserInterface)
+class UMG_API USpacer : public UWidget
+{
+	GENERATED_UCLASS_BODY()
+
+	/** The size of the spacer */
+	UPROPERTY(EditAnywhere, BlueprintReadWrite, Category=Appearance)
+	FVector2D Size;
+
+protected:
+	// UWidget interface
+	virtual TSharedRef<SWidget> RebuildWidget() OVERRIDE;
+	// End of UWidget interface
+
+	FVector2D GetSpacerSize() const;
+};
diff --git a/Engine/Source/Runtime/UMG/Public/Components/SpacerComponent.h b/Engine/Source/Runtime/UMG/Public/Components/SpacerComponent.h
deleted file mode 100644
diff --git a/Engine/Source/Runtime/UMG/Public/Components/SpinningImageComponent.h b/Engine/Source/Runtime/UMG/Public/Components/SpinningImageComponent.h
--- a/Engine/Source/Runtime/UMG/Public/Components/SpinningImageComponent.h
+++ b/Engine/Source/Runtime/UMG/Public/Components/SpinningImageComponent.h
@@ -6,7 +6,7 @@
 
 /** A widget that displays a spinning image. */
 UCLASS(meta=(BlueprintSpawnableComponent), ClassGroup=UserInterface)
-class UMG_API USpinningImageComponent : public UImageComponent
+class UMG_API USpinningImageComponent : public UImage
 {
 	GENERATED_UCLASS_BODY()
 
diff --git a/Engine/Source/Runtime/UMG/Public/Components/TextBlock.h b/Engine/Source/Runtime/UMG/Public/Components/TextBlock.h
new file mode 100644
--- /dev/null
+++ b/Engine/Source/Runtime/UMG/Public/Components/TextBlock.h
@@ -0,0 +1,48 @@
+// Copyright 1998-2014 Epic Games, Inc. All Rights Reserved.
+
+#pragma once
+
+#include "TextBlock.generated.h"
+
+UCLASS(meta=(BlueprintSpawnableComponent), ClassGroup=UserInterface)
+class UMG_API UTextBlock : public UWidget
+{
+	GENERATED_UCLASS_BODY()
+
+public:
+	UPROPERTY(EditDefaultsOnly, Category=Content)
+	FText Text;
+
+	UPROPERTY()
+	FGetText TextDelegate;
+	
+	UPROPERTY(EditAnywhere, Category=Appearance)
+	FSlateFontInfo Font;
+
+	UPROPERTY(EditAnywhere, Category=Appearance)
+	FVector2D ShadowOffset;
+
+	UPROPERTY(EditAnywhere, Category=Appearance)
+	FLinearColor ColorAndOpacity;
+
+	UPROPERTY()
+	FGetSlateColor ColorAndOpacityDelegate;
+
+	UPROPERTY(EditAnywhere, Category=Appearance)
+	FLinearColor ShadowColorAndOpacity;
+
+	UPROPERTY()
+	FGetLinearColor ShadowColorAndOpacityDelegate;
+
+	/** Text to display as the content of the button */
+	UPROPERTY(EditAnywhere, BlueprintReadOnly, Category=Style, meta=( DisplayThumbnail = "true" ))
+	USlateWidgetStyleAsset* Style;
+
+protected:
+	// UWidget interface
+	virtual TSharedRef<SWidget> RebuildWidget() OVERRIDE;
+	// End of UWidget interface
+
+	FSlateColor GetColorAndOpacity() const;
+	FLinearColor GetShadowColorAndOpacity() const;
+};
diff --git a/Engine/Source/Runtime/UMG/Public/Components/TextBlockComponent.h b/Engine/Source/Runtime/UMG/Public/Components/TextBlockComponent.h
deleted file mode 100644
diff --git a/Engine/Source/Runtime/UMG/Public/Components/VerticalBox.h b/Engine/Source/Runtime/UMG/Public/Components/VerticalBox.h
new file mode 100644
--- /dev/null
+++ b/Engine/Source/Runtime/UMG/Public/Components/VerticalBox.h
@@ -0,0 +1,46 @@
+// Copyright 1998-2014 Epic Games, Inc. All Rights Reserved.
+
+#pragma once
+
+#include "VerticalBox.generated.h"
+
+UCLASS(meta=(BlueprintSpawnableComponent), ClassGroup=UserInterface)
+class UMG_API UVerticalBox : public UPanelWidget
+{
+	GENERATED_UCLASS_BODY()
+
+	/** The items placed on the canvas */
+	UPROPERTY(EditAnywhere, EditInline, Category=Slots)
+	TArray<UVerticalBoxSlot*> Slots;
+
+	UVerticalBoxSlot* AddSlot(UWidget* Content);
+
+	// UPanelWidget
+	virtual int32 GetChildrenCount() const OVERRIDE;
+	virtual UWidget* GetChildAt(int32 Index) const OVERRIDE;
+	virtual int32 GetChildIndex(UWidget* Content) const OVERRIDE;
+	virtual bool AddChild(UWidget* Child, FVector2D Position) OVERRIDE;
+	virtual bool RemoveChild(UWidget* Child) OVERRIDE;
+	virtual void ReplaceChildAt(int32 Index, UWidget* Child) OVERRIDE;
+	virtual void InsertChildAt(int32 Index, UWidget* Child) OVERRIDE;
+	// End UPanelWidget
+
+#if WITH_EDITOR
+	// UObject interface
+	virtual void PostEditChangeProperty(struct FPropertyChangedEvent& PropertyChangedEvent) OVERRIDE;
+	// End of UObject interface
+
+	// UWidget interface
+	virtual void ConnectEditorData() OVERRIDE;
+	// End UWidget interface
+#endif
+
+protected:
+
+	TWeakPtr<class SVerticalBox> MyVerticalBox;
+
+protected:
+	// UWidget interface
+	virtual TSharedRef<SWidget> RebuildWidget() OVERRIDE;
+	// End of UWidget interface
+};
diff --git a/Engine/Source/Runtime/UMG/Public/Components/VerticalBoxComponent.h b/Engine/Source/Runtime/UMG/Public/Components/VerticalBoxComponent.h
deleted file mode 100644
diff --git a/Engine/Source/Runtime/UMG/Public/Components/Widget.h b/Engine/Source/Runtime/UMG/Public/Components/Widget.h
--- a/Engine/Source/Runtime/UMG/Public/Components/Widget.h
+++ b/Engine/Source/Runtime/UMG/Public/Components/Widget.h
@@ -3,37 +3,59 @@
 #pragma once
 
 #include "SlateWrapperTypes.h"
+
 #include "Widget.generated.h"
 
+/**
+ * Helper macro for binding to a delegate or using the constant value when constructing the underlying SWidget
+ */
 #define OPTIONAL_BINDING(ReturnType, MemberName) MemberName ## Delegate.IsBound() ? TAttribute< ReturnType >::Create(MemberName ## Delegate.GetUObject(), MemberName ## Delegate.GetFunctionName()) : TAttribute< ReturnType >(MemberName)
+
+/**
+ * Helper macro for binding to a delegate or using the constant value when constructing the underlying SWidget,
+ * also allows a conversion function to be provided to convert between the SWidget value and the value exposed to UMG.
+ */
 #define OPTIONAL_BINDING_CONVERT(ReturnType, MemberName, ConvertedType, ConversionFunction) MemberName ## Delegate.IsBound() ? TAttribute< ConvertedType >::Create(TAttribute< ConvertedType >::FGetter::CreateUObject(this, &ThisClass::ConversionFunction, TAttribute< ReturnType >::Create(MemberName ## Delegate.GetUObject(), MemberName ## Delegate.GetFunctionName()))) : ConversionFunction(TAttribute< ReturnType >(MemberName))
 
-/** This is the base class for all Slate widget wrapper components */
-UCLASS(Abstract, hideCategories=(Activation, "Components|Activation", ComponentReplication, LOD, Rendering))
+/**
+ * This is the base class for all wrapped Slate controls that are exposed to UMG.
+ */
+UCLASS(Abstract)
 class UMG_API UWidget : public UObject
 {
 	GENERATED_UCLASS_BODY()
 
+public:
+
 	// Common Bindings
 	DECLARE_DYNAMIC_DELEGATE_RetVal(bool, FGetBool);
 	DECLARE_DYNAMIC_DELEGATE_RetVal(FText, FGetText);
 	DECLARE_DYNAMIC_DELEGATE_RetVal(FLinearColor, FGetSlateColor);
 	DECLARE_DYNAMIC_DELEGATE_RetVal(FLinearColor, FGetLinearColor);
 	DECLARE_DYNAMIC_DELEGATE_RetVal(ESlateVisibility::Type, FGetSlateVisibility);
+	DECLARE_DYNAMIC_DELEGATE_RetVal(EMouseCursor::Type, FGetMouseCursor);
 
 	DECLARE_DYNAMIC_DELEGATE_RetVal_OneParam(UWidget*, FGenerateWidgetUObject, UObject*, Item);
 
 	// Events
 	DECLARE_DYNAMIC_DELEGATE_RetVal(FSReply, FOnReply);
 
-	UPROPERTY(EditAnywhere, Category=Variable)
-	FString Name;
-
-	UPROPERTY(EditAnywhere, Category=Variable)
+	/**
+	 * Allows controls to be exposed as variables in a blueprint.  Not all controls need to be exposed
+	 * as variables, so this allows only the most useful ones to end up being exposed.
+	 */
+	UPROPERTY(EditDefaultsOnly, Category=Variable)
 	bool bIsVariable;
 
+	/** Flag if the Widget was created from a blueprint */
+	UPROPERTY(Transient)
+	bool bCreatedByConstructionScript;
+
 #if WITH_EDITOR
-	UPROPERTY(Transient, EditInline, EditAnywhere, BlueprintReadWrite, Category=Layout, meta=( ShowOnlyInnerProperties ))
+	/**
+	 * The parent slot of the UWidget.  Allows us to easily inline edit the layout controlling this widget.
+	 */
+	UPROPERTY(Transient, EditInline, EditDefaultsOnly, BlueprintReadWrite, Category=Layout, meta=( ShowOnlyInnerProperties ))
 	UPanelSlot* Slot;
 #endif
 
@@ -41,28 +63,37 @@ class UMG_API UWidget : public UObject
 	UPROPERTY(EditDefaultsOnly, Category=Behavior)
 	bool bIsEnabled;
 
+	/** A bindable delegate for bIsEnabled */
 	UPROPERTY()
 	FGetBool bIsEnabledDelegate;
 
 	//TODO UMG ToolTipWidget
 
+	/** Tooltip text to show when the user hovers over the widget with the mouse */
 	UPROPERTY(EditDefaultsOnly, Category=Behavior)
 	FText ToolTipText;
 
+	/** A bindable delegate for ToolTipText */
 	UPROPERTY()
 	FGetText ToolTipTextDelegate;
 
+	/** The visibility of the widget */
 	UPROPERTY(EditDefaultsOnly, Category=Behavior)
 	TEnumAsByte<ESlateVisibility::Type> Visiblity;
 
+	/** A bindable delegate for Visibility */
 	UPROPERTY()
 	FGetSlateVisibility VisiblityDelegate;
 
+	//TODO UMG Cursor doesn't work yet, the underlying slate version needs it to be TOptional.
+
+	/** The cursor to show when the mouse is over the widget */
 	UPROPERTY(EditDefaultsOnly, Category=Behavior)
 	TEnumAsByte<EMouseCursor::Type> Cursor;
 
-	UPROPERTY(Transient)
-	bool bCreatedByConstructionScript;
+	/** A bindable delegate for Cursor */
+	UPROPERTY()
+	FGetMouseCursor CursorDelegate;
 
 	UFUNCTION(BlueprintCallable, Category="Widget")
 	bool GetIsEnabled() const;
@@ -79,21 +110,11 @@ class UMG_API UWidget : public UObject
 	UFUNCTION(BlueprintCallable, Category="Widget")
 	void SetVisibility(TEnumAsByte<ESlateVisibility::Type> InVisibility);
 
-	//TOptional< TEnumAsByte<EMouseCursor::Type> > GetCursor() const;
-
-	//void SetCursor(const TAttribute< TOptional< TEnumAsByte<EMouseCursor::Type> > >& InCursor);
-
-	// UActorComponent interface
-	virtual void OnUnregister();
-	// End of UActorComponent interface
-
 	TSharedRef<SWidget> GetWidget();
 
 #if WITH_EDITOR
 	virtual void ConnectEditorData() { }
-#endif
 
-#if WITH_EDITOR
 	// UObject interface
 	virtual void PostEditChangeProperty(struct FPropertyChangedEvent& PropertyChangedEvent) OVERRIDE;
 	// End of UObject interface
diff --git a/Engine/Source/Runtime/UMG/Public/IUMGModule.h b/Engine/Source/Runtime/UMG/Public/IUMGModule.h
--- a/Engine/Source/Runtime/UMG/Public/IUMGModule.h
+++ b/Engine/Source/Runtime/UMG/Public/IUMGModule.h
@@ -4,6 +4,9 @@
 
 #include "ModuleInterface.h"
 
+/**
+ * Interface for the UMG Module.
+ */
 class IUMGModule : public IModuleInterface
 {
 public:
diff --git a/Engine/Source/Runtime/UMG/Public/Slate/SObjectWidget.h b/Engine/Source/Runtime/UMG/Public/Slate/SObjectWidget.h
--- a/Engine/Source/Runtime/UMG/Public/Slate/SObjectWidget.h
+++ b/Engine/Source/Runtime/UMG/Public/Slate/SObjectWidget.h
@@ -2,6 +2,12 @@
 
 #pragma once
 
+/**
+ * The SObjectWidet allows UMG to insert an SWidget into the hierarchy that manages the lifetime of the
+ * UMG UWidget that created it.  Once the SObjectWidget is destroyed it frees the reference it holds to
+ * The UWidget allowing it to be garbage collected.  It also forwards the slate events to the UUserWidget
+ * so that it can forward them to listeners.
+ */
 class SObjectWidget : public SCompoundWidget, public FGCObject
 {
 	SLATE_BEGIN_ARGS(SObjectWidget)
@@ -51,5 +57,6 @@ class SObjectWidget : public SCompoundWidget, public FGCObject
 	virtual FReply OnMotionDetected(const FGeometry& MyGeometry, const FMotionEvent& InMotionEvent) OVERRIDE;
 
 private:
+	/** The UWidget that created this SObjectWidget who needs to be kept alive. */
 	class UUserWidget* WidgetObject;
 };
\ No newline at end of file
diff --git a/Engine/Source/Runtime/UMG/Public/UMG.h b/Engine/Source/Runtime/UMG/Public/UMG.h
--- a/Engine/Source/Runtime/UMG/Public/UMG.h
+++ b/Engine/Source/Runtime/UMG/Public/UMG.h
@@ -33,13 +33,13 @@
 #include "BoxPanelComponent.h"
 
 #include "CanvasPanelSlot.h"
-#include "CanvasPanelComponent.h"
+#include "CanvasPanel.h"
 
 #include "HorizontalBoxSlot.h"
-#include "HorizontalBoxComponent.h"
+#include "HorizontalBox.h"
 
 #include "VerticalBoxSlot.h"
-#include "VerticalBoxComponent.h"
+#include "VerticalBox.h"
 
 #include "UniformGridSlot.h"
 #include "UniformGridPanel.h"
@@ -47,19 +47,20 @@
 #include "OverlaySlot.h"
 #include "Overlay.h"
 
-#include "TextBlockComponent.h"
-#include "BorderComponent.h"
-#include "ButtonComponent.h"
-#include "CheckBoxComponent.h"
+#include "TextBlock.h"
+#include "Border.h"
+#include "Button.h"
+#include "CheckBox.h"
 #include "EditableText.h"
-#include "EditableTextBlockComponent.h"
-#include "ImageComponent.h"
+#include "EditableTextBlock.h"
+#include "Image.h"
 #include "Slider.h"
 #include "ComboBox.h"
 
 #include "GridPanelComponent.h"
-#include "ScrollBoxComponent.h"
-#include "SpacerComponent.h"
+#include "ScrollBox.h"
+#include "Spacer.h"
+
 #include "SpinningImageComponent.h"
 #include "ThrobberComponent.h"
 #include "CircularThrobberComponent.h"
