UMG - Adding a progress bar.  Adding a functional scrollbox. Adding documentation.

diff --git a/Engine/Source/Runtime/UMG/Private/Components/ListView.cpp b/Engine/Source/Runtime/UMG/Private/Components/ListView.cpp
--- a/Engine/Source/Runtime/UMG/Private/Components/ListView.cpp
+++ b/Engine/Source/Runtime/UMG/Private/Components/ListView.cpp
@@ -18,7 +18,7 @@ UListView::UListView(const FPostConstructInitializeProperties& PCIP)
 
 TSharedRef<SWidget> UListView::RebuildWidget()
 {
-	return SNew(SListView< UObject* >)
+	MyListView = SNew(SListView< UObject* >)
 		.SelectionMode(SelectionMode)
 		.ListItemsSource(&Items)
 		.ItemHeight(ItemHeight)
@@ -33,6 +33,8 @@ TSharedRef<SWidget> UListView::RebuildWidget()
 		//		+ SHeaderRow::Column(TEXT("Socket"))
 		//	);
 		;
+
+	return MyListView.ToSharedRef();
 }
 
 TSharedRef<ITableRow> UListView::HandleOnGenerateRow(UObject* Item, const TSharedRef< STableViewBase >& OwnerTable) const
diff --git a/Engine/Source/Runtime/UMG/Private/Components/ProcessBar.cpp b/Engine/Source/Runtime/UMG/Private/Components/ProcessBar.cpp
new file mode 100644
--- /dev/null
+++ b/Engine/Source/Runtime/UMG/Private/Components/ProcessBar.cpp
@@ -0,0 +1,42 @@
+// Copyright 1998-2014 Epic Games, Inc. All Rights Reserved.
+
+#include "UMGPrivatePCH.h"
+
+#define LOCTEXT_NAMESPACE "UMG"
+
+/////////////////////////////////////////////////////
+// UProgressBar
+
+UProgressBar::UProgressBar(const FPostConstructInitializeProperties& PCIP)
+	: Super(PCIP)
+{
+	BarFillType = EProgressBarFillType::LeftToRight;
+	Percent = 0;
+	FillColorAndOpacity = FLinearColor::White;
+	BorderPadding = FVector2D(1, 0);
+}
+
+TSharedRef<SWidget> UProgressBar::RebuildWidget()
+{
+	MyProgressBar = SNew(SProgressBar)
+		.BarFillType(BarFillType)
+		.Percent(Percent)
+		.FillColorAndOpacity(FillColorAndOpacity)
+		.BorderPadding(BorderPadding);
+
+	return MyProgressBar.ToSharedRef();
+}
+
+float UProgressBar::GetPercent()
+{
+	return MyProgressBar->GetPercent().Get(0.0f);
+}
+
+void UProgressBar::SetPercent(float InPercent)
+{
+	return MyProgressBar->SetPercent(InPercent);
+}
+
+/////////////////////////////////////////////////////
+
+#undef LOCTEXT_NAMESPACE
diff --git a/Engine/Source/Runtime/UMG/Private/Components/ScrollBox.cpp b/Engine/Source/Runtime/UMG/Private/Components/ScrollBox.cpp
--- a/Engine/Source/Runtime/UMG/Private/Components/ScrollBox.cpp
+++ b/Engine/Source/Runtime/UMG/Private/Components/ScrollBox.cpp
@@ -13,10 +13,144 @@ UScrollBox::UScrollBox(const FPostConstructInitializeProperties& PCIP)
 	bIsVariable = false;
 }
 
+int32 UScrollBox::GetChildrenCount() const
+{
+	return Slots.Num();
+}
+
+UWidget* UScrollBox::GetChildAt(int32 Index) const
+{
+	return Slots[Index]->Content;
+}
+
+int32 UScrollBox::GetChildIndex(UWidget* Content) const
+{
+	for ( int32 SlotIndex = 0; SlotIndex < Slots.Num(); ++SlotIndex )
+	{
+		UScrollBoxSlot* Slot = Slots[SlotIndex];
+
+		if ( Slot->Content == Content )
+		{
+			return SlotIndex;
+		}
+	}
+
+	return -1;
+}
+
+bool UScrollBox::AddChild(UWidget* Child, FVector2D Position)
+{
+	AddSlot(Child);
+	return true;
+}
+
+bool UScrollBox::RemoveChild(UWidget* Child)
+{
+	int32 SlotIndex = GetChildIndex(Child);
+	if ( SlotIndex != -1 )
+	{
+		Slots.RemoveAt(SlotIndex);
+		return true;
+	}
+
+	return false;
+}
+
+void UScrollBox::ReplaceChildAt(int32 Index, UWidget* Content)
+{
+	UScrollBoxSlot* Slot = Slots[Index];
+	Slot->Content = Content;
+
+#if WITH_EDITOR
+	Content->Slot = Slot;
+#endif
+}
+
+void UScrollBox::InsertChildAt(int32 Index, UWidget* Content)
+{
+	UScrollBoxSlot* Slot = ConstructObject<UScrollBoxSlot>(UScrollBoxSlot::StaticClass(), this);
+	Slot->Content = Content;
+	Slot->Parent = this;
+
+#if WITH_EDITOR
+	Content->Slot = Slot;
+#endif
+
+	Slots.Insert(Slot, Index);
+}
+
 TSharedRef<SWidget> UScrollBox::RebuildWidget()
 {
-	return SNew(SScrollBox);
+	TSharedRef<SScrollBox> NewCanvas = SNew(SScrollBox);
+	MyScrollBox = NewCanvas;
+
+	TSharedPtr<SScrollBox> Canvas = MyScrollBox.Pin();
+	if ( Canvas.IsValid() )
+	{
+		Canvas->ClearChildren();
+
+		for ( int32 SlotIndex = 0; SlotIndex < Slots.Num(); ++SlotIndex )
+		{
+			UScrollBoxSlot* Slot = Slots[SlotIndex];
+			if ( Slot == NULL )
+			{
+				Slots[SlotIndex] = Slot = ConstructObject<UScrollBoxSlot>(UScrollBoxSlot::StaticClass(), this);
+			}
+
+			Slot->Parent = this;
+			Slot->BuildSlot(NewCanvas);
+		}
+	}
+
+	return NewCanvas;
+}
+
+UScrollBoxSlot* UScrollBox::AddSlot(UWidget* Content)
+{
+	UScrollBoxSlot* Slot = ConstructObject<UScrollBoxSlot>(UScrollBoxSlot::StaticClass(), this);
+	Slot->Content = Content;
+	Slot->Parent = this;
+
+#if WITH_EDITOR
+	Content->Slot = Slot;
+#endif
+	
+	Slots.Add(Slot);
+
+	return Slot;
+}
+
+#if WITH_EDITOR
+void UScrollBox::ConnectEditorData()
+{
+	for ( UScrollBoxSlot* Slot : Slots )
+	{
+		// TODO UMG Find a better way to assign the parent pointers for slots.
+		Slot->Parent = this;
+
+		if ( Slot->Content )
+		{
+			Slot->Content->Slot = Slot;
+		}
+		//TODO UMG Should we auto delete empty slots?
+	}
+}
+
+void UScrollBox::PostEditChangeProperty(struct FPropertyChangedEvent& PropertyChangedEvent)
+{
+	// Ensure the slots have unique names
+	int32 SlotNumbering = 1;
+
+	TSet<FName> UniqueSlotNames;
+	for (int32 SlotIndex = 0; SlotIndex < Slots.Num(); ++SlotIndex)
+	{
+		if ( Slots[SlotIndex] == NULL )
+		{
+			Slots[SlotIndex] = ConstructObject<UScrollBoxSlot>(UScrollBoxSlot::StaticClass(), this);
+		}
+	}
 }
+#endif
 
 /////////////////////////////////////////////////////
 
diff --git a/Engine/Source/Runtime/UMG/Private/Components/ScrollBoxSlot.cpp b/Engine/Source/Runtime/UMG/Private/Components/ScrollBoxSlot.cpp
new file mode 100644
--- /dev/null
+++ b/Engine/Source/Runtime/UMG/Private/Components/ScrollBoxSlot.cpp
@@ -0,0 +1,23 @@
+// Copyright 1998-2014 Epic Games, Inc. All Rights Reserved.
+
+#include "UMGPrivatePCH.h"
+
+/////////////////////////////////////////////////////
+// UScrollBoxSlot
+
+UScrollBoxSlot::UScrollBoxSlot(const FPostConstructInitializeProperties& PCIP)
+	: Super(PCIP)
+	, Slot(NULL)
+{
+	HorizontalAlignment = HAlign_Left;
+}
+
+void UScrollBoxSlot::BuildSlot(TSharedRef<SScrollBox> ScrollBox)
+{
+	Slot = &ScrollBox->AddSlot()
+		.Padding(Padding)
+		.HAlign(HorizontalAlignment)
+		[
+			Content == NULL ? SNullWidget::NullWidget : Content->GetWidget()
+		];
+}
diff --git a/Engine/Source/Runtime/UMG/Private/Components/Slider.cpp b/Engine/Source/Runtime/UMG/Private/Components/Slider.cpp
--- a/Engine/Source/Runtime/UMG/Private/Components/Slider.cpp
+++ b/Engine/Source/Runtime/UMG/Private/Components/Slider.cpp
@@ -21,7 +21,7 @@ TSharedRef<SWidget> USlider::RebuildWidget()
 	//FOnMouseCaptureBeginEvent OnMouseCaptureBegin;
 	//FOnMouseCaptureEndEvent OnMouseCaptureEnd;
 
-	TSharedRef<SSlider> NewWidget = SNew(SSlider)
+	MySlider = SNew(SSlider)
 		.Orientation(Orientation)
 		.SliderBarColor(SliderBarColor)
 		.SliderHandleColor(SliderHandleColor)
@@ -30,7 +30,7 @@ TSharedRef<SWidget> USlider::RebuildWidget()
 		.OnMouseCaptureEnd(BIND_UOBJECT_DELEGATE(FSimpleDelegate, HandleOnMouseCaptureEnd))
 		.OnValueChanged(BIND_UOBJECT_DELEGATE(FOnFloatValueChanged, HandleOnValueChanged));
 
-	return NewWidget;
+	return MySlider.ToSharedRef();
 }
 
 void USlider::HandleOnValueChanged(float InValue)
@@ -50,12 +50,12 @@ void USlider::HandleOnMouseCaptureEnd()
 
 float USlider::GetValue()
 {
-	return SliderWidget()->GetValue();
+	return MySlider->GetValue();
 }
 
 void USlider::SetValue(float InValue)
 {
-	return SliderWidget()->SetValue(InValue);
+	return MySlider->SetValue(InValue);
 }
 
 /////////////////////////////////////////////////////
diff --git a/Engine/Source/Runtime/UMG/Private/Components/Spacer.cpp b/Engine/Source/Runtime/UMG/Private/Components/Spacer.cpp
--- a/Engine/Source/Runtime/UMG/Private/Components/Spacer.cpp
+++ b/Engine/Source/Runtime/UMG/Private/Components/Spacer.cpp
@@ -17,12 +17,7 @@ USpacer::USpacer(const FPostConstructInitializeProperties& PCIP)
 TSharedRef<SWidget> USpacer::RebuildWidget()
 {
 	return SNew(SSpacer)
-		.Size( BIND_UOBJECT_ATTRIBUTE(FVector2D, GetSpacerSize) );
-}
-
-FVector2D USpacer::GetSpacerSize() const
-{
-	return Size;
+		.Size(Size);
 }
 
 /////////////////////////////////////////////////////
diff --git a/Engine/Source/Runtime/UMG/Private/Components/UniformGridPanel.cpp b/Engine/Source/Runtime/UMG/Private/Components/UniformGridPanel.cpp
--- a/Engine/Source/Runtime/UMG/Private/Components/UniformGridPanel.cpp
+++ b/Engine/Source/Runtime/UMG/Private/Components/UniformGridPanel.cpp
@@ -55,7 +55,12 @@ void UUniformGridPanel::ReplaceChildAt(int32 Index, UWidget* Content)
 
 TSharedRef<SWidget> UUniformGridPanel::RebuildWidget()
 {
-	TSharedRef<SUniformGridPanel> NewPanel = SNew(SUniformGridPanel);
+	TSharedRef<SUniformGridPanel> NewPanel =
+		SNew(SUniformGridPanel)
+		.SlotPadding(SlotPadding)
+		.MinDesiredSlotWidth(MinDesiredSlotWidth)
+		.MinDesiredSlotHeight(MinDesiredSlotHeight);
+	
 	MyUniformGridPanel = NewPanel;
 
 	for ( int32 SlotIndex = 0; SlotIndex < Slots.Num(); ++SlotIndex )
diff --git a/Engine/Source/Runtime/UMG/Public/Blueprint/UserWidget.h b/Engine/Source/Runtime/UMG/Public/Blueprint/UserWidget.h
--- a/Engine/Source/Runtime/UMG/Public/Blueprint/UserWidget.h
+++ b/Engine/Source/Runtime/UMG/Public/Blueprint/UserWidget.h
@@ -72,7 +72,7 @@ public:
 DECLARE_DYNAMIC_MULTICAST_DELEGATE_OneParam(FOnVisibilityChangedEvent, ESlateVisibility::Type, Visibility);
 
 /**
- * 
+ * The user widget is extensible by users through the WidgetBlueprint.
  */
 UCLASS(Abstract, editinlinenew, BlueprintType, Blueprintable)
 class UMG_API UUserWidget : public UWidget
diff --git a/Engine/Source/Runtime/UMG/Public/Blueprint/WidgetActor.h b/Engine/Source/Runtime/UMG/Public/Blueprint/WidgetActor.h
--- a/Engine/Source/Runtime/UMG/Public/Blueprint/WidgetActor.h
+++ b/Engine/Source/Runtime/UMG/Public/Blueprint/WidgetActor.h
@@ -4,6 +4,7 @@
 
 #include "WidgetActor.generated.h"
 
+/** A widget actor can be added to the world, giving a widget 3D world position. */
 UCLASS()
 class UMG_API AWidgetActor : public AActor
 {
diff --git a/Engine/Source/Runtime/UMG/Public/Blueprint/WidgetBlueprint.h b/Engine/Source/Runtime/UMG/Public/Blueprint/WidgetBlueprint.h
--- a/Engine/Source/Runtime/UMG/Public/Blueprint/WidgetBlueprint.h
+++ b/Engine/Source/Runtime/UMG/Public/Blueprint/WidgetBlueprint.h
@@ -33,6 +33,9 @@ struct UMG_API FDelegateEditorBinding
 	FDelegateRuntimeBinding ToRuntimeBinding(class UWidgetBlueprint* Blueprint) const;
 };
 
+/**
+ * The widget blueprint enables extending UUserWidget the user extensible UWidget.
+ */
 UCLASS(dependson=(UBlueprint, UUserWidget, UWidgetBlueprintGeneratedClass), BlueprintType)
 class UMG_API UWidgetBlueprint : public UBlueprint
 {
diff --git a/Engine/Source/Runtime/UMG/Public/Blueprint/WidgetBlueprintLibrary.h b/Engine/Source/Runtime/UMG/Public/Blueprint/WidgetBlueprintLibrary.h
--- a/Engine/Source/Runtime/UMG/Public/Blueprint/WidgetBlueprintLibrary.h
+++ b/Engine/Source/Runtime/UMG/Public/Blueprint/WidgetBlueprintLibrary.h
@@ -12,7 +12,7 @@ class UWidgetBlueprintLibrary : public UBlueprintFunctionLibrary
 	UFUNCTION(BlueprintCallable, meta=( HidePin="WorldContextObject", DefaultToSelf="WorldContextObject", UnsafeDuringActorConstruction = "true", FriendlyName = "Create Widget" ), Category="User Interface|Widget")
 	static class UUserWidget* Create(UObject* WorldContextObject, TSubclassOf<class UUserWidget> WidgetType);
 
-	/** Paint Functions */
+	/** Draws a box */
 	UFUNCTION(BlueprintCallable, Category="Painting")
 	static void DrawBox(UPARAM(ref) FPaintContext& Context, FVector2D Position, FVector2D Size, USlateBrushAsset* Brush, FLinearColor Tint = FLinearColor::White);
 
@@ -29,6 +29,13 @@ class UWidgetBlueprintLibrary : public UBlueprintFunctionLibrary
 
 	// TODO UMG DrawLines
 
+	/** 
+	 * Draws text.
+	 *
+	 * @param InString		The string to draw.
+	 * @param Position		The starting position where the text is drawn in local space.
+	 * @param Tint			Color to render the line.
+	 */
 	UFUNCTION(BlueprintCallable, Category="Painting")
 	static void DrawText(UPARAM(ref) FPaintContext& Context, const FString& InString, FVector2D Position, FLinearColor Tint = FLinearColor::White);
 
diff --git a/Engine/Source/Runtime/UMG/Public/Blueprint/WidgetTree.h b/Engine/Source/Runtime/UMG/Public/Blueprint/WidgetTree.h
--- a/Engine/Source/Runtime/UMG/Public/Blueprint/WidgetTree.h
+++ b/Engine/Source/Runtime/UMG/Public/Blueprint/WidgetTree.h
@@ -4,6 +4,7 @@
 
 #include "WidgetTree.generated.h"
 
+/** The widget tree manages the collection of widgets in a blueprint widget. */
 UCLASS()
 class UMG_API UWidgetTree : public UObject
 {
@@ -21,6 +22,7 @@ class UMG_API UWidgetTree : public UObject
 
 	class UPanelWidget* FindWidgetParent(class UWidget* Widget, int32& OutChildIndex);
 
+	/** Constructs the widget, and adds it to the tree. */
 	template< class T >
 	T* ConstructWidget(TSubclassOf<class UWidget> WidgetType)
 	{
diff --git a/Engine/Source/Runtime/UMG/Public/Components/Border.h b/Engine/Source/Runtime/UMG/Public/Components/Border.h
--- a/Engine/Source/Runtime/UMG/Public/Components/Border.h
+++ b/Engine/Source/Runtime/UMG/Public/Components/Border.h
@@ -4,8 +4,10 @@
 
 #include "Border.generated.h"
 
-/** A border is a container widget that can contain one child widget, providing an opportunity to surround it with a border image and adjustable padding. */
-
+/**
+ * A border is a container widget that can contain one child widget, providing an opportunity 
+ * to surround it with a border image and adjustable padding.
+ */
 UCLASS(meta=(BlueprintSpawnableComponent), ClassGroup=UserInterface)
 class UMG_API UBorder : public UContentWidget
 {
diff --git a/Engine/Source/Runtime/UMG/Public/Components/Image.h b/Engine/Source/Runtime/UMG/Public/Components/Image.h
--- a/Engine/Source/Runtime/UMG/Public/Components/Image.h
+++ b/Engine/Source/Runtime/UMG/Public/Components/Image.h
@@ -10,6 +10,8 @@ class UMG_API UImage : public UWidget
 {
 	GENERATED_UCLASS_BODY()
 
+public:
+
 	/** Image to draw */
 	UPROPERTY(EditAnywhere, BlueprintReadWrite, Category=Appearance, meta=( DisplayThumbnail = "true" ))
 	USlateBrushAsset* Image;
diff --git a/Engine/Source/Runtime/UMG/Public/Components/ListView.h b/Engine/Source/Runtime/UMG/Public/Components/ListView.h
--- a/Engine/Source/Runtime/UMG/Public/Components/ListView.h
+++ b/Engine/Source/Runtime/UMG/Public/Components/ListView.h
@@ -4,6 +4,7 @@
 
 #include "ListView.generated.h"
 
+/** Allows thousands of items to be displayed in a list.  Generates widgets dynamically for each item. */
 UCLASS(meta=( BlueprintSpawnableComponent ), ClassGroup=UserInterface)
 class UMG_API UListView : public UTableViewBase
 {
@@ -11,29 +12,24 @@ class UMG_API UListView : public UTableViewBase
 
 public:
 
+	/** The height of each widget */
 	UPROPERTY(EditDefaultsOnly, Category=Content)
 	float ItemHeight;
 
+	/** The list of items to generate widgets for */
 	UPROPERTY(EditAnywhere, BlueprintReadWrite, Category=Content)
 	TArray<UObject*> Items;
 
+	/** The selection method for the list */
 	UPROPERTY(EditDefaultsOnly, Category=Content)
 	TEnumAsByte<ESelectionMode::Type> SelectionMode;
 
-	/** Called when the button is clicked */
+	/** Called when a widget needs to be generated */
 	UPROPERTY(EditDefaultsOnly, Category=Events)
 	FOnGenerateRowUObject OnGenerateRowEvent;
 
 protected:
-	TSharedPtr< SListView<UObject*> > ListWidget() const
-	{
-		if ( MyWidget.IsValid() )
-		{
-			return StaticCastSharedRef< SListView<UObject*> >(MyWidget.ToSharedRef());
-		}
-
-		return TSharedPtr< SListView<UObject*> >();
-	}
+	TSharedPtr< SListView<UObject*> > MyListView;
 
 	TSharedRef<ITableRow> HandleOnGenerateRow(UObject* Item, const TSharedRef< STableViewBase >& OwnerTable) const;
 
diff --git a/Engine/Source/Runtime/UMG/Public/Components/OverlaySlot.h b/Engine/Source/Runtime/UMG/Public/Components/OverlaySlot.h
--- a/Engine/Source/Runtime/UMG/Public/Components/OverlaySlot.h
+++ b/Engine/Source/Runtime/UMG/Public/Components/OverlaySlot.h
@@ -6,6 +6,7 @@
 
 #include "OverlaySlot.generated.h"
 
+/** Slot for the UOverlay panel.  Allows content to be hover above other content. */
 UCLASS()
 class UMG_API UOverlaySlot : public UPanelSlot
 {
diff --git a/Engine/Source/Runtime/UMG/Public/Components/PanelSlot.h b/Engine/Source/Runtime/UMG/Public/Components/PanelSlot.h
--- a/Engine/Source/Runtime/UMG/Public/Components/PanelSlot.h
+++ b/Engine/Source/Runtime/UMG/Public/Components/PanelSlot.h
@@ -4,11 +4,14 @@
 
 #include "PanelSlot.generated.h"
 
-UCLASS()
+/** The base class for all Slots in UMG. */
+UCLASS(Abstract)
 class UMG_API UPanelSlot : public UObject
 {
 	GENERATED_UCLASS_BODY()
-	
+
+public:
+
 	UPROPERTY()
 	class UPanelWidget* Parent;
 
diff --git a/Engine/Source/Runtime/UMG/Public/Components/PanelWidget.h b/Engine/Source/Runtime/UMG/Public/Components/PanelWidget.h
--- a/Engine/Source/Runtime/UMG/Public/Components/PanelWidget.h
+++ b/Engine/Source/Runtime/UMG/Public/Components/PanelWidget.h
@@ -4,11 +4,14 @@
 
 #include "PanelWidget.generated.h"
 
+/** The base class for all UMG panel widgets.  Panel widgets layout a collection of child widgets. */
 UCLASS(Abstract)
 class UMG_API UPanelWidget : public UWidget
 {
 	GENERATED_UCLASS_BODY()
 
+public:
+
 	virtual int32 GetChildrenCount() const { return 0; }
 	virtual UWidget* GetChildAt(int32 Index) const { return NULL; }
 
diff --git a/Engine/Source/Runtime/UMG/Public/Components/ProgressBar.h b/Engine/Source/Runtime/UMG/Public/Components/ProgressBar.h
new file mode 100644
--- /dev/null
+++ b/Engine/Source/Runtime/UMG/Public/Components/ProgressBar.h
@@ -0,0 +1,49 @@
+// Copyright 1998-2014 Epic Games, Inc. All Rights Reserved.
+
+#pragma once
+
+#include "ProgressBar.generated.h"
+
+/** ProgressBar widget */
+UCLASS(meta=(BlueprintSpawnableComponent), ClassGroup=UserInterface)
+class UMG_API UProgressBar : public UWidget
+{
+	GENERATED_UCLASS_BODY()
+
+	///** Style used for the progress bar */
+	//UPROPERTY(EditDefaultsOnly, Category=Appearance)
+	//FProgressBarStyle Style;
+
+	/** Defines if this progress bar fills Left to right or right to left */
+	UPROPERTY(EditDefaultsOnly, Category=Appearance)
+	TEnumAsByte<EProgressBarFillType::Type> BarFillType;
+
+	/** Used to determine the fill position of the progress bar ranging 0..1 */
+	UPROPERTY(EditDefaultsOnly, Category=Appearance, meta=( UIMin = "0", UIMax = "1" ))
+	float Percent;
+	//TODO UMG Slate supports TOptional<float> for Percent.  UMG should as well.
+
+	/** Fill Color and Opacity */
+	UPROPERTY(EditDefaultsOnly, Category=Appearance)
+	FLinearColor FillColorAndOpacity;
+
+	/** Border Padding around fill bar */
+	UPROPERTY(EditDefaultsOnly, Category=Appearance)
+	FVector2D BorderPadding;
+
+	/** Gets the current value of the ProgressBar. */
+	UFUNCTION(BlueprintPure, Category="Behavior")
+	float GetPercent();
+
+	/** Sets the current value of the ProgressBar. */
+	UFUNCTION(BlueprintCallable, Category="Behavior")
+	void SetPercent(float InPercent);
+
+protected:
+	/** Native Slate Widget */
+	TSharedPtr<SProgressBar> MyProgressBar;
+
+	// UWidget interface
+	virtual TSharedRef<SWidget> RebuildWidget() OVERRIDE;
+	// End of UWidget interface
+};
diff --git a/Engine/Source/Runtime/UMG/Public/Components/ScrollBox.h b/Engine/Source/Runtime/UMG/Public/Components/ScrollBox.h
--- a/Engine/Source/Runtime/UMG/Public/Components/ScrollBox.h
+++ b/Engine/Source/Runtime/UMG/Public/Components/ScrollBox.h
@@ -4,13 +4,44 @@
 
 #include "ScrollBox.generated.h"
 
+/** An arbitrary scrollable collection of widgets.  Great for presenting 10-100 widgets in a list.  Doesn't support virtualization. */
 UCLASS(meta=(BlueprintSpawnableComponent), ClassGroup=UserInterface)
-class UMG_API UScrollBox : public UWidget
+class UMG_API UScrollBox : public UPanelWidget
 {
 	GENERATED_UCLASS_BODY()
 
+	/** The slots containing the widgets that are scrollable. */
+	UPROPERTY(EditDefaultsOnly, EditInline, Category=Slots)
+	TArray<UScrollBoxSlot*> Slots;
+
+	//TODO UMG Add ways to make adding slots callable by blueprints.
+
+	UScrollBoxSlot* AddSlot(UWidget* Content);
+
+	// UPanelWidget
+	virtual int32 GetChildrenCount() const OVERRIDE;
+	virtual UWidget* GetChildAt(int32 Index) const OVERRIDE;
+	virtual int32 GetChildIndex(UWidget* Content) const OVERRIDE;
+	virtual bool AddChild(UWidget* Child, FVector2D Position) OVERRIDE;
+	virtual bool RemoveChild(UWidget* Child) OVERRIDE;
+	virtual void ReplaceChildAt(int32 Index, UWidget* Child) OVERRIDE;
+	virtual void InsertChildAt(int32 Index, UWidget* Child) OVERRIDE;
+	// End UPanelWidget
+
+#if WITH_EDITOR
+	// UObject interface
+	virtual void PostEditChangeProperty(struct FPropertyChangedEvent& PropertyChangedEvent) OVERRIDE;
+	// End of UObject interface
+
+	// UWidget interface
+	virtual void ConnectEditorData() OVERRIDE;
+	// End UWidget interface
+#endif
+
 protected:
 
+	TWeakPtr<class SScrollBox> MyScrollBox;
+
 protected:
 	// UWidget interface
 	virtual TSharedRef<SWidget> RebuildWidget() OVERRIDE;
diff --git a/Engine/Source/Runtime/UMG/Public/Components/ScrollBoxSlot.h b/Engine/Source/Runtime/UMG/Public/Components/ScrollBoxSlot.h
new file mode 100644
--- /dev/null
+++ b/Engine/Source/Runtime/UMG/Public/Components/ScrollBoxSlot.h
@@ -0,0 +1,31 @@
+// Copyright 1998-2014 Epic Games, Inc. All Rights Reserved.
+
+#pragma once
+
+#include "SlateWrapperTypes.h"
+
+#include "ScrollBoxSlot.generated.h"
+
+/** The Slot for the UScrollBox, contains the widget that are scrollable */
+UCLASS()
+class UMG_API UScrollBoxSlot : public UPanelSlot
+{
+	GENERATED_UCLASS_BODY()
+	
+	/** The padding area between the slot and the content it contains. */
+	UPROPERTY(EditDefaultsOnly, Category=Layout)
+	FMargin Padding;
+
+	/** The alignment of the object horizontally. */
+	UPROPERTY(EditDefaultsOnly, Category=Layout)
+	TEnumAsByte<EHorizontalAlignment> HorizontalAlignment;
+
+	/** Builds the underlying FSlot for the Slate layout panel. */
+	void BuildSlot(TSharedRef<SScrollBox> ScrollBox);
+
+private:
+	//TODO UMG Slots should hold weak or shared refs to slots.
+
+	/** A raw pointer to the slot to allow us to adjust the size, padding...etc at runtime. */
+	SScrollBox::FSlot* Slot;
+};
diff --git a/Engine/Source/Runtime/UMG/Public/Components/SlateWrapperTypes.h b/Engine/Source/Runtime/UMG/Public/Components/SlateWrapperTypes.h
--- a/Engine/Source/Runtime/UMG/Public/Components/SlateWrapperTypes.h
+++ b/Engine/Source/Runtime/UMG/Public/Components/SlateWrapperTypes.h
@@ -30,6 +30,7 @@ namespace ESlateVisibility
 	};
 }
 
+/** The sizing options of UWidgets */
 UENUM()
 namespace ESlateSizeRule
 {
@@ -42,13 +43,14 @@ namespace ESlateSizeRule
 	};
 }
 
+/** A struct exposing FReply options in a UStruct compatible way. */
 USTRUCT()
 struct FSReply
 {
 	GENERATED_USTRUCT_BODY()
 
 	UPROPERTY(EditAnywhere, Category=User)
-		uint32 bIsHandled : 1;
+	uint32 bIsHandled : 1;
 
 	FReply ToReply() const
 	{
@@ -63,15 +65,17 @@ struct FSReply
 	}
 };
 
+/** A struct exposing size param related properties to UMG. */
 USTRUCT()
 struct FSlateChildSize
 {
 	GENERATED_USTRUCT_BODY()
 
-	//editcondition="SizeRule=ESlateSizeRule::Fill"
+	/** The parameter of the size rule. */
 	UPROPERTY(EditAnywhere, Category=Appearance, meta=(ClampMin="0", ClampMax="1"))
 	float Value;
 
+	/** The sizing rule of the content. */
 	UPROPERTY(EditAnywhere, Category=Appearance)
 	TEnumAsByte<ESlateSizeRule::Type> SizeRule;
 
diff --git a/Engine/Source/Runtime/UMG/Public/Components/Slider.h b/Engine/Source/Runtime/UMG/Public/Components/Slider.h
--- a/Engine/Source/Runtime/UMG/Public/Components/Slider.h
+++ b/Engine/Source/Runtime/UMG/Public/Components/Slider.h
@@ -51,24 +51,17 @@ class UMG_API USlider : public UWidget
 	UPROPERTY(BlueprintAssignable)
 	FOnFloatValueChangedEvent OnValueChanged;
 
-	/*  */
+	/** Gets the current value of the slider. */
 	UFUNCTION(BlueprintPure, Category="Behavior")
 	float GetValue();
 
-	/*  */
+	/** Sets the current value of the slider. */
 	UFUNCTION(BlueprintCallable, Category="Behavior")
 	void SetValue(float InValue);
 
 protected:
-	TSharedPtr<class SSlider> SliderWidget() const
-	{
-		if ( MyWidget.IsValid() )
-		{
-			return StaticCastSharedRef<SSlider>(MyWidget.ToSharedRef());
-		}
-
-		return TSharedPtr<class SSlider>();
-	}
+	/** Native Slate Widget */
+	TSharedPtr<SSlider> MySlider;
 
 	// UWidget interface
 	virtual TSharedRef<SWidget> RebuildWidget() OVERRIDE;
diff --git a/Engine/Source/Runtime/UMG/Public/Components/Spacer.h b/Engine/Source/Runtime/UMG/Public/Components/Spacer.h
--- a/Engine/Source/Runtime/UMG/Public/Components/Spacer.h
+++ b/Engine/Source/Runtime/UMG/Public/Components/Spacer.h
@@ -11,13 +11,13 @@ class UMG_API USpacer : public UWidget
 	GENERATED_UCLASS_BODY()
 
 	/** The size of the spacer */
-	UPROPERTY(EditAnywhere, BlueprintReadWrite, Category=Appearance)
+	UPROPERTY(EditDefaultsOnly, Category=Appearance)
 	FVector2D Size;
 
+	//TODO UMG UFunction SetSize
+
 protected:
 	// UWidget interface
 	virtual TSharedRef<SWidget> RebuildWidget() OVERRIDE;
 	// End of UWidget interface
-
-	FVector2D GetSpacerSize() const;
 };
diff --git a/Engine/Source/Runtime/UMG/Public/Components/SpinningImageComponent.h b/Engine/Source/Runtime/UMG/Public/Components/SpinningImageComponent.h
--- a/Engine/Source/Runtime/UMG/Public/Components/SpinningImageComponent.h
+++ b/Engine/Source/Runtime/UMG/Public/Components/SpinningImageComponent.h
@@ -11,9 +11,11 @@ class UMG_API USpinningImageComponent : public UImage
 	GENERATED_UCLASS_BODY()
 
 	/** The amount of time in seconds for a full rotation */
-	UPROPERTY(EditAnywhere, BlueprintReadOnly, Category=Appearance)
+	UPROPERTY(EditDefaultsOnly, Category=Appearance)
 	float Period;
 
+	//TODO UMG Set period dynamically.
+
 protected:
 	// UWidget interface
 	virtual TSharedRef<SWidget> RebuildWidget() OVERRIDE;
diff --git a/Engine/Source/Runtime/UMG/Public/Components/TableViewBase.h b/Engine/Source/Runtime/UMG/Public/Components/TableViewBase.h
--- a/Engine/Source/Runtime/UMG/Public/Components/TableViewBase.h
+++ b/Engine/Source/Runtime/UMG/Public/Components/TableViewBase.h
@@ -4,11 +4,12 @@
 
 #include "TableViewBase.generated.h"
 
+/** The base class for all wrapped table views */
 UCLASS(Abstract, meta=(BlueprintSpawnableComponent), ClassGroup=UserInterface)
 class UMG_API UTableViewBase : public UWidget
 {
 	GENERATED_UCLASS_BODY()
 
+	/** Delegate for constructing a UWidget based on a UObject */
 	DECLARE_DYNAMIC_DELEGATE_RetVal_OneParam(UWidget*, FOnGenerateRowUObject, UObject*, Item);
-	
 };
diff --git a/Engine/Source/Runtime/UMG/Public/Components/TextBlock.h b/Engine/Source/Runtime/UMG/Public/Components/TextBlock.h
--- a/Engine/Source/Runtime/UMG/Public/Components/TextBlock.h
+++ b/Engine/Source/Runtime/UMG/Public/Components/TextBlock.h
@@ -4,45 +4,60 @@
 
 #include "TextBlock.generated.h"
 
+/** A simple static text widget */
 UCLASS(meta=(BlueprintSpawnableComponent), ClassGroup=UserInterface)
 class UMG_API UTextBlock : public UWidget
 {
 	GENERATED_UCLASS_BODY()
 
 public:
+	/** The text to display */
 	UPROPERTY(EditDefaultsOnly, Category=Content)
 	FText Text;
-
+	
+	/** A bindable delegate to allow logic to drive the text of the widget */
 	UPROPERTY()
 	FGetText TextDelegate;
 	
-	UPROPERTY(EditAnywhere, Category=Appearance)
+	/** The font to render the text with */
+	UPROPERTY(EditDefaultsOnly, Category=Appearance)
 	FSlateFontInfo Font;
 
-	UPROPERTY(EditAnywhere, Category=Appearance)
+	/** The direction the shadow is cast */
+	UPROPERTY(EditDefaultsOnly, Category=Appearance)
 	FVector2D ShadowOffset;
 
-	UPROPERTY(EditAnywhere, Category=Appearance)
-	FLinearColor ColorAndOpacity;
+	/** The color of the shadow */
+	UPROPERTY(EditDefaultsOnly, Category=Appearance)
+	FLinearColor ShadowColorAndOpacity;
 
+	/** A bindable delegate for the ShadowColorAndOpacity. */
 	UPROPERTY()
-	FGetSlateColor ColorAndOpacityDelegate;
+	FGetLinearColor ShadowColorAndOpacityDelegate;
 
-	UPROPERTY(EditAnywhere, Category=Appearance)
-	FLinearColor ShadowColorAndOpacity;
+	/** The color of the text */
+	UPROPERTY(EditDefaultsOnly, Category=Appearance)
+	FLinearColor ColorAndOpacity;
 
+	/** A bindable delegate for the ColorAndOpacity. */
 	UPROPERTY()
-	FGetLinearColor ShadowColorAndOpacityDelegate;
+	FGetSlateColor ColorAndOpacityDelegate;
 
-	/** Text to display as the content of the button */
-	UPROPERTY(EditAnywhere, BlueprintReadOnly, Category=Style, meta=( DisplayThumbnail = "true" ))
+	/** The style to use to render the text */
+	UPROPERTY(EditDefaultsOnly, BlueprintReadOnly, Category=Style, meta=( DisplayThumbnail = "true" ))
 	USlateWidgetStyleAsset* Style;
 
+	//TODO UMG Set text via UFunction.
+	//TODO UMG set all these properties by UFunction as well.
+
 protected:
 	// UWidget interface
 	virtual TSharedRef<SWidget> RebuildWidget() OVERRIDE;
 	// End of UWidget interface
 
+	/** Converts the bound FLinearColor to a slate color for the ColorAndOpacityDelegate */
 	FSlateColor GetColorAndOpacity() const;
+
+	/** Converts the bound FLinearColor to a slate color for the ShadowColorAndOpacityDelegate */
 	FLinearColor GetShadowColorAndOpacity() const;
 };
diff --git a/Engine/Source/Runtime/UMG/Public/Components/TileView.h b/Engine/Source/Runtime/UMG/Public/Components/TileView.h
--- a/Engine/Source/Runtime/UMG/Public/Components/TileView.h
+++ b/Engine/Source/Runtime/UMG/Public/Components/TileView.h
@@ -4,6 +4,7 @@
 
 #include "TileView.generated.h"
 
+/** A flow panel that presents the contents as a set of tiles all uniformly sized. */
 UCLASS(meta=(BlueprintSpawnableComponent), ClassGroup=UserInterface)
 class UMG_API UTileView : public UTableViewBase
 {
@@ -23,7 +24,6 @@ public:
 	UPROPERTY(EditDefaultsOnly, Category=Content)
 	TEnumAsByte<ESelectionMode::Type> SelectionMode;
 
-	/** Called when the button is clicked */
 	UPROPERTY(EditDefaultsOnly, Category=Events)
 	FOnGenerateRowUObject OnGenerateTileEvent;
 
diff --git a/Engine/Source/Runtime/UMG/Public/Components/UniformGridPanel.h b/Engine/Source/Runtime/UMG/Public/Components/UniformGridPanel.h
--- a/Engine/Source/Runtime/UMG/Public/Components/UniformGridPanel.h
+++ b/Engine/Source/Runtime/UMG/Public/Components/UniformGridPanel.h
@@ -4,15 +4,31 @@
 
 #include "UniformGridPanel.generated.h"
 
+/** A panel that evenly divides up available space between all of its children. */
 UCLASS(meta=(BlueprintSpawnableComponent), ClassGroup=UserInterface)
 class UMG_API UUniformGridPanel : public UPanelWidget
 {
 	GENERATED_UCLASS_BODY()
 
 	/** The items placed on the canvas */
-	UPROPERTY(EditAnywhere, EditInline, Category=Slots)
+	UPROPERTY(EditDefaultsOnly, EditInline, Category=Slots)
 	TArray<UUniformGridSlot*> Slots;
 
+	/** Padding given to each slot */
+	UPROPERTY(EditDefaultsOnly, Category=Layout)
+	FMargin SlotPadding;
+
+	/** The minimum desired width of the slots */
+	UPROPERTY(EditDefaultsOnly, Category=Layout)
+	float MinDesiredSlotWidth;
+
+	/** The minimum desired height of the slots */
+	UPROPERTY(EditDefaultsOnly, Category=Layout)
+	float MinDesiredSlotHeight;
+
+	//TODO UMG Add ways to make adding slots callable by blueprints.
+
+	/** Adds content to the panel, and constructs a new slot to hold it. */
 	UUniformGridSlot* AddSlot(UWidget* Content);
 
 	// UPanelWidget
diff --git a/Engine/Source/Runtime/UMG/Public/Components/UniformGridSlot.h b/Engine/Source/Runtime/UMG/Public/Components/UniformGridSlot.h
--- a/Engine/Source/Runtime/UMG/Public/Components/UniformGridSlot.h
+++ b/Engine/Source/Runtime/UMG/Public/Components/UniformGridSlot.h
@@ -6,31 +6,47 @@
 
 #include "UniformGridSlot.generated.h"
 
+/**
+ * A slot for UUniformGridPanel, these slots all share the same size as the largest slot
+ * in the grid.
+ */
 UCLASS()
 class UMG_API UUniformGridSlot : public UPanelSlot
 {
 	GENERATED_UCLASS_BODY()
 
-	UPROPERTY(EditAnywhere, Category=Appearance)
+public:
+
+	/** The alignment of the object horizontally. */
+	UPROPERTY(EditDefaultsOnly, Category=Layout)
 	TEnumAsByte<EHorizontalAlignment> HorizontalAlignment;
 
-	UPROPERTY(EditAnywhere, Category=Appearance)
+	/** The alignment of the object vertically. */
+	UPROPERTY(EditDefaultsOnly, Category=Layout)
 	TEnumAsByte<EVerticalAlignment> VerticalAlignment;
 	
-	UPROPERTY(EditAnywhere, Category=Appearance)
+	/** The row index of the cell this slot is in */
+	UPROPERTY(EditDefaultsOnly, Category=Layout)
 	int32 Row;
 	
-	UPROPERTY(EditAnywhere, Category=Appearance)
+	/** The column index of the cell this slot is in */
+	UPROPERTY(EditDefaultsOnly, Category=Layout)
 	int32 Column;
 
+	/** Sets the row index of the slot, this determines what cell the slot is in the panel */
 	UFUNCTION(BlueprintCallable, Category=Layout)
 	void SetRow(int32 InRow);
 
+	/** Sets the column index of the slot, this determines what cell the slot is in the panel */
 	UFUNCTION(BlueprintCallable, Category=Layout)
 	void SetColumn(int32 InColumn);
 
+	/** Builds the underlying FSlot for the Slate layout panel. */
 	void BuildSlot(TSharedRef<SUniformGridPanel> GridPanel);
 
 private:
+	//TODO UMG Slots should hold weak or shared refs to slots.
+
+	/** A raw pointer to the slot to allow us to adjust the size, padding...etc at runtime. */
 	SUniformGridPanel::FSlot* Slot;
 };
\ No newline at end of file
diff --git a/Engine/Source/Runtime/UMG/Public/Components/VerticalBox.h b/Engine/Source/Runtime/UMG/Public/Components/VerticalBox.h
--- a/Engine/Source/Runtime/UMG/Public/Components/VerticalBox.h
+++ b/Engine/Source/Runtime/UMG/Public/Components/VerticalBox.h
@@ -4,15 +4,21 @@
 
 #include "VerticalBox.generated.h"
 
+/**
+ * A vertical box widget is a layout panel allowing child widgets to be automatically laid out
+ * vertically.
+ */
 UCLASS(meta=(BlueprintSpawnableComponent), ClassGroup=UserInterface)
 class UMG_API UVerticalBox : public UPanelWidget
 {
 	GENERATED_UCLASS_BODY()
 
-	/** The items placed on the canvas */
-	UPROPERTY(EditAnywhere, EditInline, Category=Slots)
+	/** The slots containing the widgets that are flowed vertically. */
+	UPROPERTY(EditDefaultsOnly, EditInline, Category=Slots)
 	TArray<UVerticalBoxSlot*> Slots;
 
+	//TODO UMG Add ways to make adding slots callable by blueprints.
+
 	UVerticalBoxSlot* AddSlot(UWidget* Content);
 
 	// UPanelWidget
diff --git a/Engine/Source/Runtime/UMG/Public/Components/VerticalBoxSlot.h b/Engine/Source/Runtime/UMG/Public/Components/VerticalBoxSlot.h
--- a/Engine/Source/Runtime/UMG/Public/Components/VerticalBoxSlot.h
+++ b/Engine/Source/Runtime/UMG/Public/Components/VerticalBoxSlot.h
@@ -6,26 +6,34 @@
 
 #include "VerticalBoxSlot.generated.h"
 
+/** The Slot for the UVerticalBox, contains the widget that is flowed vertically */
 UCLASS()
 class UMG_API UVerticalBoxSlot : public UPanelSlot
 {
 	GENERATED_UCLASS_BODY()
 	
-	UPROPERTY(EditAnywhere, Category=Appearance)
+	/** The padding area between the slot and the content it contains. */
+	UPROPERTY(EditDefaultsOnly, Category=Layout)
 	FMargin Padding;
 
 	/** How much space this slot should occupy in the direction of the panel. */
-	UPROPERTY(EditAnywhere, Category=Appearance)
+	UPROPERTY(EditDefaultsOnly, Category=Layout)
 	FSlateChildSize Size;
 
-	UPROPERTY(EditAnywhere, Category=Appearance)
+	/** The alignment of the object horizontally. */
+	UPROPERTY(EditDefaultsOnly, Category=Layout)
 	TEnumAsByte<EHorizontalAlignment> HorizontalAlignment;
 
-	UPROPERTY(EditAnywhere, Category=Appearance)
+	/** The alignment of the object vertically. */
+	UPROPERTY(EditDefaultsOnly, Category=Layout)
 	TEnumAsByte<EVerticalAlignment> VerticalAlignment;
 
+	/** Builds the underlying FSlot for the Slate layout panel. */
 	void BuildSlot(TSharedRef<SVerticalBox> InVerticalBox);
 
 private:
+	//TODO UMG Slots should hold weak or shared refs to slots.
+
+	/** A raw pointer to the slot to allow us to adjust the size, padding...etc at runtime. */
 	SVerticalBox::FSlot* Slot;
 };
diff --git a/Engine/Source/Runtime/UMG/Public/Components/Widget.h b/Engine/Source/Runtime/UMG/Public/Components/Widget.h
--- a/Engine/Source/Runtime/UMG/Public/Components/Widget.h
+++ b/Engine/Source/Runtime/UMG/Public/Components/Widget.h
@@ -95,24 +95,31 @@ public:
 	UPROPERTY()
 	FGetMouseCursor CursorDelegate;
 
+	/** Gets the current enabled status of the widget */
 	UFUNCTION(BlueprintCallable, Category="Widget")
 	bool GetIsEnabled() const;
 
+	/** Sets the current enabled status of the widget */
 	UFUNCTION(BlueprintCallable, Category="Widget")
 	void SetIsEnabled(bool bInIsEnabled);
 
+	/** Sets the tooltip text for the widget. */
 	UFUNCTION(BlueprintCallable, Category="Widget")
 	void SetToolTipText(const FText& InToolTipText);
 
+	/** Gets the current visibility of the widget. */
 	UFUNCTION(BlueprintCallable, Category="Widget")
 	TEnumAsByte<ESlateVisibility::Type> GetVisibility();
 
+	/** Sets the visiblity of the widget. */
 	UFUNCTION(BlueprintCallable, Category="Widget")
 	void SetVisibility(TEnumAsByte<ESlateVisibility::Type> InVisibility);
 
+	/** Gets the underlying Slate SWidget for this UWidget. */
 	TSharedRef<SWidget> GetWidget();
 
 #if WITH_EDITOR
+	/** Allows general fixups and connections only used at editor time. */
 	virtual void ConnectEditorData() { }
 
 	// UObject interface
@@ -128,6 +135,7 @@ public:
 	static FSizeParam ConvertSerializedSizeParamToRuntime(const FSlateChildSize& Input);
 
 protected:
+	/** Function implemented by all subclasses of UWidget is called when the underlying SWidget needs to be constructed. */
 	virtual TSharedRef<SWidget> RebuildWidget();
 
 	EVisibility ConvertVisibility(TAttribute<ESlateVisibility::Type> SerializedType) const
@@ -136,5 +144,6 @@ protected:
 	}
 
 protected:
+	/** The underlying SWidget. */
 	TSharedPtr<SWidget> MyWidget;
 };
diff --git a/Engine/Source/Runtime/UMG/Public/UMG.h b/Engine/Source/Runtime/UMG/Public/UMG.h
--- a/Engine/Source/Runtime/UMG/Public/UMG.h
+++ b/Engine/Source/Runtime/UMG/Public/UMG.h
@@ -47,6 +47,9 @@
 #include "OverlaySlot.h"
 #include "Overlay.h"
 
+#include "ScrollBoxSlot.h"
+#include "ScrollBox.h"
+
 #include "TextBlock.h"
 #include "Border.h"
 #include "Button.h"
@@ -55,10 +58,10 @@
 #include "EditableTextBlock.h"
 #include "Image.h"
 #include "Slider.h"
+#include "ProgressBar.h"
 #include "ComboBox.h"
 
 #include "GridPanelComponent.h"
-#include "ScrollBox.h"
 #include "Spacer.h"
 
 #include "SpinningImageComponent.h"
