UMG - Removing the box panel.

diff --git a/Engine/Source/Runtime/UMG/Private/Components/BoxPanelComponent.cpp b/Engine/Source/Runtime/UMG/Private/Components/BoxPanelComponent.cpp
deleted file mode 100644
diff --git a/Engine/Source/Runtime/UMG/Public/Components/BoxPanelComponent.h b/Engine/Source/Runtime/UMG/Public/Components/BoxPanelComponent.h
deleted file mode 100644
diff --git a/Engine/Source/Runtime/UMG/Public/UMG.h b/Engine/Source/Runtime/UMG/Public/UMG.h
--- a/Engine/Source/Runtime/UMG/Public/UMG.h
+++ b/Engine/Source/Runtime/UMG/Public/UMG.h
@@ -30,8 +30,6 @@
 #include "PanelWidget.h"
 #include "ContentWidget.h"
 
-#include "BoxPanelComponent.h"
-
 #include "CanvasPanelSlot.h"
 #include "CanvasPanel.h"
 
