UMG - Removing the grid panel.

diff --git a/Engine/Source/Runtime/UMG/Private/Components/GridPanelComponent.cpp b/Engine/Source/Runtime/UMG/Private/Components/GridPanelComponent.cpp
deleted file mode 100644
diff --git a/Engine/Source/Runtime/UMG/Public/Components/GridPanelComponent.h b/Engine/Source/Runtime/UMG/Public/Components/GridPanelComponent.h
deleted file mode 100644
diff --git a/Engine/Source/Runtime/UMG/Public/UMG.h b/Engine/Source/Runtime/UMG/Public/UMG.h
--- a/Engine/Source/Runtime/UMG/Public/UMG.h
+++ b/Engine/Source/Runtime/UMG/Public/UMG.h
@@ -59,7 +59,6 @@
 #include "ProgressBar.h"
 #include "ComboBox.h"
 
-#include "GridPanelComponent.h"
 #include "Spacer.h"
 
 #include "SpinningImageComponent.h"
