UMG - Working on a better wysiwyg layout panel for the designer.

diff --git a/Engine/Source/Editor/UMGEditor/Private/Extensions/CanvasSlotExtension.cpp b/Engine/Source/Editor/UMGEditor/Private/Extensions/CanvasSlotExtension.cpp
--- a/Engine/Source/Editor/UMGEditor/Private/Extensions/CanvasSlotExtension.cpp
+++ b/Engine/Source/Editor/UMGEditor/Private/Extensions/CanvasSlotExtension.cpp
@@ -86,7 +86,11 @@ void FCanvasSlotExtension::MoveByAmount(UWidget* Widget, FVector2D Delta)
 		UCanvasPanelSlot* CanvasSlot = Cast<UCanvasPanelSlot>(Widget->Slot);
 		UCanvasPanel* Parent = Cast<UCanvasPanel>(CanvasSlot->Parent);
 
-		CanvasSlot->SetPosition(CanvasSlot->Position + Delta);
+		FMargin Offset = CanvasSlot->Offset;
+		Offset.Left += Delta.X;
+		Offset.Top += Delta.Y;
+
+		CanvasSlot->SetOffset(Offset);
 	}
 }
 
diff --git a/Engine/Source/Editor/UMGEditor/Private/SUMGEditorTree.cpp b/Engine/Source/Editor/UMGEditor/Private/SUMGEditorTree.cpp
--- a/Engine/Source/Editor/UMGEditor/Private/SUMGEditorTree.cpp
+++ b/Engine/Source/Editor/UMGEditor/Private/SUMGEditorTree.cpp
@@ -256,10 +256,7 @@ FReply SUMGEditorTree::CreateTestUI()
 		Button4->SetContent(Text1);
 
 		UCanvasPanelSlot* Slot = Canvas->AddSlot(Border);
-		Slot->Size.X = 200;
-		Slot->Size.Y = 300;
-		Slot->Position.X = 20;
-		Slot->Position.Y = 50;
+		Slot->Offset = FMargin(200, 300, 20, 50);
 
 		Border->SetContent(Vertical);
 
diff --git a/Engine/Source/Runtime/UMG/Private/Components/CanvasPanel.cpp b/Engine/Source/Runtime/UMG/Private/Components/CanvasPanel.cpp
--- a/Engine/Source/Runtime/UMG/Private/Components/CanvasPanel.cpp
+++ b/Engine/Source/Runtime/UMG/Private/Components/CanvasPanel.cpp
@@ -5,7 +5,7 @@
 /////////////////////////////////////////////////////
 // SFixedSizeCanvas
 
-class SFixedSizeCanvas : public SCanvas
+class SFixedSizeCanvas : public SConstraintCanvas
 {
 	SLATE_BEGIN_ARGS( SFixedSizeCanvas )
 		{
@@ -17,8 +17,8 @@ class SFixedSizeCanvas : public SCanvas
 
 	void Construct(const FArguments& InArgs, FVector2D InSize)
 	{
-		SCanvas::FArguments ParentArgs;
-		SCanvas::Construct(ParentArgs);
+		SConstraintCanvas::FArguments ParentArgs;
+		SConstraintCanvas::Construct(ParentArgs);
 
 		CanvasSize = InSize;
 	}
@@ -94,8 +94,7 @@ UCanvasPanelSlot* UCanvasPanel::AddSlot(UWidget* Content)
 bool UCanvasPanel::AddChild(UWidget* Child, FVector2D Position)
 {
 	UCanvasPanelSlot* Slot = AddSlot(Child);
-	Slot->Position = Position;
-	Slot->Size = FVector2D(100, 25);
+	Slot->Offset = FMargin(Position.X, Position.Y, 100, 25);
 
 	return true;
 }
diff --git a/Engine/Source/Runtime/UMG/Private/Components/CanvasPanelSlot.cpp b/Engine/Source/Runtime/UMG/Private/Components/CanvasPanelSlot.cpp
--- a/Engine/Source/Runtime/UMG/Private/Components/CanvasPanelSlot.cpp
+++ b/Engine/Source/Runtime/UMG/Private/Components/CanvasPanelSlot.cpp
@@ -9,19 +9,17 @@ UCanvasPanelSlot::UCanvasPanelSlot(const FPostConstructInitializeProperties& PCI
 	: Super(PCIP)
 	, Slot(NULL)
 {
-	Position = FVector2D::ZeroVector;
-	Size = FVector2D(1.0f, 1.0f);
-	HorizontalAlignment = HAlign_Left;
-	VerticalAlignment = VAlign_Top;
+	Offset = FMargin(0, 0, 1, 1);
+	Anchors = FAnchors(0.5f, 0.5f);
+	Alignment = FVector2D(0.5f, 0.5f);
 }
 
-void UCanvasPanelSlot::BuildSlot(TSharedRef<SCanvas> Canvas)
+void UCanvasPanelSlot::BuildSlot(TSharedRef<SConstraintCanvas> Canvas)
 {
 	Slot = &Canvas->AddSlot()
-		.Position(Position)
-		.Size(Size)
-		.HAlign(HorizontalAlignment)
-		.VAlign(VerticalAlignment)
+		.Offset(Offset)
+		.Anchors(Anchors)
+		.Alignment(Alignment)
 		[
 			Content == NULL ? SNullWidget::NullWidget : Content->GetWidget()
 		];
@@ -31,27 +29,26 @@ void UCanvasPanelSlot::Resize(const FVector2D& Direction, const FVector2D& Amoun
 {
 	if ( Direction.X < 0 )
 	{
-		Position.X -= Amount.X * Direction.X;
-		Size.X += Amount.X * Direction.X;
+		Offset.Left -= Amount.X * Direction.X;
+		Offset.Right += Amount.X * Direction.X;
 	}
 	if ( Direction.Y < 0 )
 	{
-		Position.Y -= Amount.Y * Direction.Y;
-		Size.Y += Amount.Y * Direction.Y;
+		Offset.Top -= Amount.Y * Direction.Y;
+		Offset.Bottom += Amount.Y * Direction.Y;
 	}
 	if ( Direction.X > 0 )
 	{
-		Size.X += Amount.X * Direction.X;
+		Offset.Right += Amount.X * Direction.X;
 	}
 	if ( Direction.Y > 0 )
 	{
-		Size.Y += Amount.Y * Direction.Y;
+		Offset.Bottom += Amount.Y * Direction.Y;
 	}
 
 	if ( Slot )
 	{
-		Slot->Position(Position);
-		Slot->Size(Size);
+		Slot->Offset(Offset);
 	}
 }
 
@@ -60,11 +57,29 @@ bool UCanvasPanelSlot::CanResize(const FVector2D& Direction) const
 	return true;
 }
 
-void UCanvasPanelSlot::SetPosition(FVector2D InPosition)
+void UCanvasPanelSlot::SetOffset(FMargin InOffset)
 {
-	Position = InPosition;
+	Offset = InOffset;
 	if ( Slot )
 	{
-		Slot->Position(InPosition);
+		Slot->Offset(InOffset);
 	}
 }
+
+void UCanvasPanelSlot::SetAnchors(FAnchors InAnchors)
+{
+	Anchors = InAnchors;
+	if ( Slot )
+	{
+		Slot->Anchors(InAnchors);
+	}
+}
+
+void UCanvasPanelSlot::SetAlignment(FVector2D InAlignment)
+{
+	Alignment = InAlignment;
+	if ( Slot )
+	{
+		Slot->Alignment(InAlignment);
+	}
+}
\ No newline at end of file
diff --git a/Engine/Source/Runtime/UMG/Private/Slate/SConstraintCanvas.cpp b/Engine/Source/Runtime/UMG/Private/Slate/SConstraintCanvas.cpp
new file mode 100644
--- /dev/null
+++ b/Engine/Source/Runtime/UMG/Private/Slate/SConstraintCanvas.cpp
@@ -0,0 +1,148 @@
+// Copyright 1998-2014 Epic Games, Inc. All Rights Reserved.
+
+/*=============================================================================
+	SConstraintCanvas.cpp: Implements the SConstraintCanvas class.
+=============================================================================*/
+
+#include "UMGPrivatePCH.h"
+
+#include "LayoutUtils.h"
+
+
+/* SConstraintCanvas interface
+ *****************************************************************************/
+
+void SConstraintCanvas::Construct( const SConstraintCanvas::FArguments& InArgs )
+{
+	const int32 NumSlots = InArgs.Slots.Num();
+	for ( int32 SlotIndex = 0; SlotIndex < NumSlots; ++SlotIndex )
+	{
+		Children.Add( InArgs.Slots[SlotIndex] );
+	}
+}
+
+
+void SConstraintCanvas::ClearChildren( )
+{
+	Children.Empty();
+}
+
+
+int32 SConstraintCanvas::RemoveSlot( const TSharedRef<SWidget>& SlotWidget )
+{
+	for (int32 SlotIdx = 0; SlotIdx < Children.Num(); ++SlotIdx)
+	{
+		if (SlotWidget == Children[SlotIdx].Widget)
+		{
+			Children.RemoveAt(SlotIdx);
+			return SlotIdx;
+		}
+	}
+
+	return -1;
+}
+
+
+/* SWidget overrides
+ *****************************************************************************/
+
+void SConstraintCanvas::OnArrangeChildren( const FGeometry& AllottedGeometry, FArrangedChildren& ArrangedChildren ) const
+{
+	if (Children.Num() > 0)
+	{
+		for (int32 ChildIndex = 0; ChildIndex < Children.Num(); ++ChildIndex)
+		{
+			const SConstraintCanvas::FSlot& CurChild = Children[ChildIndex];
+
+			const FMargin Offset = CurChild.OffsetAttr.Get();
+			const FVector2D Alignment = CurChild.AlignmentAttr.Get();
+			const FAnchors Anchors = CurChild.AnchorsAttr.Get();
+
+			const FMargin AnchorPixels =
+				FMargin(Anchors.Minimum.X * AllottedGeometry.Size.X,
+						Anchors.Minimum.Y * AllottedGeometry.Size.Y,
+						Anchors.Maximum.X * AllottedGeometry.Size.X,
+						Anchors.Maximum.Y * AllottedGeometry.Size.Y);
+
+			const bool bIsHorizontalStretch = Anchors.Minimum.X != Anchors.Maximum.X;
+			const bool bIsVerticalStretch = Anchors.Minimum.Y != Anchors.Maximum.Y;
+			
+			const FVector2D Size = FVector2D(Offset.Right, Offset.Bottom);
+			
+			// Calculate the offset based on the pivot position.
+			FVector2D AlignmentOffset = Size * Alignment;
+
+			// Calculate the local position based on the anchor and position offset.
+			FVector2D LocalPosition, LocalSize;
+
+			// Calculate the position and size based on the horizontal stretch or non-stretch
+			if ( bIsHorizontalStretch )
+			{
+				LocalPosition.X = AnchorPixels.Left + Offset.Left;
+				LocalSize.X = AnchorPixels.Right - AnchorPixels.Left - Offset.Right;
+			}
+			else
+			{
+				LocalPosition.X = AnchorPixels.Left + Offset.Left - AlignmentOffset.X;
+				LocalSize.X = Size.X;
+			}
+
+			// Calculate the position and size based on the vertical stretch or non-stretch
+			if ( bIsVerticalStretch )
+			{
+				LocalPosition.Y = AnchorPixels.Top + Offset.Top;
+				LocalSize.Y = AnchorPixels.Bottom - AnchorPixels.Top - Offset.Bottom;
+			}
+			else
+			{
+				LocalPosition.Y = AnchorPixels.Top + Offset.Top - AlignmentOffset.Y;
+				LocalSize.Y = Size.Y;
+			}
+						
+			// Add the information about this child to the output list (ArrangedChildren)
+			ArrangedChildren.AddWidget( AllottedGeometry.MakeChild(
+				// The child widget being arranged
+				CurChild.Widget,
+				// Child's local position (i.e. position within parent)
+				LocalPosition,
+				// Child's size
+				LocalSize
+			));
+		}
+	}
+}
+
+
+int32 SConstraintCanvas::OnPaint( const FGeometry& AllottedGeometry, const FSlateRect& MyClippingRect, FSlateWindowElementList& OutDrawElements, int32 LayerId, const FWidgetStyle& InWidgetStyle, bool bParentEnabled ) const
+{
+	FArrangedChildren ArrangedChildren(EVisibility::Visible);
+	this->ArrangeChildren(AllottedGeometry, ArrangedChildren);
+
+	// Because we paint multiple children, we must track the maximum layer id that they produced in case one of our parents
+	// wants to an overlay for all of its contents.
+	int32 MaxLayerId = LayerId;
+
+	for (int32 ChildIndex = 0; ChildIndex < ArrangedChildren.Num(); ++ChildIndex)
+	{
+		FArrangedWidget& CurWidget = ArrangedChildren(ChildIndex);
+		FSlateRect ChildClipRect = MyClippingRect.IntersectionWith(CurWidget.Geometry.GetClippingRect());
+		const int32 CurWidgetsMaxLayerId = CurWidget.Widget->Paint(CurWidget.Geometry, ChildClipRect, OutDrawElements, MaxLayerId + 1, InWidgetStyle, ShouldBeEnabled(bParentEnabled));
+
+		MaxLayerId = FMath::Max(MaxLayerId, CurWidgetsMaxLayerId);
+	}
+
+	return MaxLayerId;
+}
+
+
+FVector2D SConstraintCanvas::ComputeDesiredSize() const
+{
+	// ConstraintCanvas widgets have no desired size -- their size is always determined by their container
+	return FVector2D::ZeroVector;
+}
+
+
+FChildren* SConstraintCanvas::GetChildren()
+{
+	return &Children;
+}
diff --git a/Engine/Source/Runtime/UMG/Public/Components/CanvasPanelSlot.h b/Engine/Source/Runtime/UMG/Public/Components/CanvasPanelSlot.h
--- a/Engine/Source/Runtime/UMG/Public/Components/CanvasPanelSlot.h
+++ b/Engine/Source/Runtime/UMG/Public/Components/CanvasPanelSlot.h
@@ -2,6 +2,8 @@
 
 #pragma once
 
+#include "Anchors.h"
+
 #include "CanvasPanelSlot.generated.h"
 
 UCLASS()
@@ -9,53 +11,36 @@ class UMG_API UCanvasPanelSlot : public UPanelSlot
 {
 	GENERATED_UCLASS_BODY()
 
-	/** Position. */
-	UPROPERTY(EditAnywhere, Category=Appearance)
-	FVector2D Position;
-
-	/** Size. */
-	UPROPERTY(EditAnywhere, Category=Appearance)
-	FVector2D Size;
-
-	/**
-	* Horizontal pivot position
-	*  Given a top aligned slot, where '+' represents the
-	*  anchor point defined by PositionAttr.
-	*
-	*   Left				Center				Right
-	+ _ _ _ _            _ _ + _ _          _ _ _ _ +
-	|		  |		   | 		   |	  |		    |
-	| _ _ _ _ |        | _ _ _ _ _ |	  | _ _ _ _ |
-	*
-	*  Note: FILL is NOT supported.
-	*/
-	UPROPERTY(EditAnywhere, Category=Appearance)
-	TEnumAsByte<EHorizontalAlignment> HorizontalAlignment;
-
-	/**
-	* Vertical pivot position
-	*   Given a left aligned slot, where '+' represents the
-	*   anchor point defined by PositionAttr.
-	*
-	*   Top					Center			  Bottom
-	*	+_ _ _ _ _		 _ _ _ _ _		 _ _ _ _ _
-	*	|         |		| 		  |		|		  |
-	*	|         |     +		  |		|		  |
-	*	| _ _ _ _ |		| _ _ _ _ |		+ _ _ _ _ |
-	*
-	*  Note: FILL is NOT supported.
-	*/
-	UPROPERTY(EditAnywhere, Category=Appearance)
-	TEnumAsByte<EVerticalAlignment> VerticalAlignment;
-
-	void BuildSlot(TSharedRef<SCanvas> Canvas);
+	/** Offset. */
+	UPROPERTY(EditDefaultsOnly, Category=Appearance)
+	FMargin Offset;
+	
+	/** Anchors. */
+	UPROPERTY(EditDefaultsOnly, Category=Appearance)
+	FAnchors Anchors;
+
+	/** Alignment. */
+	UPROPERTY(EditDefaultsOnly, Category=Appearance)
+	FVector2D Alignment;
+
+	void BuildSlot(TSharedRef<SConstraintCanvas> Canvas);
 
 	virtual void Resize(const FVector2D& Direction, const FVector2D& Amount) OVERRIDE;
 
 	virtual bool CanResize(const FVector2D& Direction) const OVERRIDE;
 
-	void SetPosition(FVector2D Position);
+	/** Sets the position of the slot */
+	UFUNCTION(BlueprintCallable, Category="Appearance")
+	void SetOffset(FMargin InOffset);
+	
+	/** Sets the anchors on the slot */
+	UFUNCTION(BlueprintCallable, Category="Appearance")
+	void SetAnchors(FAnchors InAnchors);
+
+	/** Sets the alignment on the slot */
+	UFUNCTION(BlueprintCallable, Category="Appearance")
+	void SetAlignment(FVector2D InAlignment);
 
 private:
-	SCanvas::FSlot* Slot;
+	SConstraintCanvas::FSlot* Slot;
 };
\ No newline at end of file
diff --git a/Engine/Source/Runtime/UMG/Public/Slate/Anchors.h b/Engine/Source/Runtime/UMG/Public/Slate/Anchors.h
new file mode 100644
--- /dev/null
+++ b/Engine/Source/Runtime/UMG/Public/Slate/Anchors.h
@@ -0,0 +1,52 @@
+// Copyright 1998-2014 Epic Games, Inc. All Rights Reserved.
+
+#pragma once
+
+#include "Anchors.generated.h"
+
+/**
+ * Describes how a widget is anchored.
+ */
+USTRUCT()
+struct FAnchors
+{
+	GENERATED_USTRUCT_BODY()
+
+	/** Holds the minimum anchors, left + top. */
+	UPROPERTY(EditAnywhere, Category=Appearance)
+	FVector2D Minimum;
+
+	/** Holds the maximum anchors, right + bottom. */
+	UPROPERTY(EditAnywhere, Category=Appearance)
+	FVector2D Maximum;
+
+public:
+
+	/**
+	 * Default constructor.
+	 *
+	 * The default margin size is zero on all four sides..
+	 */
+	FAnchors()
+		: Minimum(0.0f, 0.0f)
+		, Maximum(0.0f, 0.0f)
+	{ }
+
+	/** Construct a Anchors with uniform space on all sides */
+	FAnchors(float UnifromAnchors)
+		: Minimum(UnifromAnchors, UnifromAnchors)
+		, Maximum(UnifromAnchors, UnifromAnchors)
+	{ }
+
+	/** Construct a Anchors where Horizontal describes Left and Right spacing while Vertical describes Top and Bottom spacing */
+	FAnchors(float Horizontal, float Vertical)
+		: Minimum(Horizontal, Vertical)
+		, Maximum(Horizontal, Vertical)
+	{ }
+
+	/** Construct Anchors where the spacing on each side is individually specified. */
+	FAnchors(float InLeft, float InTop, float InRight, float InBottom)
+		: Minimum(InLeft, InTop)
+		, Maximum(InRight, InBottom)
+	{ }
+};
diff --git a/Engine/Source/Runtime/UMG/Public/Slate/SConstraintCanvas.h b/Engine/Source/Runtime/UMG/Public/Slate/SConstraintCanvas.h
new file mode 100644
--- /dev/null
+++ b/Engine/Source/Runtime/UMG/Public/Slate/SConstraintCanvas.h
@@ -0,0 +1,137 @@
+// Copyright 1998-2014 Epic Games, Inc. All Rights Reserved.
+
+/*=============================================================================
+	SConstraintCanvas.h: Declares the SConstraintCanvas class.
+=============================================================================*/
+
+#pragma once
+
+/**
+ * ConstraintCanvas is a layout widget that allows you to arbitrary position and size child widgets in a relative coordinate space
+ */
+class UMG_API SConstraintCanvas
+	: public SPanel
+{
+public:
+
+	/**
+	 * ConstraintCanvas slots allow child widgets to be positioned and sized
+	 */
+	class FSlot
+	{
+	public:		
+		FSlot& Offset( const TAttribute<FMargin>& InOffset )
+		{
+			OffsetAttr = InOffset;
+			return *this;
+		}
+
+		FSlot& Anchors( const TAttribute<FAnchors>& InAnchors )
+		{
+			AnchorsAttr = InAnchors;
+			return *this;
+		}
+
+		FSlot& Alignment(const TAttribute<FVector2D>& InAlignment)
+		{
+			AlignmentAttr = InAlignment;
+			return *this;
+		}
+
+		FSlot& operator[]( TSharedRef<SWidget> InWidget )
+		{
+			Widget = InWidget;
+			return *this;
+		}
+
+		FSlot& Expose( FSlot*& OutVarToInit )
+		{
+			OutVarToInit = this;
+			return *this;
+		}
+
+		/** The child widget contained in this slot. */
+		TSharedRef<SWidget> Widget;
+
+		/** Offset */
+		TAttribute<FMargin> OffsetAttr;
+
+		/** Anchors */
+		TAttribute<FAnchors> AnchorsAttr;
+
+		/** Size */
+		TAttribute<FVector2D> AlignmentAttr;
+
+		/** Default values for a slot. */
+		FSlot()
+			: Widget( SNullWidget::NullWidget )
+			, OffsetAttr( FMargin( 0, 0, 1, 1 ) )
+			, AnchorsAttr( FAnchors( 0.5f, 0.5f ) )
+			, AlignmentAttr( FVector2D( 0.5f, 0.5f ) )
+		{ }
+	};
+
+	SLATE_BEGIN_ARGS( SConstraintCanvas )
+		{
+			_Visibility = EVisibility::SelfHitTestInvisible;
+		}
+
+		SLATE_SUPPORTS_SLOT( SConstraintCanvas::FSlot )
+
+	SLATE_END_ARGS()
+
+	/**
+	 * Construct this widget
+	 *
+	 * @param	InArgs	The declaration data for this widget
+	 */
+	void Construct( const FArguments& InArgs );
+
+	static FSlot& Slot()
+	{
+		return *(new FSlot());
+	}
+
+	/**
+	 * Adds a content slot.
+	 *
+	 * @return The added slot.
+	 */
+	FSlot& AddSlot()
+	{
+		SConstraintCanvas::FSlot& NewSlot = SConstraintCanvas::Slot();
+		this->Children.Add( &NewSlot );
+		return NewSlot;
+	}
+
+	/**
+	 * Removes a particular content slot.
+	 *
+	 * @param SlotWidget The widget in the slot to remove.
+	 */
+	int32 RemoveSlot( const TSharedRef<SWidget>& SlotWidget );
+
+	/**
+	 * Removes all slots from the panel.
+	 */
+	void ClearChildren();
+
+public:
+
+	// Begin SWidget overrides
+
+	virtual void OnArrangeChildren( const FGeometry& AllottedGeometry, FArrangedChildren& ArrangedChildren ) const OVERRIDE;
+	
+	virtual int32 OnPaint( const FGeometry& AllottedGeometry, const FSlateRect& MyClippingRect, FSlateWindowElementList& OutDrawElements, int32 LayerId, const FWidgetStyle& InWidgetStyle, bool bParentEnabled ) const OVERRIDE;
+
+	virtual FVector2D ComputeDesiredSize() const OVERRIDE;
+
+	virtual FChildren* GetChildren() OVERRIDE;
+
+	// End SWidget overrides
+
+protected:
+
+	/** The ConstraintCanvas widget's children. */
+	TPanelChildren< FSlot > Children;
+};
diff --git a/Engine/Source/Runtime/UMG/Public/UMG.h b/Engine/Source/Runtime/UMG/Public/UMG.h
--- a/Engine/Source/Runtime/UMG/Public/UMG.h
+++ b/Engine/Source/Runtime/UMG/Public/UMG.h
@@ -17,6 +17,8 @@
 #include "Slate.h"
 #include "SlateCore.h"
 
+#include "Anchors.h"
+#include "SConstraintCanvas.h"
 
 /* Components
  *****************************************************************************/
@@ -72,7 +74,7 @@
 /* Slate
  *****************************************************************************/
  
- #include "SObjectWidget.h"
+#include "SObjectWidget.h"
 
 /* Blueprint
  *****************************************************************************/
