UMG - Removing spinning image component.

diff --git a/Engine/Source/Runtime/UMG/Private/Components/SpinningImageComponent.cpp b/Engine/Source/Runtime/UMG/Private/Components/SpinningImageComponent.cpp
deleted file mode 100644
diff --git a/Engine/Source/Runtime/UMG/Public/Components/SpinningImageComponent.h b/Engine/Source/Runtime/UMG/Public/Components/SpinningImageComponent.h
deleted file mode 100644
diff --git a/Engine/Source/Runtime/UMG/Public/UMG.h b/Engine/Source/Runtime/UMG/Public/UMG.h
--- a/Engine/Source/Runtime/UMG/Public/UMG.h
+++ b/Engine/Source/Runtime/UMG/Public/UMG.h
@@ -63,7 +63,6 @@
 
 #include "Spacer.h"
 
-#include "SpinningImageComponent.h"
 #include "ThrobberComponent.h"
 #include "CircularThrobberComponent.h"
 
