UE4 Refactoring. Changed OVERRIDE and FINAL macros to keywords override and final.

diff --git a/Engine/Source/Runtime/Sockets/Private/IOS/SocketSubsystemIOS.h b/Engine/Source/Runtime/Sockets/Private/IOS/SocketSubsystemIOS.h
--- a/Engine/Source/Runtime/Sockets/Private/IOS/SocketSubsystemIOS.h
+++ b/Engine/Source/Runtime/Sockets/Private/IOS/SocketSubsystemIOS.h
@@ -45,10 +45,10 @@ public:
 	{
 	}
 
-	virtual bool Init(FString& Error) OVERRIDE;
-	virtual void Shutdown() OVERRIDE;
-	virtual bool HasNetworkDevice() OVERRIDE;
-	virtual FSocket* CreateSocket(const FName& SocketType, const FString& SocketDescription, bool bForceUDP) OVERRIDE;
+	virtual bool Init(FString& Error) override;
+	virtual void Shutdown() override;
+	virtual bool HasNetworkDevice() override;
+	virtual FSocket* CreateSocket(const FName& SocketType, const FString& SocketDescription, bool bForceUDP) override;
 
-	virtual TSharedRef<FInternetAddr> GetLocalHostAddr(FOutputDevice& Out, bool& bCanBindAll) OVERRIDE;
+	virtual TSharedRef<FInternetAddr> GetLocalHostAddr(FOutputDevice& Out, bool& bCanBindAll) override;
 };
\ No newline at end of file
diff --git a/Engine/Source/Runtime/Sockets/Private/Mac/SocketSubsystemMac.h b/Engine/Source/Runtime/Sockets/Private/Mac/SocketSubsystemMac.h
--- a/Engine/Source/Runtime/Sockets/Private/Mac/SocketSubsystemMac.h
+++ b/Engine/Source/Runtime/Sockets/Private/Mac/SocketSubsystemMac.h
@@ -56,17 +56,17 @@ public:
 	 *
 	 * @return true if initialized ok, false otherwise
 	 */
-	virtual bool Init(FString& Error) OVERRIDE;
+	virtual bool Init(FString& Error) override;
 
 	/**
 	 * Performs platform specific socket clean up
 	 */
-	virtual void Shutdown() OVERRIDE;
+	virtual void Shutdown() override;
 
 	/**
 	 * @return Whether the machine has a properly configured network device or not
 	 */
-	virtual bool HasNetworkDevice() OVERRIDE;
+	virtual bool HasNetworkDevice() override;
 };
 
 /**
@@ -92,12 +92,12 @@ public:
 	{
 	}
 
-	virtual bool Close() OVERRIDE
+	virtual bool Close() override
 	{
 		return close(Socket) == 0;
 	}
 
-	virtual bool SetReuseAddr(bool bAllowReuse) OVERRIDE
+	virtual bool SetReuseAddr(bool bAllowReuse) override
 	{
 		int Param = bAllowReuse ? 1 : 0;
 		return (setsockopt(Socket,SOL_SOCKET,SO_REUSEADDR,(char*)&Param,sizeof(Param)) == 0) && (setsockopt(Socket,SOL_SOCKET,SO_REUSEPORT,(char*)&Param,sizeof(Param)) == 0);
diff --git a/Engine/Source/Runtime/Sockets/Private/WinRT/SocketSubsystemWinRT.cpp b/Engine/Source/Runtime/Sockets/Private/WinRT/SocketSubsystemWinRT.cpp
--- a/Engine/Source/Runtime/Sockets/Private/WinRT/SocketSubsystemWinRT.cpp
+++ b/Engine/Source/Runtime/Sockets/Private/WinRT/SocketSubsystemWinRT.cpp
@@ -37,7 +37,7 @@ public:
 	 *
 	 * @param InAddr the new address to use (must convert to network byte order)
 	 */
-	virtual void SetIp(uint32 InAddr) OVERRIDE
+	virtual void SetIp(uint32 InAddr) override
 	{
 
 	}
@@ -47,7 +47,7 @@ public:
 	 *
 	 * @param InAddr the string containing the new ip address to use
 	 */
-	virtual void SetIp(const TCHAR* InAddr, bool& bIsValid) OVERRIDE
+	virtual void SetIp(const TCHAR* InAddr, bool& bIsValid) override
 	{
 
 	}
@@ -57,7 +57,7 @@ public:
 	 *
 	 * @param OutAddr the out param receiving the ip address
 	 */
-	virtual void GetIp(uint32& OutAddr) const OVERRIDE
+	virtual void GetIp(uint32& OutAddr) const override
 	{
 
 	}
@@ -67,7 +67,7 @@ public:
 	 *
 	 * @param InPort the new port to use (must convert to network byte order)
 	 */
-	virtual void SetPort(int32 InPort) OVERRIDE
+	virtual void SetPort(int32 InPort) override
 	{
 
 	}
@@ -77,7 +77,7 @@ public:
 	 *
 	 * @param OutPort the host byte order int that receives the port
 	 */
-	virtual void GetPort(int32& OutPort) const OVERRIDE
+	virtual void GetPort(int32& OutPort) const override
 	{
 
 	}
@@ -85,19 +85,19 @@ public:
 	/**
 	 * Returns the port number from this address in host byte order
 	 */
-	virtual int32 GetPort() const OVERRIDE
+	virtual int32 GetPort() const override
 	{
 		return 0;
 	}
 
 	/** Sets the address to be any address */
-	virtual void SetAnyAddress() OVERRIDE
+	virtual void SetAnyAddress() override
 	{
 
 	}
 
 	/** Sets the address to broadcast */
-	virtual void SetBroadcastAddress() OVERRIDE
+	virtual void SetBroadcastAddress() override
 	{
 
 	}
@@ -107,7 +107,7 @@ public:
 	 *
 	 * @param bAppendPort whether to append the port information or not
 	 */
-	virtual FString ToString(bool bAppendPort) const OVERRIDE
+	virtual FString ToString(bool bAppendPort) const override
 	{
 		return TEXT("");
 	}
@@ -130,7 +130,7 @@ public:
 	 *
 	 * @return true if a valid IP, false otherwise
 	 */
-	virtual bool IsValid() const OVERRIDE
+	virtual bool IsValid() const override
 	{
 		return true;
 	}
@@ -241,7 +241,7 @@ TSharedRef<FInternetAddr> FSocketSubsystemWinRT::CreateInternetAddr(uint32 Addre
 	return Result;
 }
 
-bool FSocketSubsystemWinRT::GetLocalAdapterAddresses( TArray<TSharedPtr<FInternetAddr> >& OutAdresses ) OVERRIDE
+bool FSocketSubsystemWinRT::GetLocalAdapterAddresses( TArray<TSharedPtr<FInternetAddr> >& OutAdresses ) override
 {
 	bool bCanBindAll;
 	OutAdresses.Add(GetLocalHostAddr(*GLog, bCanBindAll));
diff --git a/Engine/Source/Runtime/Sockets/Private/WinRT/SocketSubsystemWinRT.h b/Engine/Source/Runtime/Sockets/Private/WinRT/SocketSubsystemWinRT.h
--- a/Engine/Source/Runtime/Sockets/Private/WinRT/SocketSubsystemWinRT.h
+++ b/Engine/Source/Runtime/Sockets/Private/WinRT/SocketSubsystemWinRT.h
@@ -51,12 +51,12 @@ public:
 	 *
 	 * @return true if initialized ok, false otherwise
 	 */
-	virtual bool Init(FString& Error) OVERRIDE;
+	virtual bool Init(FString& Error) override;
 
 	/**
 	 * Performs platform specific socket clean up
 	 */
-	virtual void Shutdown() OVERRIDE;
+	virtual void Shutdown() override;
 
 	/**
 	 * Creates a socket
@@ -67,14 +67,14 @@ public:
 	 *
 	 * @return the new socket or NULL if failed
 	 */
-	virtual class FSocket* CreateSocket(const FName& SocketType, const FString& SocketDescription, bool bForceUDP = false) OVERRIDE;
+	virtual class FSocket* CreateSocket(const FName& SocketType, const FString& SocketDescription, bool bForceUDP = false) override;
 
 	/**
 	 * Cleans up a socket class
 	 *
 	 * @param Socket the socket object to destroy
 	 */
-	virtual void DestroySocket(class FSocket* Socket) OVERRIDE;
+	virtual void DestroySocket(class FSocket* Socket) override;
 
 	/**
 	 * Does a DNS look up of a host name
@@ -82,7 +82,7 @@ public:
 	 * @param HostName the name of the host to look up
 	 * @param Addr the address to copy the IP address to
 	 */
-	virtual ESocketErrors GetHostByName(const ANSICHAR* HostName, FInternetAddr& OutAddr) OVERRIDE;
+	virtual ESocketErrors GetHostByName(const ANSICHAR* HostName, FInternetAddr& OutAddr) override;
 
 	/**
 	 * Creates a platform specific async hostname resolution object
@@ -91,14 +91,14 @@ public:
 	 *
 	 * @return the resolve info to query for the address
 	 */
-	virtual class FResolveInfo* GetHostByName(const ANSICHAR* HostName) OVERRIDE;
+	virtual class FResolveInfo* GetHostByName(const ANSICHAR* HostName) override;
 
 	/**
 	 * Some platforms require chat data (voice, text, etc.) to be placed into
 	 * packets in a special way. This function tells the net connection
 	 * whether this is required for this platform
 	 */
-	virtual bool RequiresChatDataBeSeparate() OVERRIDE
+	virtual bool RequiresChatDataBeSeparate() override
 	{
 		return false;
 	}
@@ -107,7 +107,7 @@ public:
 	 * Some platforms require packets be encrypted. This function tells the
 	 * net connection whether this is required for this platform
 	 */
-	virtual bool RequiresEncryptedPackets() OVERRIDE
+	virtual bool RequiresEncryptedPackets() override
 	{
 		return false;
 	}
@@ -119,35 +119,35 @@ public:
 	 *
 	 * @return true if successful, false otherwise
 	 */
-	virtual bool GetHostName(FString& HostName) OVERRIDE;
+	virtual bool GetHostName(FString& HostName) override;
 
 	/**
 	 *	Create a proper FInternetAddr representation
 	 * @param Address host address
 	 * @param Port host port
 	 */
-	virtual TSharedRef<FInternetAddr> CreateInternetAddr(uint32 Address=0, uint32 Port=0) OVERRIDE;
+	virtual TSharedRef<FInternetAddr> CreateInternetAddr(uint32 Address=0, uint32 Port=0) override;
 
 	/**
 	 * @return Whether the machine has a properly configured network device or not
 	 */
-	virtual bool HasNetworkDevice() OVERRIDE;
+	virtual bool HasNetworkDevice() override;
 
 	/**
 	 *	Get the name of the socket subsystem
 	 * @return a string naming this subsystem
 	 */
-	virtual const TCHAR* GetSocketAPIName() const OVERRIDE;
+	virtual const TCHAR* GetSocketAPIName() const override;
 
-	virtual bool GetLocalAdapterAddresses( TArray<TSharedPtr<FInternetAddr> >& OutAdresses ) OVERRIDE;
+	virtual bool GetLocalAdapterAddresses( TArray<TSharedPtr<FInternetAddr> >& OutAdresses ) override;
 
 	/**
 	 * Returns the last error that has happened
 	 */
-	virtual ESocketErrors GetLastErrorCode() OVERRIDE;
+	virtual ESocketErrors GetLastErrorCode() override;
 
 	/**
 	 * Translates the platform error code to a ESocketErrors enum
 	 */
-	virtual ESocketErrors TranslateErrorCode(int32 Code) OVERRIDE;
+	virtual ESocketErrors TranslateErrorCode(int32 Code) override;
 };
\ No newline at end of file
diff --git a/Engine/Source/Runtime/Sockets/Private/Windows/SocketSubsystemWindows.h b/Engine/Source/Runtime/Sockets/Private/Windows/SocketSubsystemWindows.h
--- a/Engine/Source/Runtime/Sockets/Private/Windows/SocketSubsystemWindows.h
+++ b/Engine/Source/Runtime/Sockets/Private/Windows/SocketSubsystemWindows.h
@@ -30,19 +30,19 @@ public:
 
 	// Begin ISocketSubsystem interface
 
-	virtual bool HasNetworkDevice( ) OVERRIDE;
+	virtual bool HasNetworkDevice( ) override;
 
-	virtual ESocketErrors GetLastErrorCode( ) OVERRIDE;
+	virtual ESocketErrors GetLastErrorCode( ) override;
 
-	virtual bool GetLocalAdapterAddresses( TArray<TSharedPtr<FInternetAddr> >& OutAdresses ) OVERRIDE;
+	virtual bool GetLocalAdapterAddresses( TArray<TSharedPtr<FInternetAddr> >& OutAdresses ) override;
 
-	virtual const TCHAR* GetSocketAPIName( ) const OVERRIDE;
+	virtual const TCHAR* GetSocketAPIName( ) const override;
 
-	virtual bool Init( FString& Error ) OVERRIDE;
+	virtual bool Init( FString& Error ) override;
 
-	virtual void Shutdown( ) OVERRIDE;
+	virtual void Shutdown( ) override;
 
-	virtual ESocketErrors TranslateErrorCode( int32 Code ) OVERRIDE;
+	virtual ESocketErrors TranslateErrorCode( int32 Code ) override;
 
 	// End ISocketSubsystem interface
 
diff --git a/Engine/Source/Runtime/Sockets/Public/SocketSubsystemModule.h b/Engine/Source/Runtime/Sockets/Public/SocketSubsystemModule.h
--- a/Engine/Source/Runtime/Sockets/Public/SocketSubsystemModule.h
+++ b/Engine/Source/Runtime/Sockets/Public/SocketSubsystemModule.h
@@ -60,20 +60,20 @@ public:
 	 * Called right after the module DLL has been loaded and the module object has been created
 	 * Overloaded to allow the default subsystem a chance to load
 	 */
-	virtual void StartupModule() OVERRIDE;
+	virtual void StartupModule() override;
 
 	/**
 	 * Called before the module is unloaded, right before the module object is destroyed.
 	 * Overloaded to shut down all loaded online subsystems
 	 */
-	virtual void ShutdownModule() OVERRIDE;
+	virtual void ShutdownModule() override;
 
 	/**
 	 * Override this to set whether your module is allowed to be unloaded on the fly
 	 *
 	 * @return	Whether the module supports shutdown separate from the rest of the engine.
 	 */
-	virtual bool SupportsDynamicReloading() OVERRIDE
+	virtual bool SupportsDynamicReloading() override
 	{
 		return false;
 	}
@@ -83,7 +83,7 @@ public:
 	 *
 	 * @return	Whether the module supports shutdown on application exit
 	 */
-	virtual bool SupportsAutomaticShutdown() OVERRIDE
+	virtual bool SupportsAutomaticShutdown() override
 	{
 		return false;
 	}
diff --git a/Engine/Source/Runtime/StreamingFile/Private/StreamingNetworkPlatformFile.cpp b/Engine/Source/Runtime/StreamingFile/Private/StreamingNetworkPlatformFile.cpp
--- a/Engine/Source/Runtime/StreamingFile/Private/StreamingNetworkPlatformFile.cpp
+++ b/Engine/Source/Runtime/StreamingFile/Private/StreamingNetworkPlatformFile.cpp
@@ -69,15 +69,15 @@ public:
 		Network.SendCloseMessage(HandleId);
 	}
 
-	virtual int64		Size() OVERRIDE
+	virtual int64		Size() override
 	{
 		return FileSize;
 	}
-	virtual int64		Tell() OVERRIDE
+	virtual int64		Tell() override
 	{
 		return FilePos;
 	}
-	virtual bool		Seek(int64 NewPosition) OVERRIDE
+	virtual bool		Seek(int64 NewPosition) override
 	{
 		if (NewPosition >= 0 && NewPosition <= FileSize)
 		{
@@ -113,11 +113,11 @@ public:
 		}
 		return false;
 	}
-	virtual bool		SeekFromEnd(int64 NewPositionRelativeToEnd = 0) OVERRIDE
+	virtual bool		SeekFromEnd(int64 NewPositionRelativeToEnd = 0) override
 	{
 		return Seek(FileSize + NewPositionRelativeToEnd);
 	}
-	virtual bool		Read(uint8* Destination, int64 BytesToRead) OVERRIDE
+	virtual bool		Read(uint8* Destination, int64 BytesToRead) override
 	{
 		bool Result = false;
 		if (bReadable && BytesToRead >= 0 && BytesToRead + FilePos <= FileSize)
@@ -212,7 +212,7 @@ public:
 		}
 		return Result;
 	}
-	virtual bool		Write(const uint8* Source, int64 BytesToWrite) OVERRIDE
+	virtual bool		Write(const uint8* Source, int64 BytesToWrite) override
 	{
 		bool Result = false;
 		if (bWritable && BytesToWrite >= 0)
@@ -798,7 +798,7 @@ void FStreamingNetworkPlatformFile::PerformHeartbeat()
 class FStreamingFileModule : public IPlatformFileModule
 {
 public:
-	virtual IPlatformFile* GetPlatformFile() OVERRIDE
+	virtual IPlatformFile* GetPlatformFile() override
 	{
 		static TScopedPointer<IPlatformFile> AutoDestroySingleton(new FStreamingNetworkPlatformFile());
 		return AutoDestroySingleton.GetOwnedPointer();
diff --git a/Engine/Source/Runtime/StreamingFile/Public/StreamingNetworkPlatformFile.h b/Engine/Source/Runtime/StreamingFile/Public/StreamingNetworkPlatformFile.h
--- a/Engine/Source/Runtime/StreamingFile/Public/StreamingNetworkPlatformFile.h
+++ b/Engine/Source/Runtime/StreamingFile/Public/StreamingNetworkPlatformFile.h
@@ -118,7 +118,7 @@ class STREAMINGFILE_API FStreamingNetworkPlatformFile : public FNetworkPlatformF
 	friend class FAsyncFileSync;
 
 	// FNetworkPlatformFile interface
-	virtual bool InitializeInternal(IPlatformFile* Inner, const TCHAR* HostIP) OVERRIDE;
+	virtual bool InitializeInternal(IPlatformFile* Inner, const TCHAR* HostIP) override;
 
 public:
 
@@ -134,56 +134,56 @@ public:
 	virtual ~FStreamingNetworkPlatformFile();
 
 	// need to override what FNetworkPlatformFile does here
-	void InitializeAfterSetActive() OVERRIDE { }
+	void InitializeAfterSetActive() override { }
 
 	// Begin IPlatformFile Interface 
-	virtual bool ShouldBeUsed(IPlatformFile* Inner, const TCHAR* CmdLine) const OVERRIDE;
+	virtual bool ShouldBeUsed(IPlatformFile* Inner, const TCHAR* CmdLine) const override;
 
-	virtual IPlatformFile* GetLowerLevel() OVERRIDE
+	virtual IPlatformFile* GetLowerLevel() override
 	{
 		return NULL;
 	}
 
-	virtual const TCHAR* GetName() const OVERRIDE
+	virtual const TCHAR* GetName() const override
 	{
 		return FStreamingNetworkPlatformFile::GetTypeName();
 	}
-	virtual bool		DeleteFile(const TCHAR* Filename) OVERRIDE;
-	virtual bool		IsReadOnly(const TCHAR* Filename) OVERRIDE
+	virtual bool		DeleteFile(const TCHAR* Filename) override;
+	virtual bool		IsReadOnly(const TCHAR* Filename) override
 	{
 		FFileInfo Info;
 		GetFileInfo(Filename, Info);
 		return Info.ReadOnly;
 	}
-	virtual bool		MoveFile(const TCHAR* To, const TCHAR* From) OVERRIDE;
-	virtual bool		SetReadOnly(const TCHAR* Filename, bool bNewReadOnlyValue) OVERRIDE;
-	virtual FDateTime	GetTimeStamp(const TCHAR* Filename) OVERRIDE
+	virtual bool		MoveFile(const TCHAR* To, const TCHAR* From) override;
+	virtual bool		SetReadOnly(const TCHAR* Filename, bool bNewReadOnlyValue) override;
+	virtual FDateTime	GetTimeStamp(const TCHAR* Filename) override
 	{
 		FFileInfo Info;
 		GetFileInfo(Filename, Info);
 		return Info.TimeStamp;
 	}
-	virtual void		SetTimeStamp(const TCHAR* Filename, FDateTime DateTime) OVERRIDE;
-	virtual FDateTime	GetAccessTimeStamp(const TCHAR* Filename) OVERRIDE
+	virtual void		SetTimeStamp(const TCHAR* Filename, FDateTime DateTime) override;
+	virtual FDateTime	GetAccessTimeStamp(const TCHAR* Filename) override
 	{
 		FFileInfo Info;
 		GetFileInfo(Filename, Info);
 		return Info.AccessTimeStamp;
 	}
-	virtual IFileHandle*	OpenRead(const TCHAR* Filename) OVERRIDE;
-	virtual IFileHandle*	OpenWrite(const TCHAR* Filename, bool bAppend, bool bAllowRead) OVERRIDE;
-	virtual bool		DirectoryExists(const TCHAR* Directory) OVERRIDE;
-	virtual bool		CreateDirectoryTree(const TCHAR* Directory) OVERRIDE;
-	virtual bool		CreateDirectory(const TCHAR* Directory) OVERRIDE;
-	virtual bool		DeleteDirectory(const TCHAR* Directory) OVERRIDE;
+	virtual IFileHandle*	OpenRead(const TCHAR* Filename) override;
+	virtual IFileHandle*	OpenWrite(const TCHAR* Filename, bool bAppend, bool bAllowRead) override;
+	virtual bool		DirectoryExists(const TCHAR* Directory) override;
+	virtual bool		CreateDirectoryTree(const TCHAR* Directory) override;
+	virtual bool		CreateDirectory(const TCHAR* Directory) override;
+	virtual bool		DeleteDirectory(const TCHAR* Directory) override;
 
-	virtual bool		IterateDirectory(const TCHAR* Directory, IPlatformFile::FDirectoryVisitor& Visitor) OVERRIDE;
-	virtual bool		IterateDirectoryRecursively(const TCHAR* Directory, IPlatformFile::FDirectoryVisitor& Visitor) OVERRIDE;
-	virtual bool		DeleteDirectoryRecursively(const TCHAR* Directory) OVERRIDE;
-	virtual bool		CopyFile(const TCHAR* To, const TCHAR* From) OVERRIDE;
+	virtual bool		IterateDirectory(const TCHAR* Directory, IPlatformFile::FDirectoryVisitor& Visitor) override;
+	virtual bool		IterateDirectoryRecursively(const TCHAR* Directory, IPlatformFile::FDirectoryVisitor& Visitor) override;
+	virtual bool		DeleteDirectoryRecursively(const TCHAR* Directory) override;
+	virtual bool		CopyFile(const TCHAR* To, const TCHAR* From) override;
 
-	virtual FString ConvertToAbsolutePathForExternalAppForRead( const TCHAR* Filename ) OVERRIDE;
-	virtual FString ConvertToAbsolutePathForExternalAppForWrite( const TCHAR* Filename ) OVERRIDE;
+	virtual FString ConvertToAbsolutePathForExternalAppForRead( const TCHAR* Filename ) override;
+	virtual FString ConvertToAbsolutePathForExternalAppForWrite( const TCHAR* Filename ) override;
 
 	// End IPlatformFile Interface
 
@@ -206,9 +206,9 @@ private:
 
 	// Begin FNetworkPlatformFile interface
 
-	virtual void PerformHeartbeat() OVERRIDE;
+	virtual void PerformHeartbeat() override;
 
-	virtual void GetFileInfo(const TCHAR* Filename, FFileInfo& Info) OVERRIDE;
+	virtual void GetFileInfo(const TCHAR* Filename, FFileInfo& Info) override;
 
 	// End FNetworkPlatformFile interface
 
diff --git a/Engine/Source/Runtime/UMG/Private/UserWidget.cpp b/Engine/Source/Runtime/UMG/Private/UserWidget.cpp
--- a/Engine/Source/Runtime/UMG/Private/UserWidget.cpp
+++ b/Engine/Source/Runtime/UMG/Private/UserWidget.cpp
@@ -21,12 +21,12 @@ class SViewportWidgetHost : public SCompoundWidget
 			];
 	}
 
-	virtual bool OnHitTest(const FGeometry& MyGeometry, FVector2D InAbsoluteCursorPosition) OVERRIDE
+	virtual bool OnHitTest(const FGeometry& MyGeometry, FVector2D InAbsoluteCursorPosition) override
 	{
 		return true;
 	}
 
-	virtual bool SupportsKeyboardFocus() const OVERRIDE
+	virtual bool SupportsKeyboardFocus() const override
 	{
 		return true;
 	}
diff --git a/Engine/Source/Runtime/UMG/Public/Blueprint/UserWidget.h b/Engine/Source/Runtime/UMG/Public/Blueprint/UserWidget.h
--- a/Engine/Source/Runtime/UMG/Public/Blueprint/UserWidget.h
+++ b/Engine/Source/Runtime/UMG/Public/Blueprint/UserWidget.h
@@ -144,8 +144,8 @@ class UMG_API UUserWidget : public UWidget
 	TArray<UWidget*> Components;
 	
 	//UObject interface
-	virtual void PostInitProperties() OVERRIDE;
-	virtual class UWorld* GetWorld() const OVERRIDE;
+	virtual void PostInitProperties() override;
+	virtual class UWorld* GetWorld() const override;
 	// End of UObject interface
 
 	/*  */
@@ -241,7 +241,7 @@ class UMG_API UUserWidget : public UWidget
 	UWidget* GetHandleFromName(const FString& Name) const;
 
 protected:
-	virtual TSharedRef<SWidget> RebuildWidget() OVERRIDE;
+	virtual TSharedRef<SWidget> RebuildWidget() override;
 
 private:
 	TSharedPtr<SWidget> UserRootWidget;
diff --git a/Engine/Source/Runtime/UMG/Public/Blueprint/WidgetBlueprint.h b/Engine/Source/Runtime/UMG/Public/Blueprint/WidgetBlueprint.h
--- a/Engine/Source/Runtime/UMG/Public/Blueprint/WidgetBlueprint.h
+++ b/Engine/Source/Runtime/UMG/Public/Blueprint/WidgetBlueprint.h
@@ -54,14 +54,14 @@ public:
 	UPROPERTY()
 	TArray< UMovieScene* > AnimationData;
 	 
-	virtual void PostLoad() OVERRIDE;
-	virtual void PostLoadSubobjects(FObjectInstancingGraph* OuterInstanceGraph) OVERRIDE;
+	virtual void PostLoad() override;
+	virtual void PostLoadSubobjects(FObjectInstancingGraph* OuterInstanceGraph) override;
 	
 #if WITH_EDITOR
 	// UBlueprint interface
-	virtual UClass* GetBlueprintClass() const OVERRIDE;
+	virtual UClass* GetBlueprintClass() const override;
 
-	virtual bool SupportedByDefaultBlueprintFactory() const OVERRIDE
+	virtual bool SupportedByDefaultBlueprintFactory() const override
 	{
 		return false;
 	}
diff --git a/Engine/Source/Runtime/UMG/Public/Blueprint/WidgetBlueprintGeneratedClass.h b/Engine/Source/Runtime/UMG/Public/Blueprint/WidgetBlueprintGeneratedClass.h
--- a/Engine/Source/Runtime/UMG/Public/Blueprint/WidgetBlueprintGeneratedClass.h
+++ b/Engine/Source/Runtime/UMG/Public/Blueprint/WidgetBlueprintGeneratedClass.h
@@ -45,7 +45,7 @@ class UMG_API UWidgetBlueprintGeneratedClass : public UBlueprintGeneratedClass
 	/** This is transient data calculated at link time. */
 	TArray<UStructProperty*> WidgetNodeProperties;
 
-	virtual void Link(FArchive& Ar, bool bRelinkExistingProperties) OVERRIDE;
+	virtual void Link(FArchive& Ar, bool bRelinkExistingProperties) override;
 
 	void InitializeWidget(class UUserWidget* Actor) const;
 };
\ No newline at end of file
diff --git a/Engine/Source/Runtime/UMG/Public/Components/Border.h b/Engine/Source/Runtime/UMG/Public/Components/Border.h
--- a/Engine/Source/Runtime/UMG/Public/Components/Border.h
+++ b/Engine/Source/Runtime/UMG/Public/Components/Border.h
@@ -62,11 +62,11 @@ class UMG_API UBorder : public UContentWidget
 	bool bShowEffectWhenDisabled;
 
 	// UContentWidget interface
-	virtual void SetContent(UWidget* Content) OVERRIDE;
+	virtual void SetContent(UWidget* Content) override;
 	// End UContentWidget interface
 
 	// UWidget interface
-	void SyncronizeProperties() OVERRIDE;
+	void SyncronizeProperties() override;
 	// End of UWidget interface
 
 protected:
@@ -74,7 +74,7 @@ protected:
 
 protected:
 	// UWidget interface
-	virtual TSharedRef<SWidget> RebuildWidget() OVERRIDE;
+	virtual TSharedRef<SWidget> RebuildWidget() override;
 	// End of UWidget interface
 
 	const FSlateBrush* GetBorderBrush() const;
diff --git a/Engine/Source/Runtime/UMG/Public/Components/Button.h b/Engine/Source/Runtime/UMG/Public/Components/Button.h
--- a/Engine/Source/Runtime/UMG/Public/Components/Button.h
+++ b/Engine/Source/Runtime/UMG/Public/Components/Button.h
@@ -68,16 +68,16 @@ public:
 	void SetContentPadding(FMargin InContentPadding);
 
 	// UContentWidget interface
-	virtual void SetContent(UWidget* Content) OVERRIDE;
+	virtual void SetContent(UWidget* Content) override;
 	// End UContentWidget interface
 
 	// UWidget interface
-	void SyncronizeProperties() OVERRIDE;
+	void SyncronizeProperties() override;
 	// End of UWidget interface
 
 protected:
 	// UWidget interface
-	virtual TSharedRef<SWidget> RebuildWidget() OVERRIDE;
+	virtual TSharedRef<SWidget> RebuildWidget() override;
 	// End of UWidget interface
 
 	FReply HandleOnClicked();
diff --git a/Engine/Source/Runtime/UMG/Public/Components/CanvasPanel.h b/Engine/Source/Runtime/UMG/Public/Components/CanvasPanel.h
--- a/Engine/Source/Runtime/UMG/Public/Components/CanvasPanel.h
+++ b/Engine/Source/Runtime/UMG/Public/Components/CanvasPanel.h
@@ -20,13 +20,13 @@ class UMG_API UCanvasPanel : public UPanelWidget
 	UCanvasPanelSlot* AddSlot(UWidget* Content);
 
 	// UPanelWidget
-	virtual int32 GetChildrenCount() const OVERRIDE;
-	virtual UWidget* GetChildAt(int32 Index) const OVERRIDE;
-	virtual int32 GetChildIndex(UWidget* Content) const OVERRIDE;
-	virtual bool AddChild(UWidget* Child, FVector2D Position) OVERRIDE;
-	virtual bool RemoveChild(UWidget* Child) OVERRIDE;
-	virtual void ReplaceChildAt(int32 Index, UWidget* Child) OVERRIDE;
-	virtual void InsertChildAt(int32 Index, UWidget* Child) OVERRIDE;
+	virtual int32 GetChildrenCount() const override;
+	virtual UWidget* GetChildAt(int32 Index) const override;
+	virtual int32 GetChildIndex(UWidget* Content) const override;
+	virtual bool AddChild(UWidget* Child, FVector2D Position) override;
+	virtual bool RemoveChild(UWidget* Child) override;
+	virtual void ReplaceChildAt(int32 Index, UWidget* Child) override;
+	virtual void InsertChildAt(int32 Index, UWidget* Child) override;
 	// End UPanelWidget
 
 	/** Gets the underlying native canvas widget if it has been constructed */
@@ -37,11 +37,11 @@ class UMG_API UCanvasPanel : public UPanelWidget
 
 #if WITH_EDITOR
 	// UObject interface
-	virtual void PostEditChangeProperty(struct FPropertyChangedEvent& PropertyChangedEvent) OVERRIDE;
+	virtual void PostEditChangeProperty(struct FPropertyChangedEvent& PropertyChangedEvent) override;
 	// End of UObject interface
 
 	// UWidget interface
-	virtual void ConnectEditorData() OVERRIDE;
+	virtual void ConnectEditorData() override;
 	// End UWidget interface
 #endif
 
@@ -51,6 +51,6 @@ protected:
 
 protected:
 	// UWidget interface
-	virtual TSharedRef<SWidget> RebuildWidget() OVERRIDE;
+	virtual TSharedRef<SWidget> RebuildWidget() override;
 	// End of UWidget interface
 };
diff --git a/Engine/Source/Runtime/UMG/Public/Components/CanvasPanelSlot.h b/Engine/Source/Runtime/UMG/Public/Components/CanvasPanelSlot.h
--- a/Engine/Source/Runtime/UMG/Public/Components/CanvasPanelSlot.h
+++ b/Engine/Source/Runtime/UMG/Public/Components/CanvasPanelSlot.h
@@ -36,9 +36,9 @@ class UMG_API UCanvasPanelSlot : public UPanelSlot
 
 	void BuildSlot(TSharedRef<SConstraintCanvas> Canvas);
 
-	virtual void Resize(const FVector2D& Direction, const FVector2D& Amount) OVERRIDE;
+	virtual void Resize(const FVector2D& Direction, const FVector2D& Amount) override;
 
-	virtual bool CanResize(const FVector2D& Direction) const OVERRIDE;
+	virtual bool CanResize(const FVector2D& Direction) const override;
 
 	/** Sets the position of the slot */
 	UFUNCTION(BlueprintCallable, Category="Appearance")
@@ -53,12 +53,12 @@ class UMG_API UCanvasPanelSlot : public UPanelSlot
 	void SetAlignment(FVector2D InAlignment);
 
 	// UPanelSlot interface
-	virtual void Refresh() OVERRIDE;
+	virtual void Refresh() override;
 	// End of UPanelSlot interface
 
 #if WITH_EDITOR
 	// UObject interface
-	virtual void PreEditChange(UProperty* PropertyAboutToChange) OVERRIDE;
+	virtual void PreEditChange(UProperty* PropertyAboutToChange) override;
 	virtual void PostEditChangeChainProperty(struct FPropertyChangedChainEvent& PropertyChangedEvent);
 	// End of UObject interface
 #endif
diff --git a/Engine/Source/Runtime/UMG/Public/Components/CheckBox.h b/Engine/Source/Runtime/UMG/Public/Components/CheckBox.h
--- a/Engine/Source/Runtime/UMG/Public/Components/CheckBox.h
+++ b/Engine/Source/Runtime/UMG/Public/Components/CheckBox.h
@@ -57,7 +57,7 @@ protected:
 
 protected:
 	// UWidget interface
-	virtual TSharedRef<SWidget> RebuildWidget() OVERRIDE;
+	virtual TSharedRef<SWidget> RebuildWidget() override;
 	// End of UWidget interface
 
 	ESlateCheckBoxState::Type GetCheckState() const;
diff --git a/Engine/Source/Runtime/UMG/Public/Components/CircularThrobberComponent.h b/Engine/Source/Runtime/UMG/Public/Components/CircularThrobberComponent.h
--- a/Engine/Source/Runtime/UMG/Public/Components/CircularThrobberComponent.h
+++ b/Engine/Source/Runtime/UMG/Public/Components/CircularThrobberComponent.h
@@ -32,6 +32,6 @@ protected:
 
 protected:
 	// UWidget interface
-	virtual TSharedRef<SWidget> RebuildWidget() OVERRIDE;
+	virtual TSharedRef<SWidget> RebuildWidget() override;
 	// End of UWidget interface
 };
diff --git a/Engine/Source/Runtime/UMG/Public/Components/ContentWidget.h b/Engine/Source/Runtime/UMG/Public/Components/ContentWidget.h
--- a/Engine/Source/Runtime/UMG/Public/Components/ContentWidget.h
+++ b/Engine/Source/Runtime/UMG/Public/Components/ContentWidget.h
@@ -16,11 +16,11 @@ class UMG_API UContentWidget : public UPanelWidget
 	virtual void SetContent(UWidget* InContent);
 
 	// UPanelWidget
-	virtual int32 GetChildrenCount() const OVERRIDE;
-	virtual UWidget* GetChildAt(int32 Index) const OVERRIDE;
-	virtual bool AddChild(UWidget* Child, FVector2D Position) OVERRIDE;
-	virtual bool RemoveChild(UWidget* Child) OVERRIDE;
-	virtual void ReplaceChildAt(int32 Index, UWidget* Child) OVERRIDE;
+	virtual int32 GetChildrenCount() const override;
+	virtual UWidget* GetChildAt(int32 Index) const override;
+	virtual bool AddChild(UWidget* Child, FVector2D Position) override;
+	virtual bool RemoveChild(UWidget* Child) override;
+	virtual void ReplaceChildAt(int32 Index, UWidget* Child) override;
 	// End UPanelWidget
 
 protected:
diff --git a/Engine/Source/Runtime/UMG/Public/Components/EditableText.h b/Engine/Source/Runtime/UMG/Public/Components/EditableText.h
--- a/Engine/Source/Runtime/UMG/Public/Components/EditableText.h
+++ b/Engine/Source/Runtime/UMG/Public/Components/EditableText.h
@@ -89,12 +89,12 @@ protected:
 
 protected:
 	// UWidget interface
-	virtual TSharedRef<SWidget> RebuildWidget() OVERRIDE;
+	virtual TSharedRef<SWidget> RebuildWidget() override;
 	// End of UWidget
 
 #if WITH_EDITOR
 	// UObject interface
-	virtual void PostEditChangeProperty(struct FPropertyChangedEvent& PropertyChangedEvent) OVERRIDE;
+	virtual void PostEditChangeProperty(struct FPropertyChangedEvent& PropertyChangedEvent) override;
 	// End of UObject interface
 #endif
 
diff --git a/Engine/Source/Runtime/UMG/Public/Components/HorizontalBox.h b/Engine/Source/Runtime/UMG/Public/Components/HorizontalBox.h
--- a/Engine/Source/Runtime/UMG/Public/Components/HorizontalBox.h
+++ b/Engine/Source/Runtime/UMG/Public/Components/HorizontalBox.h
@@ -16,22 +16,22 @@ class UMG_API UHorizontalBox : public UPanelWidget
 	UHorizontalBoxSlot* AddSlot(UWidget* Content);
 
 	// UPanelWidget
-	virtual int32 GetChildrenCount() const OVERRIDE;
-	virtual UWidget* GetChildAt(int32 Index) const OVERRIDE;
-	virtual int32 GetChildIndex(UWidget* Content) const OVERRIDE;
-	virtual bool AddChild(UWidget* Child, FVector2D Position) OVERRIDE;
-	virtual bool RemoveChild(UWidget* Child) OVERRIDE;
-	virtual void ReplaceChildAt(int32 Index, UWidget* Child) OVERRIDE;
-	virtual void InsertChildAt(int32 Index, UWidget* Child) OVERRIDE;
+	virtual int32 GetChildrenCount() const override;
+	virtual UWidget* GetChildAt(int32 Index) const override;
+	virtual int32 GetChildIndex(UWidget* Content) const override;
+	virtual bool AddChild(UWidget* Child, FVector2D Position) override;
+	virtual bool RemoveChild(UWidget* Child) override;
+	virtual void ReplaceChildAt(int32 Index, UWidget* Child) override;
+	virtual void InsertChildAt(int32 Index, UWidget* Child) override;
 	// End UPanelWidget
 
 #if WITH_EDITOR
 	// UObject interface
-	virtual void PostEditChangeProperty(struct FPropertyChangedEvent& PropertyChangedEvent) OVERRIDE;
+	virtual void PostEditChangeProperty(struct FPropertyChangedEvent& PropertyChangedEvent) override;
 	// End of UObject interface
 
 	// UWidget interface
-	virtual void ConnectEditorData() OVERRIDE;
+	virtual void ConnectEditorData() override;
 	// End UWidget interface
 #endif
 
@@ -41,6 +41,6 @@ protected:
 
 protected:
 	// UWidget interface
-	virtual TSharedRef<SWidget> RebuildWidget() OVERRIDE;
+	virtual TSharedRef<SWidget> RebuildWidget() override;
 	// End of UWidget interface
 };
diff --git a/Engine/Source/Runtime/UMG/Public/Components/HorizontalBoxSlot.h b/Engine/Source/Runtime/UMG/Public/Components/HorizontalBoxSlot.h
--- a/Engine/Source/Runtime/UMG/Public/Components/HorizontalBoxSlot.h
+++ b/Engine/Source/Runtime/UMG/Public/Components/HorizontalBoxSlot.h
@@ -40,7 +40,7 @@ class UMG_API UHorizontalBoxSlot : public UPanelSlot
 	void BuildSlot(TSharedRef<SHorizontalBox> HorizontalBox);
 
 	// UPanelSlot interface
-	virtual void Refresh() OVERRIDE;
+	virtual void Refresh() override;
 	// End of UPanelSlot interface
 
 private:
diff --git a/Engine/Source/Runtime/UMG/Public/Components/Overlay.h b/Engine/Source/Runtime/UMG/Public/Components/Overlay.h
--- a/Engine/Source/Runtime/UMG/Public/Components/Overlay.h
+++ b/Engine/Source/Runtime/UMG/Public/Components/Overlay.h
@@ -16,20 +16,20 @@ class UMG_API UOverlay : public UPanelWidget
 	UOverlaySlot* AddSlot(UWidget* Content);
 
 	// UPanelWidget
-	virtual int32 GetChildrenCount() const OVERRIDE;
-	virtual UWidget* GetChildAt(int32 Index) const OVERRIDE;
-	virtual bool AddChild(UWidget* Child, FVector2D Position) OVERRIDE;
-	virtual bool RemoveChild(UWidget* Child) OVERRIDE;
-	virtual void ReplaceChildAt(int32 Index, UWidget* Child) OVERRIDE;
+	virtual int32 GetChildrenCount() const override;
+	virtual UWidget* GetChildAt(int32 Index) const override;
+	virtual bool AddChild(UWidget* Child, FVector2D Position) override;
+	virtual bool RemoveChild(UWidget* Child) override;
+	virtual void ReplaceChildAt(int32 Index, UWidget* Child) override;
 	// End UPanelWidget
 
 #if WITH_EDITOR
 	// UObject interface
-	virtual void PostEditChangeProperty(struct FPropertyChangedEvent& PropertyChangedEvent) OVERRIDE;
+	virtual void PostEditChangeProperty(struct FPropertyChangedEvent& PropertyChangedEvent) override;
 	// End of UObject interface
 
 	// UWidget interface
-	virtual void ConnectEditorData() OVERRIDE;
+	virtual void ConnectEditorData() override;
 	// End UWidget interface
 #endif
 
@@ -39,6 +39,6 @@ protected:
 
 protected:
 	// UWidget interface
-	virtual TSharedRef<SWidget> RebuildWidget() OVERRIDE;
+	virtual TSharedRef<SWidget> RebuildWidget() override;
 	// End of UWidget interface
 };
diff --git a/Engine/Source/Runtime/UMG/Public/Components/OverlaySlot.h b/Engine/Source/Runtime/UMG/Public/Components/OverlaySlot.h
--- a/Engine/Source/Runtime/UMG/Public/Components/OverlaySlot.h
+++ b/Engine/Source/Runtime/UMG/Public/Components/OverlaySlot.h
@@ -36,7 +36,7 @@ class UMG_API UOverlaySlot : public UPanelSlot
 	void BuildSlot(TSharedRef<SHorizontalBox> HorizontalBox);
 
 	// UPanelSlot interface
-	virtual void Refresh() OVERRIDE;
+	virtual void Refresh() override;
 	// End of UPanelSlot interface
 
 	/** Builds the underlying FSlot for the Slate layout panel. */
diff --git a/Engine/Source/Runtime/UMG/Public/Components/PanelSlot.h b/Engine/Source/Runtime/UMG/Public/Components/PanelSlot.h
--- a/Engine/Source/Runtime/UMG/Public/Components/PanelSlot.h
+++ b/Engine/Source/Runtime/UMG/Public/Components/PanelSlot.h
@@ -33,7 +33,7 @@ public:
 
 #if WITH_EDITOR
 
-	virtual void PostEditChangeProperty(struct FPropertyChangedEvent& PropertyChangedEvent) OVERRIDE
+	virtual void PostEditChangeProperty(struct FPropertyChangedEvent& PropertyChangedEvent) override
 	{
 		Super::PostEditChangeProperty(PropertyChangedEvent);
 
diff --git a/Engine/Source/Runtime/UMG/Public/Components/TextBlock.h b/Engine/Source/Runtime/UMG/Public/Components/TextBlock.h
--- a/Engine/Source/Runtime/UMG/Public/Components/TextBlock.h
+++ b/Engine/Source/Runtime/UMG/Public/Components/TextBlock.h
@@ -59,12 +59,12 @@ public:
 	//SLATE_EVENT(FOnClicked, OnDoubleClicked)
 
 	// UWidget interface
-	void SyncronizeProperties() OVERRIDE;
+	void SyncronizeProperties() override;
 	// End of UWidget interface
 
 protected:
 	// UWidget interface
-	virtual TSharedRef<SWidget> RebuildWidget() OVERRIDE;
+	virtual TSharedRef<SWidget> RebuildWidget() override;
 	// End of UWidget interface
 
 protected:
diff --git a/Engine/Source/Runtime/UMG/Public/Components/UniformGridPanel.h b/Engine/Source/Runtime/UMG/Public/Components/UniformGridPanel.h
--- a/Engine/Source/Runtime/UMG/Public/Components/UniformGridPanel.h
+++ b/Engine/Source/Runtime/UMG/Public/Components/UniformGridPanel.h
@@ -32,20 +32,20 @@ class UMG_API UUniformGridPanel : public UPanelWidget
 	UUniformGridSlot* AddSlot(UWidget* Content);
 
 	// UPanelWidget
-	virtual int32 GetChildrenCount() const OVERRIDE;
-	virtual UWidget* GetChildAt(int32 Index) const OVERRIDE;
-	virtual bool AddChild(UWidget* Child, FVector2D Position) OVERRIDE;
-	virtual bool RemoveChild(UWidget* Child) OVERRIDE;
-	virtual void ReplaceChildAt(int32 Index, UWidget* Child) OVERRIDE;
+	virtual int32 GetChildrenCount() const override;
+	virtual UWidget* GetChildAt(int32 Index) const override;
+	virtual bool AddChild(UWidget* Child, FVector2D Position) override;
+	virtual bool RemoveChild(UWidget* Child) override;
+	virtual void ReplaceChildAt(int32 Index, UWidget* Child) override;
 	// End UPanelWidget
 
 #if WITH_EDITOR
 	// UObject interface
-	virtual void PostEditChangeProperty(struct FPropertyChangedEvent& PropertyChangedEvent) OVERRIDE;
+	virtual void PostEditChangeProperty(struct FPropertyChangedEvent& PropertyChangedEvent) override;
 	// End of UObject interface
 
 	// UWidget interface
-	virtual void ConnectEditorData() OVERRIDE;
+	virtual void ConnectEditorData() override;
 	// End UWidget interface
 #endif
 
@@ -55,6 +55,6 @@ protected:
 
 protected:
 	// UWidget interface
-	virtual TSharedRef<SWidget> RebuildWidget() OVERRIDE;
+	virtual TSharedRef<SWidget> RebuildWidget() override;
 	// End of UWidget interface
 };
diff --git a/Engine/Source/Runtime/UMG/Public/Components/Widget.h b/Engine/Source/Runtime/UMG/Public/Components/Widget.h
--- a/Engine/Source/Runtime/UMG/Public/Components/Widget.h
+++ b/Engine/Source/Runtime/UMG/Public/Components/Widget.h
@@ -134,7 +134,7 @@ public:
 	virtual void ConnectEditorData() { }
 
 	// UObject interface
-	virtual void PostEditChangeProperty(struct FPropertyChangedEvent& PropertyChangedEvent) OVERRIDE
+	virtual void PostEditChangeProperty(struct FPropertyChangedEvent& PropertyChangedEvent) override
 	{
 		Super::PostEditChangeProperty(PropertyChangedEvent);
 
diff --git a/Engine/Source/Runtime/UMG/Public/Slate/SConstraintCanvas.h b/Engine/Source/Runtime/UMG/Public/Slate/SConstraintCanvas.h
--- a/Engine/Source/Runtime/UMG/Public/Slate/SConstraintCanvas.h
+++ b/Engine/Source/Runtime/UMG/Public/Slate/SConstraintCanvas.h
@@ -120,13 +120,13 @@ public:
 
 	// Begin SWidget overrides
 
-	virtual void OnArrangeChildren( const FGeometry& AllottedGeometry, FArrangedChildren& ArrangedChildren ) const OVERRIDE;
+	virtual void OnArrangeChildren( const FGeometry& AllottedGeometry, FArrangedChildren& ArrangedChildren ) const override;
 	
-	virtual int32 OnPaint( const FGeometry& AllottedGeometry, const FSlateRect& MyClippingRect, FSlateWindowElementList& OutDrawElements, int32 LayerId, const FWidgetStyle& InWidgetStyle, bool bParentEnabled ) const OVERRIDE;
+	virtual int32 OnPaint( const FGeometry& AllottedGeometry, const FSlateRect& MyClippingRect, FSlateWindowElementList& OutDrawElements, int32 LayerId, const FWidgetStyle& InWidgetStyle, bool bParentEnabled ) const override;
 
-	virtual FVector2D ComputeDesiredSize() const OVERRIDE;
+	virtual FVector2D ComputeDesiredSize() const override;
 
-	virtual FChildren* GetChildren() OVERRIDE;
+	virtual FChildren* GetChildren() override;
 
 	// End SWidget overrides
 
diff --git a/Engine/Source/Runtime/UMG/Public/Slate/SObjectWidget.h b/Engine/Source/Runtime/UMG/Public/Slate/SObjectWidget.h
--- a/Engine/Source/Runtime/UMG/Public/Slate/SObjectWidget.h
+++ b/Engine/Source/Runtime/UMG/Public/Slate/SObjectWidget.h
@@ -17,44 +17,44 @@ class SObjectWidget : public SCompoundWidget, public FGCObject
 
 	void Construct(const FArguments& InArgs, class UUserWidget* InWidgetObject);
 
-	virtual void AddReferencedObjects(FReferenceCollector& Collector) OVERRIDE;
+	virtual void AddReferencedObjects(FReferenceCollector& Collector) override;
 
-	virtual int32 OnPaint(const FGeometry& AllottedGeometry, const FSlateRect& MyClippingRect, FSlateWindowElementList& OutDrawElements, int32 LayerId, const FWidgetStyle& InWidgetStyle, bool bParentEnabled) const OVERRIDE;
+	virtual int32 OnPaint(const FGeometry& AllottedGeometry, const FSlateRect& MyClippingRect, FSlateWindowElementList& OutDrawElements, int32 LayerId, const FWidgetStyle& InWidgetStyle, bool bParentEnabled) const override;
 
-	virtual FReply OnKeyboardFocusReceived(const FGeometry& MyGeometry, const FKeyboardFocusEvent& InKeyboardFocusEvent) OVERRIDE;
-	virtual void OnKeyboardFocusLost(const FKeyboardFocusEvent& InKeyboardFocusEvent) OVERRIDE;
-	virtual void OnKeyboardFocusChanging(const FWeakWidgetPath& PreviousFocusPath, const FWidgetPath& NewWidgetPath) OVERRIDE;
+	virtual FReply OnKeyboardFocusReceived(const FGeometry& MyGeometry, const FKeyboardFocusEvent& InKeyboardFocusEvent) override;
+	virtual void OnKeyboardFocusLost(const FKeyboardFocusEvent& InKeyboardFocusEvent) override;
+	virtual void OnKeyboardFocusChanging(const FWeakWidgetPath& PreviousFocusPath, const FWidgetPath& NewWidgetPath) override;
 	
-	virtual FReply OnKeyChar(const FGeometry& MyGeometry, const FCharacterEvent& InCharacterEvent) OVERRIDE;
-	virtual FReply OnPreviewKeyDown(const FGeometry& MyGeometry, const FKeyboardEvent& InKeyboardEvent) OVERRIDE;
-	virtual FReply OnKeyDown(const FGeometry& MyGeometry, const FKeyboardEvent& InKeyboardEvent) OVERRIDE;
-	virtual FReply OnKeyUp(const FGeometry& MyGeometry, const FKeyboardEvent& InKeyboardEvent) OVERRIDE;
+	virtual FReply OnKeyChar(const FGeometry& MyGeometry, const FCharacterEvent& InCharacterEvent) override;
+	virtual FReply OnPreviewKeyDown(const FGeometry& MyGeometry, const FKeyboardEvent& InKeyboardEvent) override;
+	virtual FReply OnKeyDown(const FGeometry& MyGeometry, const FKeyboardEvent& InKeyboardEvent) override;
+	virtual FReply OnKeyUp(const FGeometry& MyGeometry, const FKeyboardEvent& InKeyboardEvent) override;
 	
-	virtual FReply OnMouseButtonDown(const FGeometry& MyGeometry, const FPointerEvent& MouseEvent) OVERRIDE;
-	virtual FReply OnPreviewMouseButtonDown(const FGeometry& MyGeometry, const FPointerEvent& MouseEvent) OVERRIDE;
-	virtual FReply OnMouseButtonUp(const FGeometry& MyGeometry, const FPointerEvent& MouseEvent) OVERRIDE;
-	virtual FReply OnMouseMove(const FGeometry& MyGeometry, const FPointerEvent& MouseEvent) OVERRIDE;
-	virtual void OnMouseEnter(const FGeometry& MyGeometry, const FPointerEvent& MouseEvent) OVERRIDE;
-	virtual void OnMouseLeave(const FPointerEvent& MouseEvent) OVERRIDE;
-	virtual FReply OnMouseWheel(const FGeometry& MyGeometry, const FPointerEvent& MouseEvent) OVERRIDE;
-	virtual FCursorReply OnCursorQuery(const FGeometry& MyGeometry, const FPointerEvent& CursorEvent) const OVERRIDE;
-	virtual FReply OnMouseButtonDoubleClick(const FGeometry& InMyGeometry, const FPointerEvent& InMouseEvent) OVERRIDE;
+	virtual FReply OnMouseButtonDown(const FGeometry& MyGeometry, const FPointerEvent& MouseEvent) override;
+	virtual FReply OnPreviewMouseButtonDown(const FGeometry& MyGeometry, const FPointerEvent& MouseEvent) override;
+	virtual FReply OnMouseButtonUp(const FGeometry& MyGeometry, const FPointerEvent& MouseEvent) override;
+	virtual FReply OnMouseMove(const FGeometry& MyGeometry, const FPointerEvent& MouseEvent) override;
+	virtual void OnMouseEnter(const FGeometry& MyGeometry, const FPointerEvent& MouseEvent) override;
+	virtual void OnMouseLeave(const FPointerEvent& MouseEvent) override;
+	virtual FReply OnMouseWheel(const FGeometry& MyGeometry, const FPointerEvent& MouseEvent) override;
+	virtual FCursorReply OnCursorQuery(const FGeometry& MyGeometry, const FPointerEvent& CursorEvent) const override;
+	virtual FReply OnMouseButtonDoubleClick(const FGeometry& InMyGeometry, const FPointerEvent& InMouseEvent) override;
 	
-	virtual FReply OnDragDetected(const FGeometry& MyGeometry, const FPointerEvent& MouseEvent) OVERRIDE;
-	virtual void OnDragEnter(const FGeometry& MyGeometry, const FDragDropEvent& DragDropEvent) OVERRIDE;
-	virtual void OnDragLeave(const FDragDropEvent& DragDropEvent) OVERRIDE;
-	virtual FReply OnDragOver(const FGeometry& MyGeometry, const FDragDropEvent& DragDropEvent) OVERRIDE;
-	virtual FReply OnDrop(const FGeometry& MyGeometry, const FDragDropEvent& DragDropEvent) OVERRIDE;
+	virtual FReply OnDragDetected(const FGeometry& MyGeometry, const FPointerEvent& MouseEvent) override;
+	virtual void OnDragEnter(const FGeometry& MyGeometry, const FDragDropEvent& DragDropEvent) override;
+	virtual void OnDragLeave(const FDragDropEvent& DragDropEvent) override;
+	virtual FReply OnDragOver(const FGeometry& MyGeometry, const FDragDropEvent& DragDropEvent) override;
+	virtual FReply OnDrop(const FGeometry& MyGeometry, const FDragDropEvent& DragDropEvent) override;
 	
-	virtual FReply OnControllerButtonPressed(const FGeometry& MyGeometry, const FControllerEvent& ControllerEvent) OVERRIDE;
-	virtual FReply OnControllerButtonReleased(const FGeometry& MyGeometry, const FControllerEvent& ControllerEvent) OVERRIDE;
-	virtual FReply OnControllerAnalogValueChanged(const FGeometry& MyGeometry, const FControllerEvent& ControllerEvent) OVERRIDE;
+	virtual FReply OnControllerButtonPressed(const FGeometry& MyGeometry, const FControllerEvent& ControllerEvent) override;
+	virtual FReply OnControllerButtonReleased(const FGeometry& MyGeometry, const FControllerEvent& ControllerEvent) override;
+	virtual FReply OnControllerAnalogValueChanged(const FGeometry& MyGeometry, const FControllerEvent& ControllerEvent) override;
 	
-	virtual FReply OnTouchGesture(const FGeometry& MyGeometry, const FPointerEvent& GestureEvent) OVERRIDE;
-	virtual FReply OnTouchStarted(const FGeometry& MyGeometry, const FPointerEvent& InTouchEvent) OVERRIDE;
-	virtual FReply OnTouchMoved(const FGeometry& MyGeometry, const FPointerEvent& InTouchEvent) OVERRIDE;
-	virtual FReply OnTouchEnded(const FGeometry& MyGeometry, const FPointerEvent& InTouchEvent) OVERRIDE;
-	virtual FReply OnMotionDetected(const FGeometry& MyGeometry, const FMotionEvent& InMotionEvent) OVERRIDE;
+	virtual FReply OnTouchGesture(const FGeometry& MyGeometry, const FPointerEvent& GestureEvent) override;
+	virtual FReply OnTouchStarted(const FGeometry& MyGeometry, const FPointerEvent& InTouchEvent) override;
+	virtual FReply OnTouchMoved(const FGeometry& MyGeometry, const FPointerEvent& InTouchEvent) override;
+	virtual FReply OnTouchEnded(const FGeometry& MyGeometry, const FPointerEvent& InTouchEvent) override;
+	virtual FReply OnMotionDetected(const FGeometry& MyGeometry, const FMotionEvent& InMotionEvent) override;
 
 private:
 	/** The UWidget that created this SObjectWidget who needs to be kept alive. */
diff --git a/Engine/Source/Runtime/Windows/D3D11RHI/Private/D3D11RHIPrivate.h b/Engine/Source/Runtime/Windows/D3D11RHI/Private/D3D11RHIPrivate.h
--- a/Engine/Source/Runtime/Windows/D3D11RHI/Private/D3D11RHIPrivate.h
+++ b/Engine/Source/Runtime/Windows/D3D11RHI/Private/D3D11RHIPrivate.h
@@ -99,12 +99,12 @@ public:
 	/**
 	 * Initializes all D3D resources.
 	 */
-	virtual void InitDynamicRHI() OVERRIDE;
+	virtual void InitDynamicRHI() override;
 
 	/**
 	 * Releases all D3D resources.
 	 */
-	virtual void ReleaseDynamicRHI() OVERRIDE;
+	virtual void ReleaseDynamicRHI() override;
 
 private:
 	/**
@@ -142,12 +142,12 @@ public:
 	/**
 	 * Initializes all D3D resources.
 	 */
-	virtual void InitDynamicRHI() OVERRIDE;
+	virtual void InitDynamicRHI() override;
 
 	/**
 	 * Releases all D3D resources.
 	 */
-	virtual void ReleaseDynamicRHI() OVERRIDE;
+	virtual void ReleaseDynamicRHI() override;
 
 
 private:
@@ -178,15 +178,15 @@ public:
 	 * Returns the time in ms that the GPU spent in this draw event.  
 	 * This blocks the CPU if necessary, so can cause hitching.
 	 */
-	virtual float GetTiming() OVERRIDE;
+	virtual float GetTiming() override;
 
 
-	virtual void StartTiming() OVERRIDE
+	virtual void StartTiming() override
 	{
 		Timing.StartTiming();
 	}
 
-	virtual void StopTiming() OVERRIDE
+	virtual void StopTiming() override
 	{
 		Timing.EndTiming();
 	}
@@ -217,15 +217,15 @@ public:
 	}
 
 	/** Start this frame of per tracking */
-	virtual void StartFrame() OVERRIDE;
+	virtual void StartFrame() override;
 
 	/** End this frame of per tracking, but do not block yet */
-	virtual void EndFrame() OVERRIDE;
+	virtual void EndFrame() override;
 
 	/** Calculates root timing base frequency (if needed by this RHI) */
-	virtual float GetRootTimingResults() OVERRIDE;
+	virtual float GetRootTimingResults() override;
 
-	virtual void LogDisjointQuery() OVERRIDE;
+	virtual void LogDisjointQuery() override;
 
 	/** Timer tracking inclusive time spent in the root nodes. */
 	FD3D11BufferedGPUTiming RootEventTiming;
@@ -281,14 +281,14 @@ struct FD3DGPUProfiler : public FGPUProfiler
 		FrameTiming.InitResource();
 	}
 
-	virtual FGPUProfilerEventNode* CreateEventNode(const TCHAR* InName, FGPUProfilerEventNode* InParent) OVERRIDE
+	virtual FGPUProfilerEventNode* CreateEventNode(const TCHAR* InName, FGPUProfilerEventNode* InParent) override
 	{
 		FD3D11EventNode* EventNode = new FD3D11EventNode(InName, InParent, D3D11RHI);
 		return EventNode;
 	}
 
-	virtual void PushEvent(const TCHAR* Name) OVERRIDE;
-	virtual void PopEvent() OVERRIDE;
+	virtual void PushEvent(const TCHAR* Name) override;
+	virtual void PopEvent() override;
 
 	void BeginFrame(class FD3D11DynamicRHI* InRHI);
 
@@ -325,10 +325,10 @@ public:
 	virtual void InitD3DDevice();
 
 	// FDynamicRHI interface.
-	virtual void Init() OVERRIDE;
-	virtual void Shutdown() OVERRIDE;
-	virtual void PushEvent(const TCHAR* Name) OVERRIDE { GPUProfilingData.PushEvent(Name); }
-	virtual void PopEvent() OVERRIDE { GPUProfilingData.PopEvent(); }
+	virtual void Init() override;
+	virtual void Shutdown() override;
+	virtual void PushEvent(const TCHAR* Name) override { GPUProfilingData.PushEvent(Name); }
+	virtual void PopEvent() override { GPUProfilingData.PopEvent(); }
 
 	/**
 	 * Reads a D3D query's data into the provided buffer.
@@ -624,11 +624,11 @@ class FD3D11DynamicRHIModule : public IDynamicRHIModule
 {
 public:
 	// IModuleInterface
-	virtual bool SupportsDynamicReloading() OVERRIDE { return false; }
+	virtual bool SupportsDynamicReloading() override { return false; }
 
 	// IDynamicRHIModule
-	virtual bool IsSupported() OVERRIDE;
-	virtual FDynamicRHI* CreateRHI() OVERRIDE;
+	virtual bool IsSupported() override;
+	virtual FDynamicRHI* CreateRHI() override;
 
 private:
 	FD3D11Adapter ChosenAdapter;
@@ -730,13 +730,13 @@ class FVector4VertexDeclaration : public FRenderResource
 {
 public:
 	FVertexDeclarationRHIRef VertexDeclarationRHI;
-	virtual void InitRHI() OVERRIDE
+	virtual void InitRHI() override
 	{
 		FVertexDeclarationElementList Elements;
 		Elements.Add(FVertexElement(0,0,VET_Float4,0));
 		VertexDeclarationRHI = RHICreateVertexDeclaration(Elements);
 	}
-	virtual void ReleaseRHI() OVERRIDE
+	virtual void ReleaseRHI() override
 	{
 		VertexDeclarationRHI.SafeRelease();
 	}
diff --git a/Engine/Source/Runtime/Windows/D3D11RHI/Public/D3D11Resources.h b/Engine/Source/Runtime/Windows/D3D11RHI/Public/D3D11Resources.h
--- a/Engine/Source/Runtime/Windows/D3D11RHI/Public/D3D11Resources.h
+++ b/Engine/Source/Runtime/Windows/D3D11RHI/Public/D3D11Resources.h
@@ -361,7 +361,7 @@ public:
 	bool IsCubemap() const { return bCubemap; }
 
 	/** FRHITexture override.  See FRHITexture::GetNativeResource() */
-	virtual void* GetNativeResource() const OVERRIDE
+	virtual void* GetNativeResource() const override
 	{ 
 		return GetResource();
 	}
diff --git a/Engine/Source/Runtime/Windows/D3D11RHI/Public/D3D11Util.h b/Engine/Source/Runtime/Windows/D3D11RHI/Public/D3D11Util.h
--- a/Engine/Source/Runtime/Windows/D3D11RHI/Public/D3D11Util.h
+++ b/Engine/Source/Runtime/Windows/D3D11RHI/Public/D3D11Util.h
@@ -203,8 +203,8 @@ public:
 	ID3D11Buffer* Unlock();
 
 	// Begin FRenderResource interface.
-	virtual void InitRHI() OVERRIDE;
-	virtual void ReleaseRHI() OVERRIDE;
+	virtual void InitRHI() override;
+	virtual void ReleaseRHI() override;
 	// End FRenderResource interface.
 
 private:
diff --git a/Engine/Source/Runtime/Windows/XAudio2/Public/XAudio2Device.h b/Engine/Source/Runtime/Windows/XAudio2/Public/XAudio2Device.h
--- a/Engine/Source/Runtime/Windows/XAudio2/Public/XAudio2Device.h
+++ b/Engine/Source/Runtime/Windows/XAudio2/Public/XAudio2Device.h
@@ -77,21 +77,21 @@ class FXAudio2Device : public FAudioDevice
 {
 
 	/** Starts up any platform specific hardware/APIs */
-	virtual bool InitializeHardware() OVERRIDE;
+	virtual bool InitializeHardware() override;
 
 	/** Shuts down any platform specific hardware/APIs */
-	virtual void TeardownHardware() OVERRIDE;
+	virtual void TeardownHardware() override;
 
 	/** Lets the platform any tick actions */
-	virtual void UpdateHardware() OVERRIDE;
+	virtual void UpdateHardware() override;
 
 	/** Creates a new platform specific sound source */
-	virtual FAudioEffectsManager* CreateEffectsManager() OVERRIDE;
+	virtual FAudioEffectsManager* CreateEffectsManager() override;
 
 	/** Creates a new platform specific sound source */
-	virtual FSoundSource* CreateSoundSource() OVERRIDE;
+	virtual FSoundSource* CreateSoundSource() override;
 
-	virtual FName GetRuntimeFormat(USoundWave* SoundWave) OVERRIDE
+	virtual FName GetRuntimeFormat(USoundWave* SoundWave) override
 	{
 		static FName NAME_OPUS(TEXT("OPUS"));
 
@@ -108,19 +108,19 @@ class FXAudio2Device : public FAudioDevice
 #endif //WITH_OGGVORBIS
 	}
 
-	virtual bool HasCompressedAudioInfoClass(USoundWave* SoundWave) OVERRIDE;
+	virtual bool HasCompressedAudioInfoClass(USoundWave* SoundWave) override;
 
-	virtual bool SupportsRealtimeDecompression() const OVERRIDE
+	virtual bool SupportsRealtimeDecompression() const override
 	{
 		return true;
 	}
 
-	virtual class ICompressedAudioInfo* CreateCompressedAudioInfo(USoundWave* SoundWave) OVERRIDE;
+	virtual class ICompressedAudioInfo* CreateCompressedAudioInfo(USoundWave* SoundWave) override;
 
 	/** 
 	 * Check for errors and output a human readable string 
 	 */
-	virtual bool ValidateAPICall(const TCHAR* Function, int32 ErrorCode) OVERRIDE;
+	virtual bool ValidateAPICall(const TCHAR* Function, int32 ErrorCode) override;
 
 	/**
 	 * Exec handler used to parse console commands.
@@ -129,7 +129,7 @@ class FXAudio2Device : public FAudioDevice
 	 * @param	Ar		Output device to use in case the handler prints anything
 	 * @return	true if command was handled, false otherwise
 	 */
-	virtual bool Exec( UWorld* InWorld, const TCHAR* Cmd, FOutputDevice& Ar ) OVERRIDE;
+	virtual bool Exec( UWorld* InWorld, const TCHAR* Cmd, FOutputDevice& Ar ) override;
 
 protected:
 
