UMG - Fixing a non-unity issue.

diff --git a/Engine/Source/Runtime/UMG/Public/UMG.h b/Engine/Source/Runtime/UMG/Public/UMG.h
--- a/Engine/Source/Runtime/UMG/Public/UMG.h
+++ b/Engine/Source/Runtime/UMG/Public/UMG.h
@@ -16,6 +16,7 @@
 #include "Anchors.h"
 #include "SConstraintCanvas.h"
 
+#include "Slate/SlateBrushAsset.h"
 
 /* Public Includes
  *****************************************************************************/
