UMG - The setter functions on button match the properties now.

diff --git a/Engine/Source/Runtime/UMG/Private/Components/Button.cpp b/Engine/Source/Runtime/UMG/Private/Components/Button.cpp
--- a/Engine/Source/Runtime/UMG/Private/Components/Button.cpp
+++ b/Engine/Source/Runtime/UMG/Private/Components/Button.cpp
@@ -85,9 +85,14 @@ void UButton::SetContent(UWidget* InContent)
 	}
 }
 
-void UButton::SetButtonColorAndOpacity(FLinearColor InButtonColorAndOpacity)
+void UButton::SetColorAndOpacity(FLinearColor Color)
 {
-	MyButton->SetBorderBackgroundColor(InButtonColorAndOpacity);
+	MyButton->SetColorAndOpacity(Color);
+}
+
+void UButton::SetBackgroundColor(FLinearColor Color)
+{
+	MyButton->SetBorderBackgroundColor(Color);
 }
 
 void UButton::SetForegroundColor(FLinearColor InForegroundColor)
diff --git a/Engine/Source/Runtime/UMG/Public/Components/Button.h b/Engine/Source/Runtime/UMG/Public/Components/Button.h
--- a/Engine/Source/Runtime/UMG/Public/Components/Button.h
+++ b/Engine/Source/Runtime/UMG/Public/Components/Button.h
@@ -61,7 +61,11 @@ public:
 	
 	/**  */
 	UFUNCTION(BlueprintCallable, Category="Appearance")
-	void SetButtonColorAndOpacity(FLinearColor InButtonColorAndOpacity);
+	void SetColorAndOpacity(FLinearColor InColorAndOpacity);
+
+	/**  */
+	UFUNCTION(BlueprintCallable, Category="Appearance")
+	void SetBackgroundColor(FLinearColor InBackgroundColor);
 
 	/**  */
 	UFUNCTION(BlueprintCallable, Category="Appearance")
