UMG - Fixing more undo issues.  Fixing a series of problems with the canvas resize/move handles that didn't work correctly in all circumstances.

diff --git a/Engine/Source/Editor/UMGEditor/Private/Extensions/CanvasSlotExtension.cpp b/Engine/Source/Editor/UMGEditor/Private/Extensions/CanvasSlotExtension.cpp
--- a/Engine/Source/Editor/UMGEditor/Private/Extensions/CanvasSlotExtension.cpp
+++ b/Engine/Source/Editor/UMGEditor/Private/Extensions/CanvasSlotExtension.cpp
@@ -95,6 +95,20 @@ void FCanvasSlotExtension::MoveByAmount(UWidget* Widget, FVector2D Delta)
 		Offsets.Left += Delta.X;
 		Offsets.Top += Delta.Y;
 
+		// If the slot is stretched horizontally we need to move the right side as it no longer represents width, but
+		// now represents margin from the right stretched side.
+		if ( CanvasSlot->LayoutData.Anchors.IsStretchedHorizontal() )
+		{
+			Offsets.Right -= Delta.X;
+		}
+
+		// If the slot is stretched vertically we need to move the bottom side as it no longer represents width, but
+		// now represents margin from the bottom stretched side.
+		if ( CanvasSlot->LayoutData.Anchors.IsStretchedVertical() )
+		{
+			Offsets.Bottom -= Delta.Y;
+		}
+
 		CanvasSlot->SetOffset(Offsets);
 	}
 }
diff --git a/Engine/Source/Runtime/UMG/Private/Components/CanvasPanelSlot.cpp b/Engine/Source/Runtime/UMG/Private/Components/CanvasPanelSlot.cpp
--- a/Engine/Source/Runtime/UMG/Private/Components/CanvasPanelSlot.cpp
+++ b/Engine/Source/Runtime/UMG/Private/Components/CanvasPanelSlot.cpp
@@ -11,7 +11,7 @@ UCanvasPanelSlot::UCanvasPanelSlot(const FPostConstructInitializeProperties& PCI
 {
 	LayoutData.Offsets = FMargin(0, 0, 1, 1);
 	LayoutData.Anchors = FAnchors(0.0f, 0.0f);
-	LayoutData.Alignment = FVector2D(0.5f, 0.5f);
+	LayoutData.Alignment = FVector2D(0.0f, 0.0f);
 }
 
 void UCanvasPanelSlot::BuildSlot(TSharedRef<SConstraintCanvas> Canvas)
@@ -30,20 +30,49 @@ void UCanvasPanelSlot::Resize(const FVector2D& Direction, const FVector2D& Amoun
 	if ( Direction.X < 0 )
 	{
 		LayoutData.Offsets.Left -= Amount.X * Direction.X;
-		LayoutData.Offsets.Right += Amount.X * Direction.X;
+
+		// Keep the right side stable if it's not stretched and we're resizing left.
+		if ( !LayoutData.Anchors.IsStretchedHorizontal() )
+		{
+			// TODO UMG This doesn't work correctly with Alignment, fix or remove alignment
+			LayoutData.Offsets.Right += Amount.X * Direction.X;
+		}
 	}
+	
 	if ( Direction.Y < 0 )
 	{
 		LayoutData.Offsets.Top -= Amount.Y * Direction.Y;
-		LayoutData.Offsets.Bottom += Amount.Y * Direction.Y;
+
+		// Keep the bottom side stable if it's not stretched and we're resizing top.
+		if ( !LayoutData.Anchors.IsStretchedVertical() )
+		{
+			// TODO UMG This doesn't work correctly with Alignment, fix or remove alignment
+			LayoutData.Offsets.Bottom += Amount.Y * Direction.Y;
+		}
 	}
+	
 	if ( Direction.X > 0 )
 	{
-		LayoutData.Offsets.Right += Amount.X * Direction.X;
+		if ( LayoutData.Anchors.IsStretchedHorizontal() )
+		{
+			LayoutData.Offsets.Right -= Amount.X * Direction.X;
+		}
+		else
+		{
+			LayoutData.Offsets.Right += Amount.X * Direction.X;
+		}
 	}
+
 	if ( Direction.Y > 0 )
 	{
-		LayoutData.Offsets.Bottom += Amount.Y * Direction.Y;
+		if ( LayoutData.Anchors.IsStretchedVertical() )
+		{
+			LayoutData.Offsets.Bottom -= Amount.Y * Direction.Y;
+		}
+		else
+		{
+			LayoutData.Offsets.Bottom += Amount.Y * Direction.Y;
+		}
 	}
 
 	if ( Slot )
diff --git a/Engine/Source/Runtime/UMG/Private/WidgetTree.cpp b/Engine/Source/Runtime/UMG/Private/WidgetTree.cpp
--- a/Engine/Source/Runtime/UMG/Private/WidgetTree.cpp
+++ b/Engine/Source/Runtime/UMG/Private/WidgetTree.cpp
@@ -87,6 +87,8 @@ bool UWidgetTree::RemoveWidgetRecursive(UWidget* InRemovedWidget)
 
 bool UWidgetTree::RemoveWidget(UWidget* InRemovedWidget)
 {
+	Modify();
+
 	bool bRemoved = RemoveWidgetRecursive(InRemovedWidget);
 
 	int32 IndexRemoved = WidgetTemplates.Remove(InRemovedWidget);
diff --git a/Engine/Source/Runtime/UMG/Public/Blueprint/WidgetTree.h b/Engine/Source/Runtime/UMG/Public/Blueprint/WidgetTree.h
--- a/Engine/Source/Runtime/UMG/Public/Blueprint/WidgetTree.h
+++ b/Engine/Source/Runtime/UMG/Public/Blueprint/WidgetTree.h
@@ -10,22 +10,27 @@ class UMG_API UWidgetTree : public UObject
 {
 	GENERATED_UCLASS_BODY()
 
+public:
 	/** Array of templates for widgets */
 	UPROPERTY()
-	TArray<class UWidget*> WidgetTemplates;
+	TArray<UWidget*> WidgetTemplates;
+
+public:
 
 	void RenameWidget(UWidget* Widget, FString& NewName);
 
-	class UWidget* FindWidget(FString& Name) const;
+	UWidget* FindWidget(FString& Name) const;
 
-	bool RemoveWidget(class UWidget* Widget);
+	bool RemoveWidget(UWidget* Widget);
 
-	class UPanelWidget* FindWidgetParent(class UWidget* Widget, int32& OutChildIndex);
+	class UPanelWidget* FindWidgetParent(UWidget* Widget, int32& OutChildIndex);
 
 	/** Constructs the widget, and adds it to the tree. */
 	template< class T >
-	T* ConstructWidget(TSubclassOf<class UWidget> WidgetType)
+	T* ConstructWidget(TSubclassOf<UWidget> WidgetType)
 	{
+		Modify();
+
 		UWidget* Widget = (UWidget*)ConstructObject<UWidget>(WidgetType, this);
 		Widget->SetFlags(RF_Transactional);
 		WidgetTemplates.Add(Widget);
@@ -35,5 +40,4 @@ class UMG_API UWidgetTree : public UObject
 	
 private:
 	bool RemoveWidgetRecursive(UWidget* InRemovedWidget);
-	
 };
