UMG - Adding a comment to explain loading the editor module from the runtime module.

diff --git a/Engine/Source/Runtime/UMG/Private/UMGModule.cpp b/Engine/Source/Runtime/UMG/Private/UMGModule.cpp
--- a/Engine/Source/Runtime/UMG/Private/UMGModule.cpp
+++ b/Engine/Source/Runtime/UMG/Private/UMGModule.cpp
@@ -17,6 +17,8 @@ public:
 	virtual void StartupModule() override
 	{
 #if WITH_EDITOR
+		// This is done so that the compiler is available in non-cooked builds when the widget blueprint is 
+		// compiled again in the running game.
 		if ( FParse::Param(FCommandLine::Get(), TEXT("umg")) )
 		{
 			FModuleManager::Get().LoadModule(TEXT("UMGEditor"));
