UMG - Adding the WidgetSwitcher.  Fixing some runtime slot creation issues.

diff --git a/Engine/Source/Runtime/Slate/Private/Widgets/Layout/SWidgetSwitcher.cpp b/Engine/Source/Runtime/Slate/Private/Widgets/Layout/SWidgetSwitcher.cpp
--- a/Engine/Source/Runtime/Slate/Private/Widgets/Layout/SWidgetSwitcher.cpp
+++ b/Engine/Source/Runtime/Slate/Private/Widgets/Layout/SWidgetSwitcher.cpp
@@ -28,6 +28,21 @@ FSimpleSlot& SWidgetSwitcher::AddSlot( int32 SlotIndex )
 }
 
 
+int32 SWidgetSwitcher::RemoveSlot(TSharedRef<SWidget> WidgetToRemove)
+{
+	for( int32 SlotIndex=0; SlotIndex < AllChildren.Num(); ++SlotIndex )
+	{
+		if ( AllChildren[SlotIndex].Widget == WidgetToRemove )
+		{
+			AllChildren.RemoveAt(SlotIndex);
+			return SlotIndex;
+		}
+	}
+
+	return -1;
+}
+
+
 void SWidgetSwitcher::Construct( const FArguments& InArgs )
 {
 	OneDynamicChild = FOneDynamicChild( &AllChildren, &WidgetIndex );
diff --git a/Engine/Source/Runtime/Slate/Public/Widgets/Layout/SWidgetSwitcher.h b/Engine/Source/Runtime/Slate/Public/Widgets/Layout/SWidgetSwitcher.h
--- a/Engine/Source/Runtime/Slate/Public/Widgets/Layout/SWidgetSwitcher.h
+++ b/Engine/Source/Runtime/Slate/Public/Widgets/Layout/SWidgetSwitcher.h
@@ -39,6 +39,13 @@ public:
 	 */
 	FSimpleSlot& AddSlot( int32 SlotIndex = INDEX_NONE );
 
+	/**
+	 * Removes a slot with the corresponding widget in it.  Returns the index where the widget was found, otherwise -1.
+	 *
+	 * @param Widget - The widget to find and remove.
+	 */
+	int32 RemoveSlot( TSharedRef<SWidget> WidgetToRemove );
+
 	/**
 	 * Constructs the widget.
 	 */
diff --git a/Engine/Source/Runtime/UMG/Private/Components/VerticalBox.cpp b/Engine/Source/Runtime/UMG/Private/Components/VerticalBox.cpp
--- a/Engine/Source/Runtime/UMG/Private/Components/VerticalBox.cpp
+++ b/Engine/Source/Runtime/UMG/Private/Components/VerticalBox.cpp
@@ -39,13 +39,6 @@ int32 UVerticalBox::GetChildIndex(UWidget* Content) const
 bool UVerticalBox::AddChild(UWidget* Child, FVector2D Position)
 {
 	UVerticalBoxSlot* Slot = AddSlot(Child);
-	
-	// Add the child to the live canvas if it already exists
-	if (MyVerticalBox.IsValid())
-	{
-		Slot->BuildSlot(MyVerticalBox.ToSharedRef());
-	}
-	
 	return true;
 }
 
@@ -118,6 +111,12 @@ UVerticalBoxSlot* UVerticalBox::AddSlot(UWidget* Content)
 	
 	Slots.Add(Slot);
 
+	// Add the child to the live canvas if it already exists
+	if ( MyVerticalBox.IsValid() )
+	{
+		Slot->BuildSlot(MyVerticalBox.ToSharedRef());
+	}
+
 	return Slot;
 }
 
@@ -136,8 +135,4 @@ void UVerticalBox::ConnectEditorData()
 		//TODO UMG Should we auto delete empty slots?
 	}
 }
-
-void UVerticalBox::PostEditChangeProperty(struct FPropertyChangedEvent& PropertyChangedEvent)
-{
-}
 #endif
diff --git a/Engine/Source/Runtime/UMG/Private/Components/WidgetSwitcher.cpp b/Engine/Source/Runtime/UMG/Private/Components/WidgetSwitcher.cpp
new file mode 100644
--- /dev/null
+++ b/Engine/Source/Runtime/UMG/Private/Components/WidgetSwitcher.cpp
@@ -0,0 +1,178 @@
+// Copyright 1998-2014 Epic Games, Inc. All Rights Reserved.
+
+#include "UMGPrivatePCH.h"
+
+/////////////////////////////////////////////////////
+// UWidgetSwitcher
+
+UWidgetSwitcher::UWidgetSwitcher(const FPostConstructInitializeProperties& PCIP)
+	: Super(PCIP)
+{
+	bIsVariable = true;
+}
+
+int32 UWidgetSwitcher::GetNumWidgets() const
+{
+	if ( MyWidgetSwitcher.IsValid() )
+	{
+		return MyWidgetSwitcher->GetNumWidgets();
+	}
+
+	return Slots.Num();
+}
+
+int32 UWidgetSwitcher::GetActiveWidgetIndex() const
+{
+	if ( MyWidgetSwitcher.IsValid() )
+	{
+		return MyWidgetSwitcher->GetActiveWidgetIndex();
+	}
+
+	return ActiveWidgetIndex;
+}
+
+void UWidgetSwitcher::SetActiveWidgetIndex(int32 Index)
+{
+	ActiveWidgetIndex = Index;
+	if ( MyWidgetSwitcher.IsValid() )
+	{
+		// Ensure the index is clamped to a valid range.
+		int32 SafeIndex = FMath::Clamp(Index, 0, FMath::Max(0, Slots.Num() - 1));
+		MyWidgetSwitcher->SetActiveWidgetIndex(SafeIndex);
+	}
+}
+
+int32 UWidgetSwitcher::GetChildrenCount() const
+{
+	return Slots.Num();
+}
+
+UWidget* UWidgetSwitcher::GetChildAt(int32 Index) const
+{
+	return Slots[Index]->Content;
+}
+
+int32 UWidgetSwitcher::GetChildIndex(UWidget* Content) const
+{
+	for ( int32 SlotIndex = 0; SlotIndex < Slots.Num(); ++SlotIndex )
+	{
+		UWidgetSwitcherSlot* Slot = Slots[SlotIndex];
+
+		if ( Slot->Content == Content )
+		{
+			return SlotIndex;
+		}
+	}
+
+	return -1;
+}
+
+bool UWidgetSwitcher::AddChild(UWidget* Child, FVector2D Position)
+{
+	UWidgetSwitcherSlot* Slot = AddSlot(Child);
+	return true;
+}
+
+bool UWidgetSwitcher::RemoveChild(UWidget* Child)
+{
+	int32 SlotIndex = GetChildIndex(Child);
+	if ( SlotIndex != -1 )
+	{
+		Slots.RemoveAt(SlotIndex);
+		
+		// Remove the widget from the live slot if it exists.
+		if (MyWidgetSwitcher.IsValid())
+		{
+			MyWidgetSwitcher->RemoveSlot(Child->GetWidget());
+		}
+		
+		return true;
+	}
+
+	return false;
+}
+
+void UWidgetSwitcher::ReplaceChildAt(int32 Index, UWidget* Content)
+{
+	UWidgetSwitcherSlot* Slot = Slots[Index];
+	Slot->Content = Content;
+
+#if WITH_EDITOR
+	Content->Slot = Slot;
+#endif
+}
+
+void UWidgetSwitcher::InsertChildAt(int32 Index, UWidget* Content)
+{
+	UWidgetSwitcherSlot* Slot = ConstructObject<UWidgetSwitcherSlot>(UWidgetSwitcherSlot::StaticClass(), this);
+	Slot->SetFlags(RF_Transactional);
+	Slot->Content = Content;
+	Slot->Parent = this;
+
+#if WITH_EDITOR
+	Content->Slot = Slot;
+#endif
+
+	Slots.Insert(Slot, Index);
+}
+
+UWidgetSwitcherSlot* UWidgetSwitcher::AddSlot(UWidget* Content)
+{
+	UWidgetSwitcherSlot* Slot = ConstructObject<UWidgetSwitcherSlot>(UWidgetSwitcherSlot::StaticClass(), this);
+	Slot->SetFlags(RF_Transactional);
+	Slot->Content = Content;
+	Slot->Parent = this;
+
+#if WITH_EDITOR
+	Content->Slot = Slot;
+#endif
+
+	Slots.Add(Slot);
+
+	// Add the child to the live canvas if it already exists
+	if ( MyWidgetSwitcher.IsValid() )
+	{
+		Slot->BuildSlot(MyWidgetSwitcher.ToSharedRef());
+	}
+
+	return Slot;
+}
+
+TSharedRef<SWidget> UWidgetSwitcher::RebuildWidget()
+{
+	MyWidgetSwitcher = SNew(SWidgetSwitcher);
+
+	for ( auto Slot : Slots )
+	{
+		Slot->Parent = this;
+		Slot->BuildSlot(MyWidgetSwitcher.ToSharedRef());
+	}
+
+	return BuildDesignTimeWidget( MyWidgetSwitcher.ToSharedRef() );
+}
+
+void UWidgetSwitcher::SyncronizeProperties()
+{
+	Super::SyncronizeProperties();
+
+	SetActiveWidgetIndex(ActiveWidgetIndex);
+}
+
+#if WITH_EDITOR
+
+void UWidgetSwitcher::ConnectEditorData()
+{
+	for ( UWidgetSwitcherSlot* Slot : Slots )
+	{
+		// TODO UMG Find a better way to assign the parent pointers for slots.
+		Slot->Parent = this;
+
+		if ( Slot->Content )
+		{
+			Slot->Content->Slot = Slot;
+		}
+		//TODO UMG Should we auto delete empty slots?
+	}
+}
+
+#endif
diff --git a/Engine/Source/Runtime/UMG/Private/Components/WidgetSwitcherSlot.cpp b/Engine/Source/Runtime/UMG/Private/Components/WidgetSwitcherSlot.cpp
new file mode 100644
--- /dev/null
+++ b/Engine/Source/Runtime/UMG/Private/Components/WidgetSwitcherSlot.cpp
@@ -0,0 +1,59 @@
+// Copyright 1998-2014 Epic Games, Inc. All Rights Reserved.
+
+#include "UMGPrivatePCH.h"
+
+/////////////////////////////////////////////////////
+// UWidgetSwitcherSlot
+
+UWidgetSwitcherSlot::UWidgetSwitcherSlot(const FPostConstructInitializeProperties& PCIP)
+	: Super(PCIP)
+	, Slot(NULL)
+{
+	HorizontalAlignment = HAlign_Fill;
+	VerticalAlignment = VAlign_Fill;
+}
+
+void UWidgetSwitcherSlot::BuildSlot(TSharedRef<SWidgetSwitcher> WidgetSwitcher)
+{
+	Slot = &WidgetSwitcher->AddSlot()
+		.Padding(Padding)
+		.HAlign(HorizontalAlignment)
+		.VAlign(VerticalAlignment)
+		[
+			Content == NULL ? SNullWidget::NullWidget : Content->GetWidget()
+		];
+}
+
+void UWidgetSwitcherSlot::SetPadding(FMargin InPadding)
+{
+	Padding = InPadding;
+	if ( Slot )
+	{
+		Slot->Padding(InPadding);
+	}
+}
+
+void UWidgetSwitcherSlot::SetHorizontalAlignment(EHorizontalAlignment InHorizontalAlignment)
+{
+	HorizontalAlignment = InHorizontalAlignment;
+	if ( Slot )
+	{
+		Slot->HAlign(InHorizontalAlignment);
+	}
+}
+
+void UWidgetSwitcherSlot::SetVerticalAlignment(EVerticalAlignment InVerticalAlignment)
+{
+	VerticalAlignment = InVerticalAlignment;
+	if ( Slot )
+	{
+		Slot->VAlign(InVerticalAlignment);
+	}
+}
+
+void UWidgetSwitcherSlot::Refresh()
+{
+	SetPadding(Padding);
+	SetHorizontalAlignment(HorizontalAlignment);
+	SetVerticalAlignment(VerticalAlignment);
+}
diff --git a/Engine/Source/Runtime/UMG/Public/Components/VerticalBox.h b/Engine/Source/Runtime/UMG/Public/Components/VerticalBox.h
--- a/Engine/Source/Runtime/UMG/Public/Components/VerticalBox.h
+++ b/Engine/Source/Runtime/UMG/Public/Components/VerticalBox.h
@@ -32,10 +32,6 @@ class UMG_API UVerticalBox : public UPanelWidget
 	// End UPanelWidget
 
 #if WITH_EDITOR
-	// UObject interface
-	virtual void PostEditChangeProperty(struct FPropertyChangedEvent& PropertyChangedEvent) override;
-	// End of UObject interface
-
 	// UWidget interface
 	virtual void ConnectEditorData() override;
 	// End UWidget interface
diff --git a/Engine/Source/Runtime/UMG/Public/Components/WidgetSwitcher.h b/Engine/Source/Runtime/UMG/Public/Components/WidgetSwitcher.h
new file mode 100644
--- /dev/null
+++ b/Engine/Source/Runtime/UMG/Public/Components/WidgetSwitcher.h
@@ -0,0 +1,66 @@
+// Copyright 1998-2014 Epic Games, Inc. All Rights Reserved.
+
+#pragma once
+
+#include "WidgetSwitcher.generated.h"
+
+/**
+ * A vertical box widget is a layout panel allowing child widgets to be automatically laid out
+ * vertically.
+ */
+UCLASS(meta=( Category="Panel" ), ClassGroup=UserInterface)
+class UMG_API UWidgetSwitcher : public UPanelWidget
+{
+	GENERATED_UCLASS_BODY()
+
+	/** The slots containing the widgets that are flowed vertically. */
+	UPROPERTY()
+	TArray<UWidgetSwitcherSlot*> Slots;
+
+	/** Image to draw */
+	UPROPERTY(EditDefaultsOnly, Category="Switcher", meta=( UIMin=0, ClampMin=0 ))
+	int32 ActiveWidgetIndex;
+
+	/** Gets the number of widgets that this switcher manages. */
+	UFUNCTION(BlueprintCallable, Category="Switcher")
+	int32 GetNumWidgets() const;
+
+	/** Gets the slot index of the currently active widget */
+	UFUNCTION(BlueprintCallable, Category="Switcher")
+	int32 GetActiveWidgetIndex() const;
+
+	/** Activates the widget at the specified index. */
+	UFUNCTION(BlueprintCallable, Category="Switcher")
+	void SetActiveWidgetIndex( int32 Index );
+	
+	//TODO UMG Add ways to make adding slots callable by blueprints.
+
+	UWidgetSwitcherSlot* AddSlot(UWidget* Content);
+
+	// UPanelWidget
+	virtual int32 GetChildrenCount() const override;
+	virtual UWidget* GetChildAt(int32 Index) const override;
+	virtual int32 GetChildIndex(UWidget* Content) const override;
+	virtual bool AddChild(UWidget* Child, FVector2D Position) override;
+	virtual bool RemoveChild(UWidget* Child) override;
+	virtual void ReplaceChildAt(int32 Index, UWidget* Child) override;
+	virtual void InsertChildAt(int32 Index, UWidget* Child) override;
+	// End UPanelWidget
+
+	void SyncronizeProperties();
+
+#if WITH_EDITOR
+	// UWidget interface
+	virtual void ConnectEditorData() override;
+	// End UWidget interface
+#endif
+
+protected:
+
+	TSharedPtr<class SWidgetSwitcher> MyWidgetSwitcher;
+
+protected:
+	// UWidget interface
+	virtual TSharedRef<SWidget> RebuildWidget() override;
+	// End of UWidget interface
+};
diff --git a/Engine/Source/Runtime/UMG/Public/Components/WidgetSwitcherSlot.h b/Engine/Source/Runtime/UMG/Public/Components/WidgetSwitcherSlot.h
new file mode 100644
--- /dev/null
+++ b/Engine/Source/Runtime/UMG/Public/Components/WidgetSwitcherSlot.h
@@ -0,0 +1,48 @@
+// Copyright 1998-2014 Epic Games, Inc. All Rights Reserved.
+
+#pragma once
+
+#include "SlateWrapperTypes.h"
+
+#include "WidgetSwitcherSlot.generated.h"
+
+/** The Slot for the UWidgetSwitcher, contains the widget that is flowed vertically */
+UCLASS()
+class UMG_API UWidgetSwitcherSlot : public UPanelSlot
+{
+	GENERATED_UCLASS_BODY()
+	
+	/** The padding area between the slot and the content it contains. */
+	UPROPERTY(EditDefaultsOnly, Category=Layout)
+	FMargin Padding;
+
+	/** The alignment of the object horizontally. */
+	UPROPERTY(EditDefaultsOnly, Category=Layout)
+	TEnumAsByte<EHorizontalAlignment> HorizontalAlignment;
+
+	/** The alignment of the object vertically. */
+	UPROPERTY(EditDefaultsOnly, Category=Layout)
+	TEnumAsByte<EVerticalAlignment> VerticalAlignment;
+
+	UFUNCTION(BlueprintCallable, Category="Appearance")
+	void SetPadding(FMargin InPadding);
+
+	UFUNCTION(BlueprintCallable, Category="Appearance")
+	void SetHorizontalAlignment(EHorizontalAlignment InHorizontalAlignment);
+
+	UFUNCTION(BlueprintCallable, Category="Appearance")
+	void SetVerticalAlignment(EVerticalAlignment InVerticalAlignment);
+
+	// UPanelSlot interface
+	virtual void Refresh() override;
+	// End of UPanelSlot interface
+
+	/** Builds the underlying FSlot for the Slate layout panel. */
+	void BuildSlot(TSharedRef<SWidgetSwitcher> InWidgetSwitcher);
+
+private:
+	//TODO UMG Slots should hold weak or shared refs to slots.
+
+	/** A raw pointer to the slot to allow us to adjust the size, padding...etc at runtime. */
+	FSimpleSlot* Slot;
+};
diff --git a/Engine/Source/Runtime/UMG/Public/UMG.h b/Engine/Source/Runtime/UMG/Public/UMG.h
--- a/Engine/Source/Runtime/UMG/Public/UMG.h
+++ b/Engine/Source/Runtime/UMG/Public/UMG.h
@@ -49,6 +49,9 @@
 #include "ScrollBoxSlot.h"
 #include "ScrollBox.h"
 
+#include "WidgetSwitcherSlot.h"
+#include "WidgetSwitcher.h"
+
 #include "TextBlock.h"
 #include "Border.h"
 #include "Button.h"
