UMG - Fixing a few startup issues with running in the game after a bunch of changes to UserWidget.

diff --git a/Engine/Source/Runtime/UMG/Private/UserWidget.cpp b/Engine/Source/Runtime/UMG/Private/UserWidget.cpp
--- a/Engine/Source/Runtime/UMG/Private/UserWidget.cpp
+++ b/Engine/Source/Runtime/UMG/Private/UserWidget.cpp
@@ -131,7 +131,6 @@ protected:
 UUserWidget::UUserWidget(const FPostConstructInitializeProperties& PCIP)
 	: Super(PCIP)
 {
-	Visiblity = ESlateVisibility::Visible;
 	/*PrimaryActorTick.bCanEverTick = true;
 	PrimaryActorTick.bTickEvenWhenPaused = true;*/
 	bShowCursorWhenVisible = false;
@@ -179,7 +178,8 @@ UWidget* UUserWidget::GetWidgetHandle(TSharedRef<SWidget> InWidget)
 {
 	GetWidget();
 	
-	return WidgetToComponent.FindRef(InWidget);
+	TWeakObjectPtr<UWidget> VisualWidget = WidgetToComponent.FindRef(InWidget);
+	return VisualWidget.Get();
 }
 
 TSharedRef<SWidget> UUserWidget::RebuildWidget()
@@ -204,20 +204,6 @@ TSharedRef<SWidget> UUserWidget::RebuildWidget()
 
 		WidgetToComponent.Add(Widget, Handle);
 	}
-	
-	// If this is a game world add the widget to the current worlds viewport.
-	UWorld* World = GetWorld();
-	if ( World && World->IsGameWorld() )
-	{
-		if ( Visiblity == ESlateVisibility::Visible )
-		{
-			Show();
-		}
-		else
-		{
-			Hide();
-		}
-	}
 
 	return MakeWidget();
 }
@@ -245,7 +231,8 @@ UWidget* UUserWidget::GetHandleFromName(const FString& Name) const
 	{
 		if ( Entry.Value->GetName().Equals(Name, ESearchCase::IgnoreCase) )
 		{
-			return Entry.Value;
+			TWeakObjectPtr<UWidget> VisualWidget = Entry.Value;
+			return VisualWidget.Get();
 		}
 	}
 
diff --git a/Engine/Source/Runtime/UMG/Public/Blueprint/UserWidget.h b/Engine/Source/Runtime/UMG/Public/Blueprint/UserWidget.h
--- a/Engine/Source/Runtime/UMG/Public/Blueprint/UserWidget.h
+++ b/Engine/Source/Runtime/UMG/Public/Blueprint/UserWidget.h
@@ -262,7 +262,7 @@ protected:
 
 private:
 	TSharedPtr<SWidget> UserRootWidget;
-	TMap< TWeakPtr<SWidget>, UWidget* > WidgetToComponent;
+	TMap< TWeakPtr<SWidget>, TWeakObjectPtr<UWidget> > WidgetToComponent;
 
 	TWeakPtr<SWidget> FullScreenWidget;
 };
