UMG - Fixing bug with dragging out widgets not being removed.  Undo fixes.  General cleanup.

diff --git a/Engine/Source/Editor/UMGEditor/Private/SUMGDesigner.cpp b/Engine/Source/Editor/UMGEditor/Private/SUMGDesigner.cpp
--- a/Engine/Source/Editor/UMGEditor/Private/SUMGDesigner.cpp
+++ b/Engine/Source/Editor/UMGEditor/Private/SUMGDesigner.cpp
@@ -693,6 +693,11 @@ void SUMGDesigner::OnDragLeave(const FDragDropEvent& DragDropEvent)
 	
 	if (DropPreviewWidget)
 	{
+		if ( DropPreviewParent )
+		{
+			DropPreviewParent->RemoveChild(DropPreviewWidget);
+		}
+
 		UWidgetBlueprint* BP = GetBlueprint();
 		BP->WidgetTree->RemoveWidget(DropPreviewWidget);
 		DropPreviewWidget = NULL;
@@ -788,6 +793,9 @@ bool SUMGDesigner::AddToTemplate(const FGeometry& MyGeometry, const FDragDropEve
 				const FScopedTransaction Transaction(LOCTEXT("Designer_AddWidget", "Add Widget"));
 				Parent->SetFlags(RF_Transactional);
 				Parent->Modify();
+
+				BP->WidgetTree->SetFlags(RF_Transactional);
+				BP->WidgetTree->Modify();
 				
 				UWidget* Widget = DragDropOp->Template->Create(BP->WidgetTree);
 				Widget->IsDesignTime(true);
diff --git a/Engine/Source/Runtime/UMG/Private/Components/CanvasPanel.cpp b/Engine/Source/Runtime/UMG/Private/Components/CanvasPanel.cpp
--- a/Engine/Source/Runtime/UMG/Private/Components/CanvasPanel.cpp
+++ b/Engine/Source/Runtime/UMG/Private/Components/CanvasPanel.cpp
@@ -23,6 +23,11 @@ class SFixedSizeCanvas : public SConstraintCanvas
 		CanvasSize = InSize;
 	}
 
+	void SetDesiredSize(FVector2D InSize)
+	{
+		CanvasSize = InSize;
+	}
+
 	// SWidget interface
 	virtual FVector2D ComputeDesiredSize() const override
 	{
@@ -143,6 +148,13 @@ TSharedRef<SWidget> UCanvasPanel::RebuildWidget()
 	return BuildDesignTimeWidget( MyCanvas.ToSharedRef() );
 }
 
+void UCanvasPanel::SyncronizeProperties()
+{
+	Super::SyncronizeProperties();
+
+	MyCanvas->SetDesiredSize(DesiredCanvasSize);
+}
+
 UCanvasPanelSlot* UCanvasPanel::AddSlot(UWidget* Content)
 {
 	UCanvasPanelSlot* Slot = ConstructObject<UCanvasPanelSlot>(UCanvasPanelSlot::StaticClass(), this);
@@ -196,8 +208,4 @@ void UCanvasPanel::ConnectEditorData()
 	}
 }
 
-void UCanvasPanel::PostEditChangeProperty(struct FPropertyChangedEvent& PropertyChangedEvent)
-{
-}
-
 #endif
\ No newline at end of file
diff --git a/Engine/Source/Runtime/UMG/Private/Components/HorizontalBox.cpp b/Engine/Source/Runtime/UMG/Private/Components/HorizontalBox.cpp
--- a/Engine/Source/Runtime/UMG/Private/Components/HorizontalBox.cpp
+++ b/Engine/Source/Runtime/UMG/Private/Components/HorizontalBox.cpp
@@ -129,8 +129,4 @@ void UHorizontalBox::ConnectEditorData()
 		Slot->Parent = this;
 	}
 }
-
-void UHorizontalBox::PostEditChangeProperty(struct FPropertyChangedEvent& PropertyChangedEvent)
-{
-}
 #endif
diff --git a/Engine/Source/Runtime/UMG/Private/Components/Overlay.cpp b/Engine/Source/Runtime/UMG/Private/Components/Overlay.cpp
--- a/Engine/Source/Runtime/UMG/Private/Components/Overlay.cpp
+++ b/Engine/Source/Runtime/UMG/Private/Components/Overlay.cpp
@@ -97,6 +97,7 @@ UOverlaySlot* UOverlay::AddSlot(UWidget* Content)
 }
 
 #if WITH_EDITOR
+
 void UOverlay::ConnectEditorData()
 {
 	for ( UOverlaySlot* Slot : Slots )
@@ -109,8 +110,4 @@ void UOverlay::ConnectEditorData()
 	}
 }
 
-void UOverlay::PostEditChangeProperty(struct FPropertyChangedEvent& PropertyChangedEvent)
-{
-
-}
 #endif
diff --git a/Engine/Source/Runtime/UMG/Private/Components/ScrollBox.cpp b/Engine/Source/Runtime/UMG/Private/Components/ScrollBox.cpp
--- a/Engine/Source/Runtime/UMG/Private/Components/ScrollBox.cpp
+++ b/Engine/Source/Runtime/UMG/Private/Components/ScrollBox.cpp
@@ -124,6 +124,7 @@ UScrollBoxSlot* UScrollBox::AddSlot(UWidget* Content)
 }
 
 #if WITH_EDITOR
+
 void UScrollBox::ConnectEditorData()
 {
 	for ( UScrollBoxSlot* Slot : Slots )
@@ -139,9 +140,6 @@ void UScrollBox::ConnectEditorData()
 	}
 }
 
-void UScrollBox::PostEditChangeProperty(struct FPropertyChangedEvent& PropertyChangedEvent)
-{
-}
 #endif
 
 /////////////////////////////////////////////////////
diff --git a/Engine/Source/Runtime/UMG/Private/Components/UniformGridPanel.cpp b/Engine/Source/Runtime/UMG/Private/Components/UniformGridPanel.cpp
--- a/Engine/Source/Runtime/UMG/Private/Components/UniformGridPanel.cpp
+++ b/Engine/Source/Runtime/UMG/Private/Components/UniformGridPanel.cpp
@@ -111,8 +111,4 @@ void UUniformGridPanel::ConnectEditorData()
 	}
 }
 
-void UUniformGridPanel::PostEditChangeProperty(struct FPropertyChangedEvent& PropertyChangedEvent)
-{
-}
-
 #endif
diff --git a/Engine/Source/Runtime/UMG/Private/WidgetTree.cpp b/Engine/Source/Runtime/UMG/Private/WidgetTree.cpp
--- a/Engine/Source/Runtime/UMG/Private/WidgetTree.cpp
+++ b/Engine/Source/Runtime/UMG/Private/WidgetTree.cpp
@@ -74,9 +74,9 @@ bool UWidgetTree::RemoveWidgetRecursive(UWidget* InRemovedWidget)
 		UPanelWidget* NonLeafTemplate = Cast<UPanelWidget>(Template);
 		if ( NonLeafTemplate )
 		{
-			NonLeafTemplate->Modify();
 			if ( NonLeafTemplate->RemoveChild(InRemovedWidget) )
 			{
+				NonLeafTemplate->Modify();
 				return true;
 			}
 		}
@@ -87,10 +87,7 @@ bool UWidgetTree::RemoveWidgetRecursive(UWidget* InRemovedWidget)
 
 bool UWidgetTree::RemoveWidget(UWidget* InRemovedWidget)
 {
-	Modify();
-
 	bool bRemoved = RemoveWidgetRecursive(InRemovedWidget);
-
 	int32 IndexRemoved = WidgetTemplates.Remove(InRemovedWidget);
 	
 	return bRemoved || IndexRemoved != -1;
diff --git a/Engine/Source/Runtime/UMG/Public/Components/CanvasPanel.h b/Engine/Source/Runtime/UMG/Public/Components/CanvasPanel.h
--- a/Engine/Source/Runtime/UMG/Public/Components/CanvasPanel.h
+++ b/Engine/Source/Runtime/UMG/Public/Components/CanvasPanel.h
@@ -36,11 +36,11 @@ class UMG_API UCanvasPanel : public UPanelWidget
 	/** Computes the geometry for a particular slot based on the current geometry of the canvas. */
 	bool GetGeometryForSlot(UCanvasPanelSlot* Slot, FGeometry& ArrangedGeometry) const;
 
-#if WITH_EDITOR
-	// UObject interface
-	virtual void PostEditChangeProperty(struct FPropertyChangedEvent& PropertyChangedEvent) override;
-	// End of UObject interface
+	// UWidget interface
+	virtual void SyncronizeProperties() override;
+	// End of UWidget interface
 
+#if WITH_EDITOR
 	// UWidget interface
 	virtual void ConnectEditorData() override;
 	// End UWidget interface
diff --git a/Engine/Source/Runtime/UMG/Public/Components/CheckBox.h b/Engine/Source/Runtime/UMG/Public/Components/CheckBox.h
--- a/Engine/Source/Runtime/UMG/Public/Components/CheckBox.h
+++ b/Engine/Source/Runtime/UMG/Public/Components/CheckBox.h
@@ -100,7 +100,7 @@ protected:
 	bool IsChecked() const;
 	
 	// UWidget interface
-	void SyncronizeProperties() override;
+	virtual void SyncronizeProperties() override;
 	// End of UWidget interface
 
 protected:
diff --git a/Engine/Source/Runtime/UMG/Public/Components/HorizontalBox.h b/Engine/Source/Runtime/UMG/Public/Components/HorizontalBox.h
--- a/Engine/Source/Runtime/UMG/Public/Components/HorizontalBox.h
+++ b/Engine/Source/Runtime/UMG/Public/Components/HorizontalBox.h
@@ -26,10 +26,6 @@ class UMG_API UHorizontalBox : public UPanelWidget
 	// End UPanelWidget
 
 #if WITH_EDITOR
-	// UObject interface
-	virtual void PostEditChangeProperty(struct FPropertyChangedEvent& PropertyChangedEvent) override;
-	// End of UObject interface
-
 	// UWidget interface
 	virtual void ConnectEditorData() override;
 	// End UWidget interface
diff --git a/Engine/Source/Runtime/UMG/Public/Components/Overlay.h b/Engine/Source/Runtime/UMG/Public/Components/Overlay.h
--- a/Engine/Source/Runtime/UMG/Public/Components/Overlay.h
+++ b/Engine/Source/Runtime/UMG/Public/Components/Overlay.h
@@ -24,10 +24,6 @@ class UMG_API UOverlay : public UPanelWidget
 	// End UPanelWidget
 
 #if WITH_EDITOR
-	// UObject interface
-	virtual void PostEditChangeProperty(struct FPropertyChangedEvent& PropertyChangedEvent) override;
-	// End of UObject interface
-
 	// UWidget interface
 	virtual void ConnectEditorData() override;
 	// End UWidget interface
diff --git a/Engine/Source/Runtime/UMG/Public/Components/ScrollBox.h b/Engine/Source/Runtime/UMG/Public/Components/ScrollBox.h
--- a/Engine/Source/Runtime/UMG/Public/Components/ScrollBox.h
+++ b/Engine/Source/Runtime/UMG/Public/Components/ScrollBox.h
@@ -29,10 +29,6 @@ class UMG_API UScrollBox : public UPanelWidget
 	// End UPanelWidget
 
 #if WITH_EDITOR
-	// UObject interface
-	virtual void PostEditChangeProperty(struct FPropertyChangedEvent& PropertyChangedEvent) override;
-	// End of UObject interface
-
 	// UWidget interface
 	virtual void ConnectEditorData() override;
 	// End UWidget interface
diff --git a/Engine/Source/Runtime/UMG/Public/Components/Slider.h b/Engine/Source/Runtime/UMG/Public/Components/Slider.h
--- a/Engine/Source/Runtime/UMG/Public/Components/Slider.h
+++ b/Engine/Source/Runtime/UMG/Public/Components/Slider.h
@@ -62,7 +62,7 @@ class UMG_API USlider : public UWidget
 	void SetValue(float InValue);
 	
 	// UWidget interface
-	void SyncronizeProperties() override;
+	virtual void SyncronizeProperties() override;
 	// End of UWidget interface
 
 protected:
diff --git a/Engine/Source/Runtime/UMG/Public/Components/UniformGridPanel.h b/Engine/Source/Runtime/UMG/Public/Components/UniformGridPanel.h
--- a/Engine/Source/Runtime/UMG/Public/Components/UniformGridPanel.h
+++ b/Engine/Source/Runtime/UMG/Public/Components/UniformGridPanel.h
@@ -40,10 +40,6 @@ class UMG_API UUniformGridPanel : public UPanelWidget
 	// End UPanelWidget
 
 #if WITH_EDITOR
-	// UObject interface
-	virtual void PostEditChangeProperty(struct FPropertyChangedEvent& PropertyChangedEvent) override;
-	// End of UObject interface
-
 	// UWidget interface
 	virtual void ConnectEditorData() override;
 	// End UWidget interface
