UMG - Updating Image to sync properties like the other newer widgets.  General cleanup.

diff --git a/Engine/Source/Runtime/Slate/Private/Widgets/Images/SImage.cpp b/Engine/Source/Runtime/Slate/Private/Widgets/Images/SImage.cpp
--- a/Engine/Source/Runtime/Slate/Private/Widgets/Images/SImage.cpp
+++ b/Engine/Source/Runtime/Slate/Private/Widgets/Images/SImage.cpp
@@ -92,3 +92,13 @@ void SImage::SetColorAndOpacity( FLinearColor InColorAndOpacity )
 {
 	ColorAndOpacity = InColorAndOpacity;
 }
+
+void SImage::SetImage(TAttribute<const FSlateBrush*> InImage)
+{
+	Image = InImage;
+}
+
+void SImage::SetOnMouseButtonDown(FPointerEventHandler EventHandler)
+{
+	OnMouseButtonDownHandler = EventHandler;
+}
diff --git a/Engine/Source/Runtime/Slate/Public/Widgets/Images/SImage.h b/Engine/Source/Runtime/Slate/Public/Widgets/Images/SImage.h
--- a/Engine/Source/Runtime/Slate/Public/Widgets/Images/SImage.h
+++ b/Engine/Source/Runtime/Slate/Public/Widgets/Images/SImage.h
@@ -39,11 +39,17 @@ public:
 
 public:
 
-	/** The the color and opacity of this image */
+	/** See the ColorAndOpacity attribute */
 	void SetColorAndOpacity( const TAttribute<FSlateColor>& InColorAndOpacity );
-	/** The the color and opacity of this image */
+	
+	/** See the ColorAndOpacity attribute */
 	void SetColorAndOpacity( FLinearColor InColorAndOpacity );
 
+	/** See the Image attribute */
+	void SetImage(TAttribute<const FSlateBrush*> InImage);
+	
+	/** See OnMouseButtonDown event */
+	void SetOnMouseButtonDown(FPointerEventHandler EventHandler);
 
 protected:
 
diff --git a/Engine/Source/Runtime/UMG/Private/Components/Image.cpp b/Engine/Source/Runtime/UMG/Private/Components/Image.cpp
--- a/Engine/Source/Runtime/UMG/Private/Components/Image.cpp
+++ b/Engine/Source/Runtime/UMG/Private/Components/Image.cpp
@@ -13,15 +13,37 @@ UImage::UImage(const FPostConstructInitializeProperties& PCIP)
 	, Image()
 	, ColorAndOpacity(FLinearColor::White)
 {
-	//@TODO: Would like to default to the checkerboard, but TextureName always wins over TextureObject :(
-	//	Image = *FEditorStyle::GetDefaultBrush();
 }
 
 TSharedRef<SWidget> UImage::RebuildWidget()
 {
-	return SNew(SImage)
-		.Image( BIND_UOBJECT_ATTRIBUTE(const FSlateBrush*, GetImageBrush) )
-		.ColorAndOpacity( BIND_UOBJECT_ATTRIBUTE(FSlateColor, GetColorAndOpacity) );
+	MyImage = SNew(SImage);
+	return MyImage.ToSharedRef();
+}
+
+void UImage::SyncronizeProperties()
+{
+	MyImage->SetImage(GetImageBrush());
+	MyImage->SetColorAndOpacity(ColorAndOpacity);
+	MyImage->SetOnMouseButtonDown(BIND_UOBJECT_DELEGATE(FPointerEventHandler, HandleMouseButtonDown));
+}
+
+void UImage::SetColorAndOpacity(FLinearColor InColorAndOpacity)
+{
+	ColorAndOpacity = InColorAndOpacity;
+	if ( MyImage.IsValid() )
+	{
+		MyImage->SetColorAndOpacity(ColorAndOpacity);
+	}
+}
+
+void UImage::SetImage(USlateBrushAsset* InImage)
+{
+	Image = InImage;
+	if ( MyImage.IsValid() )
+	{
+		MyImage->SetImage(GetImageBrush());
+	}
 }
 
 const FSlateBrush* UImage::GetImageBrush() const
@@ -35,9 +57,14 @@ const FSlateBrush* UImage::GetImageBrush() const
 	return &Image->Brush;
 }
 
-FSlateColor UImage::GetColorAndOpacity() const
+FReply UImage::HandleMouseButtonDown(const FGeometry& Geometry, const FPointerEvent& MouseEvent)
 {
-	return ColorAndOpacity;
+	if ( OnMouseButtonDownEvent.IsBound() )
+	{
+		return OnMouseButtonDownEvent.Execute(Geometry, MouseEvent).ToReply(MyImage.ToSharedRef());
+	}
+
+	return FReply::Unhandled();
 }
 
 /////////////////////////////////////////////////////
diff --git a/Engine/Source/Runtime/UMG/Public/Components/Button.h b/Engine/Source/Runtime/UMG/Public/Components/Button.h
--- a/Engine/Source/Runtime/UMG/Public/Components/Button.h
+++ b/Engine/Source/Runtime/UMG/Public/Components/Button.h
@@ -84,7 +84,7 @@ public:
 	// End UContentWidget interface
 
 	// UWidget interface
-	void SyncronizeProperties() override;
+	virtual void SyncronizeProperties() override;
 	// End of UWidget interface
 
 protected:
diff --git a/Engine/Source/Runtime/UMG/Public/Components/Image.h b/Engine/Source/Runtime/UMG/Public/Components/Image.h
--- a/Engine/Source/Runtime/UMG/Public/Components/Image.h
+++ b/Engine/Source/Runtime/UMG/Public/Components/Image.h
@@ -15,12 +15,27 @@ class UMG_API UImage : public UWidget
 public:
 
 	/** Image to draw */
-	UPROPERTY(EditAnywhere, BlueprintReadWrite, Category=Appearance, meta=( DisplayThumbnail = "true" ))
+	UPROPERTY(EditDefaultsOnly, Category=Appearance, meta=( DisplayThumbnail = "true" ))
 	USlateBrushAsset* Image;
 
 	/** Color and opacity */
-	UPROPERTY(EditAnywhere, BlueprintReadWrite, Category=Appearance)
-	FSlateColor ColorAndOpacity;
+	UPROPERTY(EditDefaultsOnly, Category=Appearance)
+	FLinearColor ColorAndOpacity;
+
+	UPROPERTY(EditDefaultsOnly, Category=Events)
+	FOnPointerEvent OnMouseButtonDownEvent;
+
+	/**  */
+	UFUNCTION(BlueprintCallable, Category="Appearance")
+	void SetColorAndOpacity(FLinearColor InColorAndOpacity);
+
+	/**  */
+	UFUNCTION(BlueprintCallable, Category="Appearance")
+	void SetImage(USlateBrushAsset* InImage);
+
+	// UWidget interface
+	virtual void SyncronizeProperties() override;
+	// End of UWidget interface
 
 protected:
 	// UWidget interface
@@ -28,5 +43,9 @@ protected:
 	// End of UWidget interface
 
 	const FSlateBrush* GetImageBrush() const;
-	FSlateColor GetColorAndOpacity() const;
+
+	FReply HandleMouseButtonDown(const FGeometry& Geometry, const FPointerEvent& MouseEvent);
+
+protected:
+	TSharedPtr<SImage> MyImage;
 };
diff --git a/Engine/Source/Runtime/UMG/Public/Components/ProgressBar.h b/Engine/Source/Runtime/UMG/Public/Components/ProgressBar.h
--- a/Engine/Source/Runtime/UMG/Public/Components/ProgressBar.h
+++ b/Engine/Source/Runtime/UMG/Public/Components/ProgressBar.h
@@ -52,7 +52,7 @@ public:
 	void SetPercent(float InPercent);
 	
 	// UWidget interface
-	void SyncronizeProperties() override;
+	virtual void SyncronizeProperties() override;
 	// End of UWidget interface
 
 protected:
diff --git a/Engine/Source/Runtime/UMG/Public/Components/Spacer.h b/Engine/Source/Runtime/UMG/Public/Components/Spacer.h
--- a/Engine/Source/Runtime/UMG/Public/Components/Spacer.h
+++ b/Engine/Source/Runtime/UMG/Public/Components/Spacer.h
@@ -19,7 +19,7 @@ class UMG_API USpacer : public UWidget
 	void SetSize(FVector2D InSize);
 	
 	// UWidget interface
-	void SyncronizeProperties() override;
+	virtual void SyncronizeProperties() override;
 	// End of UWidget interface
 
 protected:
