UMG - Marking several properties as Advanced.  Better defaults for Border.

diff --git a/Engine/Source/Runtime/UMG/Private/Components/Border.cpp b/Engine/Source/Runtime/UMG/Private/Components/Border.cpp
--- a/Engine/Source/Runtime/UMG/Private/Components/Border.cpp
+++ b/Engine/Source/Runtime/UMG/Private/Components/Border.cpp
@@ -13,17 +13,19 @@ UBorder::UBorder(const FPostConstructInitializeProperties& PCIP)
 {
 	bIsVariable = false;
 
-	SBorder::FArguments BorderDefaults;
+	HorizontalAlignment = HAlign_Fill;
+	VerticalAlignment = VAlign_Fill;
 
-	HorizontalAlignment = BorderDefaults._HAlign;
-	VerticalAlignment = BorderDefaults._VAlign;
-	ContentPadding = BorderDefaults._Padding.Get();
-	ContentScale = BorderDefaults._ContentScale.Get();
+	ContentScale = FVector2D(1.0f, 1.0f);
 	ContentColorAndOpacity = FLinearColor::White;
 
-	DesiredSizeScale = BorderDefaults._DesiredSizeScale.Get();
+	DesiredSizeScale = FVector2D(1.0f, 1.0f);
 	BorderColorAndOpacity = FLinearColor::White;
 	ForegroundColor = FLinearColor::Black;
+
+	SBorder::FArguments BorderDefaults;
+
+	ContentPadding = BorderDefaults._Padding.Get();
 	bShowEffectWhenDisabled = BorderDefaults._ShowEffectWhenDisabled.Get();
 }
 
diff --git a/Engine/Source/Runtime/UMG/Public/Components/Border.h b/Engine/Source/Runtime/UMG/Public/Components/Border.h
--- a/Engine/Source/Runtime/UMG/Public/Components/Border.h
+++ b/Engine/Source/Runtime/UMG/Public/Components/Border.h
@@ -27,20 +27,20 @@ class UMG_API UBorder : public UContentWidget
 	UPROPERTY(EditDefaultsOnly, Category="Content Layout")
 	FMargin ContentPadding;
 
-	/** The scaling factor for the border content */
-	UPROPERTY(EditDefaultsOnly, Category="Content Layout")
-	FVector2D ContentScale;
-
 	/** Color and opacity multiplier of content in the border */
 	UPROPERTY(EditDefaultsOnly, Category="Content Layout")
 	FLinearColor ContentColorAndOpacity;
 
+	/** The scaling factor for the border content */
+	UPROPERTY(EditDefaultsOnly, Category="Content Layout", AdvancedDisplay)
+	FVector2D ContentScale;
+
 	/** Image to use for the border */
 	UPROPERTY(EditDefaultsOnly, Category=Appearance, meta=( DisplayThumbnail = "true" ))
 	USlateBrushAsset* BorderBrush;
 
 	/** The scaling factor for the border */
-	UPROPERTY(EditDefaultsOnly, Category=Appearance)
+	UPROPERTY(EditDefaultsOnly, Category=Appearance, AdvancedDisplay)
 	FVector2D DesiredSizeScale;
 
 	/** Color and opacity of the actual border image */
@@ -52,7 +52,7 @@ class UMG_API UBorder : public UContentWidget
 	FLinearColor ForegroundColor;
 
 	/** Whether or not to show the disabled effect when this border is disabled */
-	UPROPERTY(EditDefaultsOnly, Category=Appearance)
+	UPROPERTY(EditDefaultsOnly, Category=Appearance, AdvancedDisplay)
 	bool bShowEffectWhenDisabled;
 
 	UPROPERTY(EditDefaultsOnly, Category=Events)
diff --git a/Engine/Source/Runtime/UMG/Public/Components/Button.h b/Engine/Source/Runtime/UMG/Public/Components/Button.h
--- a/Engine/Source/Runtime/UMG/Public/Components/Button.h
+++ b/Engine/Source/Runtime/UMG/Public/Components/Button.h
@@ -30,11 +30,11 @@ public:
 	FMargin ContentPadding;
 
 	/** The scaling factor for the button border */
-	UPROPERTY(EditDefaultsOnly, Category=Appearance)
+	UPROPERTY(EditDefaultsOnly, Category=Appearance, AdvancedDisplay)
 	FVector2D DesiredSizeScale;
 
 	/** The scaling factor for the button content */
-	UPROPERTY(EditDefaultsOnly, Category=Appearance)
+	UPROPERTY(EditDefaultsOnly, Category=Appearance, AdvancedDisplay)
 	FVector2D ContentScale;
 	
 	/** The color multiplier for the button images */
diff --git a/Engine/Source/Runtime/UMG/Public/Components/EditableText.h b/Engine/Source/Runtime/UMG/Public/Components/EditableText.h
--- a/Engine/Source/Runtime/UMG/Public/Components/EditableText.h
+++ b/Engine/Source/Runtime/UMG/Public/Components/EditableText.h
@@ -64,23 +64,23 @@ protected:
 	float MinimumDesiredWidth;
 
 	/** Workaround as we loose focus when the auto completion closes. */
-	UPROPERTY(EditDefaultsOnly, Category=Behavior)
+	UPROPERTY(EditDefaultsOnly, Category=Behavior, AdvancedDisplay)
 	bool IsCaretMovedWhenGainFocus;
 
 	/** Whether to select all text when the user clicks to give focus on the widget */
-	UPROPERTY(EditDefaultsOnly, Category=Behavior)
+	UPROPERTY(EditDefaultsOnly, Category=Behavior, AdvancedDisplay)
 	bool SelectAllTextWhenFocused;
 
 	/** Whether to allow the user to back out of changes when they press the escape key */
-	UPROPERTY(EditDefaultsOnly, Category=Behavior)
+	UPROPERTY(EditDefaultsOnly, Category=Behavior, AdvancedDisplay)
 	bool RevertTextOnEscape;
 
 	/** Whether to clear keyboard focus when pressing enter to commit changes */
-	UPROPERTY(EditDefaultsOnly, Category=Behavior)
+	UPROPERTY(EditDefaultsOnly, Category=Behavior, AdvancedDisplay)
 	bool ClearKeyboardFocusOnCommit;
 
 	/** Whether to select all text when pressing enter to commit changes */
-	UPROPERTY(EditDefaultsOnly, Category=Behavior)
+	UPROPERTY(EditDefaultsOnly, Category=Behavior, AdvancedDisplay)
 	bool SelectAllTextOnCommit;
 
 	/** Called whenever the text is changed interactively by the user */
diff --git a/Engine/Source/Runtime/UMG/Public/Components/Slider.h b/Engine/Source/Runtime/UMG/Public/Components/Slider.h
--- a/Engine/Source/Runtime/UMG/Public/Components/Slider.h
+++ b/Engine/Source/Runtime/UMG/Public/Components/Slider.h
@@ -13,18 +13,14 @@ UCLASS(meta=( Category="Common" ), ClassGroup=UserInterface)
 class UMG_API USlider : public UWidget
 {
 	GENERATED_UCLASS_BODY()
-
-	/** Whether the slidable area should be indented to fit the handle. */
-	UPROPERTY(EditDefaultsOnly, Category=Appearance)
-	bool IndentHandle;
-
-	/** Whether the handle is interactive or fixed. */
-	UPROPERTY(EditDefaultsOnly, Category=Appearance)
-	bool Locked;
-
+	
 	///** The style used to draw the slider. */
 	//SLATE_STYLE_ARGUMENT(FSliderStyle, Style)
 
+	/** The volume value to display. */
+	UPROPERTY(EditDefaultsOnly, Category=Appearance)
+	float Value;
+
 	/** The slider's orientation. */
 	UPROPERTY(EditDefaultsOnly, Category=Appearance)
 	TEnumAsByte<EOrientation> Orientation;
@@ -37,9 +33,13 @@ class UMG_API USlider : public UWidget
 	UPROPERTY(EditDefaultsOnly, Category=Appearance)
 	FLinearColor SliderHandleColor;
 
-	/** The volume value to display. */
-	UPROPERTY(EditDefaultsOnly, Category=Appearance)
-	float Value;
+	/** Whether the slidable area should be indented to fit the handle. */
+	UPROPERTY(EditDefaultsOnly, Category=Appearance, AdvancedDisplay)
+	bool IndentHandle;
+
+	/** Whether the handle is interactive or fixed. */
+	UPROPERTY(EditDefaultsOnly, Category=Appearance, AdvancedDisplay)
+	bool Locked;
 
 	/** Invoked when the mouse is pressed and a capture begins. */
 	UPROPERTY(BlueprintAssignable)
diff --git a/Engine/Source/Runtime/UMG/Public/Components/TextBlock.h b/Engine/Source/Runtime/UMG/Public/Components/TextBlock.h
--- a/Engine/Source/Runtime/UMG/Public/Components/TextBlock.h
+++ b/Engine/Source/Runtime/UMG/Public/Components/TextBlock.h
@@ -18,6 +18,14 @@ public:
 	/** A bindable delegate to allow logic to drive the text of the widget */
 	UPROPERTY()
 	FGetText TextDelegate;
+
+	/** The color of the text */
+	UPROPERTY(EditDefaultsOnly, Category=Appearance)
+	FLinearColor ColorAndOpacity;
+
+	/** A bindable delegate for the ColorAndOpacity. */
+	UPROPERTY()
+	FGetSlateColor ColorAndOpacityDelegate;
 	
 	/** The font to render the text with */
 	UPROPERTY(EditDefaultsOnly, Category=Appearance)
@@ -35,20 +43,12 @@ public:
 	UPROPERTY()
 	FGetLinearColor ShadowColorAndOpacityDelegate;
 
-	/** The color of the text */
-	UPROPERTY(EditDefaultsOnly, Category=Appearance)
-	FLinearColor ColorAndOpacity;
-
-	/** A bindable delegate for the ColorAndOpacity. */
-	UPROPERTY()
-	FGetSlateColor ColorAndOpacityDelegate;
-
 	/** Whether text wraps onto a new line when it's length exceeds this width; if this value is zero or negative, no wrapping occurs. */
-	UPROPERTY(EditDefaultsOnly, Category=Appearance)
+	UPROPERTY(EditDefaultsOnly, Category=Appearance, AdvancedDisplay)
 	float WrapTextAt;
 
 	/** True if we're wrapping text automatically based on the computed horizontal space for this widget */
-	UPROPERTY(EditDefaultsOnly, Category=Appearance)
+	UPROPERTY(EditDefaultsOnly, Category=Appearance, AdvancedDisplay)
 	bool AutoWrapText;
 
 	/** The style to use to render the text */
