UMG - Small cleanup on the editabletext/block.  Needs more work.

diff --git a/Engine/Source/Runtime/UMG/Private/Components/EditableText.cpp b/Engine/Source/Runtime/UMG/Private/Components/EditableText.cpp
--- a/Engine/Source/Runtime/UMG/Private/Components/EditableText.cpp
+++ b/Engine/Source/Runtime/UMG/Private/Components/EditableText.cpp
@@ -61,6 +61,8 @@ TSharedRef<SWidget> UEditableText::RebuildWidget()
 
 void UEditableText::SyncronizeProperties()
 {
+	Super::SyncronizeProperties();
+
 	MyEditableText->SetText(Text);
 	MyEditableText->SetHintText(HintText);
 	MyEditableText->SetIsReadOnly(IsReadOnly);
@@ -70,6 +72,25 @@ void UEditableText::SyncronizeProperties()
 	// TODO UMG Complete making all properties settable on SEditableText
 }
 
+FText UEditableText::GetText() const
+{
+	if ( MyEditableText.IsValid() )
+	{
+		return MyEditableText->GetText();
+	}
+
+	return Text;
+}
+
+void UEditableText::SetText(FText InText)
+{
+	Text = InText;
+	if ( MyEditableText.IsValid() )
+	{
+		MyEditableText->SetText(Text);
+	}
+}
+
 void UEditableText::HandleOnTextChanged(const FText& Text)
 {
 	OnTextChanged.Broadcast(Text);
diff --git a/Engine/Source/Runtime/UMG/Private/Components/EditableTextBlock.cpp b/Engine/Source/Runtime/UMG/Private/Components/EditableTextBlock.cpp
--- a/Engine/Source/Runtime/UMG/Private/Components/EditableTextBlock.cpp
+++ b/Engine/Source/Runtime/UMG/Private/Components/EditableTextBlock.cpp
@@ -34,8 +34,7 @@ TSharedRef<SWidget> UEditableTextBlock::RebuildWidget()
 {
 	FString FontPath = FPaths::EngineContentDir() / Font.FontName.ToString();
 
-	return SNew(SEditableTextBox)
-		.Text(Text)
+	MyEditableTextBlock = SNew(SEditableTextBox)
 		.HintText(HintText)
 		.Font(FSlateFontInfo(FontPath, Font.Size))
 		.ForegroundColor(ForegroundColor)
@@ -52,6 +51,40 @@ TSharedRef<SWidget> UEditableTextBlock::RebuildWidget()
 		.SelectAllTextOnCommit(SelectAllTextOnCommit)
 		.OnTextChanged(BIND_UOBJECT_DELEGATE(FOnTextChanged, SlateOnTextChanged))
 		;
+
+	return MyEditableTextBlock.ToSharedRef();
+}
+
+void UEditableTextBlock::SyncronizeProperties()
+{
+	Super::SyncronizeProperties();
+
+	MyEditableTextBlock->SetText(Text);
+//	MyEditableTextBlock->SetHintText(HintText);
+//	MyEditableTextBlock->SetIsReadOnly(IsReadOnly);
+//	MyEditableTextBlock->SetIsPassword(IsPassword);
+//	MyEditableTextBlock->SetColorAndOpacity(ColorAndOpacity);
+
+	// TODO UMG Complete making all properties settable on SEditableTextBox
+}
+
+FText UEditableTextBlock::GetText() const
+{
+	if ( MyEditableTextBlock.IsValid() )
+	{
+		return MyEditableTextBlock->GetText();
+	}
+
+	return Text;
+}
+
+void UEditableTextBlock::SetText(FText InText)
+{
+	Text = InText;
+	if ( MyEditableTextBlock.IsValid() )
+	{
+		MyEditableTextBlock->SetText(Text);
+	}
 }
 
 void UEditableTextBlock::SlateOnTextChanged(const FText& Text)
diff --git a/Engine/Source/Runtime/UMG/Public/Components/EditableText.h b/Engine/Source/Runtime/UMG/Public/Components/EditableText.h
--- a/Engine/Source/Runtime/UMG/Public/Components/EditableText.h
+++ b/Engine/Source/Runtime/UMG/Public/Components/EditableText.h
@@ -90,9 +90,17 @@ protected:
 	/** Called whenever the text is committed.  This happens when the user presses enter or the text box loses focus. */
 	UPROPERTY(BlueprintAssignable)
 	FOnEditableTextCommittedEvent OnTextCommitted;
+
+	/**  */
+	UFUNCTION(BlueprintCallable, Category="Widget")
+	FText GetText() const;
+
+	/**  */
+	UFUNCTION(BlueprintCallable, Category="Widget")
+	void SetText(FText InText);
 	
 	// UWidget interface
-	void SyncronizeProperties() override;
+	virtual void SyncronizeProperties() override;
 	// End of UWidget interface
 
 protected:
diff --git a/Engine/Source/Runtime/UMG/Public/Components/EditableTextBlock.h b/Engine/Source/Runtime/UMG/Public/Components/EditableTextBlock.h
--- a/Engine/Source/Runtime/UMG/Public/Components/EditableTextBlock.h
+++ b/Engine/Source/Runtime/UMG/Public/Components/EditableTextBlock.h
@@ -16,63 +16,63 @@ class UMG_API UEditableTextBlock : public UWidget
 protected:
 
 	/** The text content for this editable text box widget */
-	UPROPERTY(EditAnywhere, BlueprintReadWrite, Category=Content)
+	UPROPERTY(EditDefaultsOnly, Category=Content)
 	FText Text;
 
 	/** Hint text that appears when there is no text in the text box */
-	UPROPERTY(EditAnywhere, BlueprintReadWrite, Category=Content)
+	UPROPERTY(EditDefaultsOnly, Category=Content)
 	FText HintText;
 
 	/** Font color and opacity (overrides Style) */
-	UPROPERTY(EditAnywhere, Category=Appearance)
+	UPROPERTY(EditDefaultsOnly, Category=Appearance)
 	FSlateFontInfo Font;
 
 	/** Text color and opacity (overrides Style) */
-	UPROPERTY(EditAnywhere, BlueprintReadWrite, Category=Appearance)
+	UPROPERTY(EditDefaultsOnly, Category=Appearance)
 	FSlateColor ForegroundColor;
 
 	/** The color of the background/border around the editable text (overrides Style) */
-	UPROPERTY(EditAnywhere, BlueprintReadWrite, Category=Appearance)
+	UPROPERTY(EditDefaultsOnly, Category=Appearance)
 	FSlateColor BackgroundColor;
 
 	/** Text color and opacity when read-only (overrides Style) */
-	UPROPERTY(EditAnywhere, BlueprintReadWrite, Category=Appearance)
+	UPROPERTY(EditDefaultsOnly, Category=Appearance)
 	FSlateColor ReadOnlyForegroundColor;
 
 	/** Sets whether this text box can actually be modified interactively by the user */
-	UPROPERTY(EditAnywhere, BlueprintReadWrite, Category=Appearance)
+	UPROPERTY(EditDefaultsOnly, Category=Appearance)
 	bool IsReadOnly;
 
 	/** Sets whether this text box is for storing a password */
-	UPROPERTY(EditAnywhere, BlueprintReadWrite, Category=Appearance)
+	UPROPERTY(EditDefaultsOnly, Category=Appearance)
 	bool IsPassword;
 
 	/** Minimum width that a text block should be */
-	UPROPERTY(EditAnywhere, BlueprintReadWrite, Category=Appearance)
+	UPROPERTY(EditDefaultsOnly, Category=Appearance)
 	float MinimumDesiredWidth;
 
 	/** Padding between the box/border and the text widget inside (overrides Style) */
-	UPROPERTY(EditAnywhere, BlueprintReadWrite, Category=Appearance)
+	UPROPERTY(EditDefaultsOnly, Category=Appearance)
 	FMargin Padding;
 
 	/** Workaround as we loose focus when the auto completion closes. */
-	UPROPERTY(EditAnywhere, BlueprintReadWrite, Category=Behavior)
+	UPROPERTY(EditDefaultsOnly, Category=Behavior, AdvancedDisplay)
 	bool IsCaretMovedWhenGainFocus;
 
 	/** Whether to select all text when the user clicks to give focus on the widget */
-	UPROPERTY(EditAnywhere, BlueprintReadWrite, Category=Behavior)
+	UPROPERTY(EditDefaultsOnly, Category=Behavior, AdvancedDisplay)
 	bool SelectAllTextWhenFocused;
 
 	/** Whether to allow the user to back out of changes when they press the escape key */
-	UPROPERTY(EditAnywhere, BlueprintReadWrite, Category=Behavior)
+	UPROPERTY(EditDefaultsOnly, Category=Behavior, AdvancedDisplay)
 	bool RevertTextOnEscape;
 
 	/** Whether to clear keyboard focus when pressing enter to commit changes */
-	UPROPERTY(EditAnywhere, BlueprintReadWrite, Category=Behavior)
+	UPROPERTY(EditDefaultsOnly, Category=Behavior, AdvancedDisplay)
 	bool ClearKeyboardFocusOnCommit;
 
 	/** Whether to select all text when pressing enter to commit changes */
-	UPROPERTY(EditAnywhere, BlueprintReadWrite, Category=Behavior)
+	UPROPERTY(EditDefaultsOnly, Category=Behavior, AdvancedDisplay)
 	bool SelectAllTextOnCommit;
 
 	/** Called whenever the text is changed interactively by the user */
@@ -86,10 +86,25 @@ protected:
 	/** Provide a alternative mechanism for error reporting. */
 	//SLATE_ARGUMENT(TSharedPtr<class IErrorReportingWidget>, ErrorReporting)
 
+	/**  */
+	UFUNCTION(BlueprintCallable, Category="Widget")
+	FText GetText() const;
+
+	/**  */
+	UFUNCTION(BlueprintCallable, Category="Widget")
+	void SetText(FText InText);
+	
+	// UWidget interface
+	virtual void SyncronizeProperties() override;
+	// End of UWidget interface
+
 protected:
 	// UWidget interface
 	virtual TSharedRef<SWidget> RebuildWidget() override;
 	// End of UWidget
 
 	void SlateOnTextChanged(const FText& Text);
+
+protected:
+	TSharedPtr<SEditableTextBox> MyEditableTextBlock;
 };
