UMG - Changing a few properties on Border, exposing more things as bindable.

diff --git a/Engine/Source/Runtime/UMG/Private/Components/Border.cpp b/Engine/Source/Runtime/UMG/Private/Components/Border.cpp
--- a/Engine/Source/Runtime/UMG/Private/Components/Border.cpp
+++ b/Engine/Source/Runtime/UMG/Private/Components/Border.cpp
@@ -20,7 +20,7 @@ UBorder::UBorder(const FPostConstructInitializeProperties& PCIP)
 	ContentColorAndOpacity = FLinearColor::White;
 
 	DesiredSizeScale = FVector2D(1.0f, 1.0f);
-	BorderColorAndOpacity = FLinearColor::White;
+	BrushColor = FLinearColor::White;
 	ForegroundColor = FLinearColor::Black;
 
 	SBorder::FArguments BorderDefaults;
@@ -45,7 +45,7 @@ void UBorder::SyncronizeProperties()
 	MyBorder->SetVAlign(VerticalAlignment);
 	MyBorder->SetPadding(ContentPadding);
 	
-	MyBorder->SetBorderBackgroundColor(BorderColorAndOpacity);
+	MyBorder->SetBorderBackgroundColor(BrushColor);
 	MyBorder->SetColorAndOpacity(ContentColorAndOpacity);
 	MyBorder->SetForegroundColor(ForegroundColor);
 	
@@ -74,13 +74,13 @@ void UBorder::SetContent(UWidget* InContent)
 
 const FSlateBrush* UBorder::GetBorderBrush() const
 {
-	if ( BorderBrush == NULL )
+	if ( Brush == NULL )
 	{
 		SBorder::FArguments BorderDefaults;
 		return BorderDefaults._BorderImage.Get();
 	}
 
-	return &BorderBrush->Brush;
+	return &Brush->Brush;
 }
 
 FReply UBorder::HandleMouseButtonDown(const FGeometry& Geometry, const FPointerEvent& MouseEvent)
diff --git a/Engine/Source/Runtime/UMG/Private/Components/EditableTextBlock.cpp b/Engine/Source/Runtime/UMG/Private/Components/EditableTextBlock.cpp
--- a/Engine/Source/Runtime/UMG/Private/Components/EditableTextBlock.cpp
+++ b/Engine/Source/Runtime/UMG/Private/Components/EditableTextBlock.cpp
@@ -35,13 +35,10 @@ TSharedRef<SWidget> UEditableTextBlock::RebuildWidget()
 	FString FontPath = FPaths::EngineContentDir() / Font.FontName.ToString();
 
 	MyEditableTextBlock = SNew(SEditableTextBox)
-		.HintText(HintText)
 		.Font(FSlateFontInfo(FontPath, Font.Size))
 		.ForegroundColor(ForegroundColor)
 		.BackgroundColor(BackgroundColor)
 		.ReadOnlyForegroundColor(ReadOnlyForegroundColor)
-		.IsReadOnly(IsReadOnly)
-		.IsPassword(IsPassword)
 		.MinDesiredWidth(MinimumDesiredWidth)
 		.Padding(Padding)
 		.IsCaretMovedWhenGainFocus(IsCaretMovedWhenGainFocus)
@@ -60,9 +57,9 @@ void UEditableTextBlock::SyncronizeProperties()
 	Super::SyncronizeProperties();
 
 	MyEditableTextBlock->SetText(Text);
-//	MyEditableTextBlock->SetHintText(HintText);
-//	MyEditableTextBlock->SetIsReadOnly(IsReadOnly);
-//	MyEditableTextBlock->SetIsPassword(IsPassword);
+	MyEditableTextBlock->SetHintText(HintText);
+	MyEditableTextBlock->SetIsReadOnly(IsReadOnly);
+	MyEditableTextBlock->SetIsPassword(IsPassword);
 //	MyEditableTextBlock->SetColorAndOpacity(ColorAndOpacity);
 
 	// TODO UMG Complete making all properties settable on SEditableTextBox
diff --git a/Engine/Source/Runtime/UMG/Private/Components/ScrollBoxSlot.cpp b/Engine/Source/Runtime/UMG/Private/Components/ScrollBoxSlot.cpp
--- a/Engine/Source/Runtime/UMG/Private/Components/ScrollBoxSlot.cpp
+++ b/Engine/Source/Runtime/UMG/Private/Components/ScrollBoxSlot.cpp
@@ -9,7 +9,7 @@ UScrollBoxSlot::UScrollBoxSlot(const FPostConstructInitializeProperties& PCIP)
 	: Super(PCIP)
 	, Slot(NULL)
 {
-	HorizontalAlignment = HAlign_Left;
+	HorizontalAlignment = HAlign_Fill;
 }
 
 void UScrollBoxSlot::BuildSlot(TSharedRef<SScrollBox> ScrollBox)
diff --git a/Engine/Source/Runtime/UMG/Public/Components/Border.h b/Engine/Source/Runtime/UMG/Public/Components/Border.h
--- a/Engine/Source/Runtime/UMG/Public/Components/Border.h
+++ b/Engine/Source/Runtime/UMG/Public/Components/Border.h
@@ -37,15 +37,19 @@ class UMG_API UBorder : public UContentWidget
 
 	/** Image to use for the border */
 	UPROPERTY(EditDefaultsOnly, Category=Appearance, meta=( DisplayThumbnail = "true" ))
-	USlateBrushAsset* BorderBrush;
+	USlateBrushAsset* Brush;
 
-	/** The scaling factor for the border */
-	UPROPERTY(EditDefaultsOnly, Category=Appearance, AdvancedDisplay)
-	FVector2D DesiredSizeScale;
+	/** A bindable delegate for the Brush. */
+	UPROPERTY()
+	FGetSlateBrushAsset BrushDelegate;
 
 	/** Color and opacity of the actual border image */
 	UPROPERTY(EditDefaultsOnly, Category=Appearance)
-	FLinearColor BorderColorAndOpacity;
+	FLinearColor BrushColor;
+
+	/** A bindable delegate for the BrushColor. */
+	UPROPERTY()
+	FGetLinearColor BrushColorDelegate;
 
 	/** The foreground color of text and some glyphs that appear as the border's content. */
 	UPROPERTY(EditDefaultsOnly, Category=Appearance)
@@ -55,6 +59,10 @@ class UMG_API UBorder : public UContentWidget
 	UPROPERTY(EditDefaultsOnly, Category=Appearance, AdvancedDisplay)
 	bool bShowEffectWhenDisabled;
 
+	/** The scaling factor for the border */
+	UPROPERTY(EditDefaultsOnly, Category=Appearance, AdvancedDisplay)
+	FVector2D DesiredSizeScale;
+
 	UPROPERTY(EditDefaultsOnly, Category=Events)
 	FOnPointerEvent OnMouseButtonDownEvent;
 
diff --git a/Engine/Source/Runtime/UMG/Public/Components/CheckBox.h b/Engine/Source/Runtime/UMG/Public/Components/CheckBox.h
--- a/Engine/Source/Runtime/UMG/Public/Components/CheckBox.h
+++ b/Engine/Source/Runtime/UMG/Public/Components/CheckBox.h
@@ -14,7 +14,7 @@ class UMG_API UCheckBox : public UWidget
 
 protected:
 	/** Style of the check box */
-	UPROPERTY(EditDefaultsOnly, BlueprintReadOnly, Category=Style, meta=( DisplayThumbnail = "true" ))
+	UPROPERTY(EditDefaultsOnly, Category=Style, meta=( DisplayThumbnail = "true" ))
 	USlateWidgetStyleAsset* Style;
 
 	/** Image to use when the checkbox is unchecked */
@@ -58,23 +58,23 @@ protected:
 	FOnCheckBoxComponentStateChanged OnCheckStateChanged;
 
 	/** Whether the check box is currently in a checked state */
-	UPROPERTY(EditDefaultsOnly, BlueprintReadWrite, Category=Appearance)
+	UPROPERTY(EditDefaultsOnly, Category=Appearance)
 	bool bIsChecked;
 
 	/** How the content of the toggle button should align within the given space */
-	UPROPERTY(EditDefaultsOnly, BlueprintReadOnly, Category=Appearance)
+	UPROPERTY(EditDefaultsOnly, Category=Appearance)
 	TEnumAsByte<EHorizontalAlignment> HorizontalAlignment;
 
 	/** Spacing between the check box image and its content */
-	UPROPERTY(EditDefaultsOnly, BlueprintReadOnly, Category=Appearance)
+	UPROPERTY(EditDefaultsOnly, Category=Appearance)
 	FMargin Padding;
 
 	/** Foreground color for the checkbox's content and parts */
-	UPROPERTY(EditDefaultsOnly, BlueprintReadOnly, Category=Appearance)
+	UPROPERTY(EditDefaultsOnly, Category=Appearance)
 	FSlateColor ForegroundColor;
 
 	/** The color of the background border */
-	UPROPERTY(EditDefaultsOnly, BlueprintReadOnly, Category=Appearance)
+	UPROPERTY(EditDefaultsOnly, Category=Appearance)
 	FSlateColor BorderBackgroundColor;
 
 	//SLATE_ATTRIBUTE( bool, ReadOnly )
diff --git a/Engine/Source/Runtime/UMG/Public/Components/TextBlock.h b/Engine/Source/Runtime/UMG/Public/Components/TextBlock.h
--- a/Engine/Source/Runtime/UMG/Public/Components/TextBlock.h
+++ b/Engine/Source/Runtime/UMG/Public/Components/TextBlock.h
@@ -11,6 +11,10 @@ class UMG_API UTextBlock : public UWidget
 	GENERATED_UCLASS_BODY()
 
 public:
+	/** The style to use to render the text */
+	UPROPERTY(EditDefaultsOnly, Category=Style, meta=( DisplayThumbnail = "true" ))
+	USlateWidgetStyleAsset* Style;
+
 	/** The text to display */
 	UPROPERTY(EditDefaultsOnly, Category=Content)
 	FText Text;
@@ -51,10 +55,6 @@ public:
 	UPROPERTY(EditDefaultsOnly, Category=Appearance, AdvancedDisplay)
 	bool AutoWrapText;
 
-	/** The style to use to render the text */
-	UPROPERTY(EditDefaultsOnly, BlueprintReadOnly, Category=Style, meta=( DisplayThumbnail = "true" ))
-	USlateWidgetStyleAsset* Style;
-
 	///** Called when this text is double clicked */
 	//SLATE_EVENT(FOnClicked, OnDoubleClicked)
 
diff --git a/Engine/Source/Runtime/UMG/Public/Components/Widget.h b/Engine/Source/Runtime/UMG/Public/Components/Widget.h
--- a/Engine/Source/Runtime/UMG/Public/Components/Widget.h
+++ b/Engine/Source/Runtime/UMG/Public/Components/Widget.h
@@ -34,6 +34,7 @@ public:
 	DECLARE_DYNAMIC_DELEGATE_RetVal(FLinearColor, FGetLinearColor);
 	DECLARE_DYNAMIC_DELEGATE_RetVal(ESlateVisibility::Type, FGetSlateVisibility);
 	DECLARE_DYNAMIC_DELEGATE_RetVal(EMouseCursor::Type, FGetMouseCursor);
+	DECLARE_DYNAMIC_DELEGATE_RetVal(USlateBrushAsset*, FGetSlateBrushAsset);
 
 	DECLARE_DYNAMIC_DELEGATE_RetVal_OneParam(UWidget*, FGenerateWidgetUObject, UObject*, Item);
 
