UMG - Working on icons for the widgets that appear in the tree.

diff --git a/Engine/Source/Editor/UMGEditor/Private/SUMGEditorTreeItem.cpp b/Engine/Source/Editor/UMGEditor/Private/SUMGEditorTreeItem.cpp
--- a/Engine/Source/Editor/UMGEditor/Private/SUMGEditorTreeItem.cpp
+++ b/Engine/Source/Editor/UMGEditor/Private/SUMGEditorTreeItem.cpp
@@ -27,14 +27,30 @@ void SUMGEditorTreeItem::Construct(const FArguments& InArgs, const TSharedRef< S
 		.Padding(2.0f)
 		.Content()
 		[
-			SNew(SInlineEditableTextBlock)
-			.ToolTipText(this, &SUMGEditorTreeItem::GetItemTooltipText)
-			.Font(this, &SUMGEditorTreeItem::GetItemFont)
-			.Text(this, &SUMGEditorTreeItem::GetItemText)
-			.HighlightText(InArgs._HighlightText)
-			.OnVerifyTextChanged(this, &SUMGEditorTreeItem::OnVerifyNameTextChanged)
-			.OnTextCommitted(this, &SUMGEditorTreeItem::OnNameTextCommited)
-			.IsSelected(this, &SUMGEditorTreeItem::IsSelectedExclusively)
+			SNew(SHorizontalBox)
+
+			+ SHorizontalBox::Slot()
+			.AutoWidth()
+			[
+				SNew(SImage)
+				.ColorAndOpacity(FLinearColor(1,1,1,0.5))
+				.Image(Item->GetEditorIcon())
+			]
+
+			+ SHorizontalBox::Slot()
+			.FillWidth(1.0f)
+			.Padding(2, 0, 0, 0)
+			.VAlign(VAlign_Center)
+			[
+				SNew(SInlineEditableTextBlock)
+				.ToolTipText(this, &SUMGEditorTreeItem::GetItemTooltipText)
+				.Font(this, &SUMGEditorTreeItem::GetItemFont)
+				.Text(this, &SUMGEditorTreeItem::GetItemText)
+				.HighlightText(InArgs._HighlightText)
+				.OnVerifyTextChanged(this, &SUMGEditorTreeItem::OnVerifyNameTextChanged)
+				.OnTextCommitted(this, &SUMGEditorTreeItem::OnNameTextCommited)
+				.IsSelected(this, &SUMGEditorTreeItem::IsSelectedExclusively)
+			]
 		],
 		InOwnerTableView);
 }
@@ -72,7 +88,7 @@ FSlateFontInfo SUMGEditorTreeItem::GetItemFont() const
 	if ( !Item->IsGeneratedName() && Item->bIsVariable )
 	{
 		// TODO UMG Hacky move into style area
-		return FSlateFontInfo(FPaths::EngineContentDir() / TEXT("Slate/Fonts/Roboto-Bold.ttf"), 9);
+		return FSlateFontInfo(FPaths::EngineContentDir() / TEXT("Slate/Fonts/Roboto-Bold.ttf"), 10);
 	}
 	else
 	{
diff --git a/Engine/Source/Runtime/UMG/Private/Components/Button.cpp b/Engine/Source/Runtime/UMG/Private/Components/Button.cpp
--- a/Engine/Source/Runtime/UMG/Private/Components/Button.cpp
+++ b/Engine/Source/Runtime/UMG/Private/Components/Button.cpp
@@ -136,6 +136,14 @@ bool UButton::IsPressed() const
 	return MyButton->IsPressed();
 }
 
+#if WITH_EDITOR
+
+const FSlateBrush* UButton::GetEditorIcon()
+{
+	return FUMGStyle::Get().GetBrush("Widget.Button");
+}
+
+#endif
 
 /////////////////////////////////////////////////////
 
diff --git a/Engine/Source/Runtime/UMG/Private/Components/CanvasPanel.cpp b/Engine/Source/Runtime/UMG/Private/Components/CanvasPanel.cpp
--- a/Engine/Source/Runtime/UMG/Private/Components/CanvasPanel.cpp
+++ b/Engine/Source/Runtime/UMG/Private/Components/CanvasPanel.cpp
@@ -202,4 +202,9 @@ void UCanvasPanel::ConnectEditorData()
 	}
 }
 
+const FSlateBrush* UCanvasPanel::GetEditorIcon()
+{
+	return FUMGStyle::Get().GetBrush("Widget.Canvas");
+}
+
 #endif
\ No newline at end of file
diff --git a/Engine/Source/Runtime/UMG/Private/Components/CheckBox.cpp b/Engine/Source/Runtime/UMG/Private/Components/CheckBox.cpp
--- a/Engine/Source/Runtime/UMG/Private/Components/CheckBox.cpp
+++ b/Engine/Source/Runtime/UMG/Private/Components/CheckBox.cpp
@@ -111,6 +111,15 @@ void UCheckBox::SlateOnCheckStateChangedCallback(ESlateCheckBoxState::Type NewSt
 	}
 }
 
+#if WITH_EDITOR
+
+const FSlateBrush* UCheckBox::GetEditorIcon()
+{
+	return FUMGStyle::Get().GetBrush("Widget.CheckBox");
+}
+
+#endif
+
 /////////////////////////////////////////////////////
 
 #undef LOCTEXT_NAMESPACE
diff --git a/Engine/Source/Runtime/UMG/Private/Components/EditableText.cpp b/Engine/Source/Runtime/UMG/Private/Components/EditableText.cpp
--- a/Engine/Source/Runtime/UMG/Private/Components/EditableText.cpp
+++ b/Engine/Source/Runtime/UMG/Private/Components/EditableText.cpp
@@ -56,7 +56,7 @@ TSharedRef<SWidget> UEditableText::RebuildWidget()
 	.OnTextCommitted(BIND_UOBJECT_DELEGATE(FOnTextCommitted, HandleOnTextCommitted))
 	;
 	
-	return MyEditableText.ToSharedRef();
+	return BuildDesignTimeWidget( MyEditableText.ToSharedRef() );
 }
 
 void UEditableText::SyncronizeProperties()
diff --git a/Engine/Source/Runtime/UMG/Private/Components/EditableTextBlock.cpp b/Engine/Source/Runtime/UMG/Private/Components/EditableTextBlock.cpp
--- a/Engine/Source/Runtime/UMG/Private/Components/EditableTextBlock.cpp
+++ b/Engine/Source/Runtime/UMG/Private/Components/EditableTextBlock.cpp
@@ -89,6 +89,15 @@ void UEditableTextBlock::SlateOnTextChanged(const FText& Text)
 	OnTextChanged.Broadcast(Text);
 }
 
+#if WITH_EDITOR
+
+const FSlateBrush* UEditableTextBlock::GetEditorIcon()
+{
+	return FUMGStyle::Get().GetBrush("Widget.EditableTextBlock");
+}
+
+#endif
+
 /////////////////////////////////////////////////////
 
 #undef LOCTEXT_NAMESPACE
diff --git a/Engine/Source/Runtime/UMG/Private/Components/HorizontalBox.cpp b/Engine/Source/Runtime/UMG/Private/Components/HorizontalBox.cpp
--- a/Engine/Source/Runtime/UMG/Private/Components/HorizontalBox.cpp
+++ b/Engine/Source/Runtime/UMG/Private/Components/HorizontalBox.cpp
@@ -115,6 +115,7 @@ UHorizontalBoxSlot* UHorizontalBox::AddSlot(UWidget* Content)
 }
 
 #if WITH_EDITOR
+
 void UHorizontalBox::ConnectEditorData()
 {
 	for ( UHorizontalBoxSlot* Slot : Slots )
@@ -123,4 +124,10 @@ void UHorizontalBox::ConnectEditorData()
 		Slot->Parent = this;
 	}
 }
+
+const FSlateBrush* UHorizontalBox::GetEditorIcon()
+{
+	return FUMGStyle::Get().GetBrush("Widget.HorizontalBox");
+}
+
 #endif
diff --git a/Engine/Source/Runtime/UMG/Private/Components/Image.cpp b/Engine/Source/Runtime/UMG/Private/Components/Image.cpp
--- a/Engine/Source/Runtime/UMG/Private/Components/Image.cpp
+++ b/Engine/Source/Runtime/UMG/Private/Components/Image.cpp
@@ -67,6 +67,16 @@ FReply UImage::HandleMouseButtonDown(const FGeometry& Geometry, const FPointerEv
 	return FReply::Unhandled();
 }
 
+#if WITH_EDITOR
+
+const FSlateBrush* UImage::GetEditorIcon()
+{
+	return FUMGStyle::Get().GetBrush("Widget.Image");
+}
+
+#endif
+
+
 /////////////////////////////////////////////////////
 
 #undef LOCTEXT_NAMESPACE
diff --git a/Engine/Source/Runtime/UMG/Private/Components/TextBlock.cpp b/Engine/Source/Runtime/UMG/Private/Components/TextBlock.cpp
--- a/Engine/Source/Runtime/UMG/Private/Components/TextBlock.cpp
+++ b/Engine/Source/Runtime/UMG/Private/Components/TextBlock.cpp
@@ -90,6 +90,11 @@ void UTextBlock::OnDesignerDoubleClicked()
 	MyEditorTextBlock->EnterEditingMode();
 }
 
+const FSlateBrush* UTextBlock::GetEditorIcon()
+{
+	return FUMGStyle::Get().GetBrush("Widget.TextBlock");
+}
+
 #endif
 
 /////////////////////////////////////////////////////
diff --git a/Engine/Source/Runtime/UMG/Private/Components/VerticalBox.cpp b/Engine/Source/Runtime/UMG/Private/Components/VerticalBox.cpp
--- a/Engine/Source/Runtime/UMG/Private/Components/VerticalBox.cpp
+++ b/Engine/Source/Runtime/UMG/Private/Components/VerticalBox.cpp
@@ -115,6 +115,7 @@ UVerticalBoxSlot* UVerticalBox::AddSlot(UWidget* Content)
 }
 
 #if WITH_EDITOR
+
 void UVerticalBox::ConnectEditorData()
 {
 	for ( UVerticalBoxSlot* Slot : Slots )
@@ -129,4 +130,10 @@ void UVerticalBox::ConnectEditorData()
 		//TODO UMG Should we auto delete empty slots?
 	}
 }
+
+const FSlateBrush* UVerticalBox::GetEditorIcon()
+{
+	return FUMGStyle::Get().GetBrush("Widget.VerticalBox");
+}
+
 #endif
diff --git a/Engine/Source/Runtime/UMG/Private/Components/Widget.cpp b/Engine/Source/Runtime/UMG/Private/Components/Widget.cpp
--- a/Engine/Source/Runtime/UMG/Private/Components/Widget.cpp
+++ b/Engine/Source/Runtime/UMG/Private/Components/Widget.cpp
@@ -145,6 +145,11 @@ FString UWidget::GetLabel() const
 	}
 }
 
+const FSlateBrush* UWidget::GetEditorIcon()
+{
+	return FUMGStyle::Get().GetBrush("Widget");
+}
+
 TSharedRef<SWidget> UWidget::GetToolboxPreviewWidget() const
 {
 	return SNew(SImage);
diff --git a/Engine/Source/Runtime/UMG/Private/UMGModule.cpp b/Engine/Source/Runtime/UMG/Private/UMGModule.cpp
--- a/Engine/Source/Runtime/UMG/Private/UMGModule.cpp
+++ b/Engine/Source/Runtime/UMG/Private/UMGModule.cpp
@@ -16,6 +16,8 @@ public:
 	/** Called right after the module DLL has been loaded and the module object has been created */
 	virtual void StartupModule() override
 	{
+		FUMGStyle::Initialize();
+
 #if WITH_EDITOR
 		// This is done so that the compiler is available in non-cooked builds when the widget blueprint is 
 		// compiled again in the running game.
@@ -29,6 +31,7 @@ public:
 	/** Called before the module is unloaded, right before the module object is destroyed. */
 	virtual void ShutdownModule() override
 	{
+		FUMGStyle::Shutdown();
 	}
 };
 
diff --git a/Engine/Source/Runtime/UMG/Private/UMGPrivatePCH.h b/Engine/Source/Runtime/UMG/Private/UMGPrivatePCH.h
--- a/Engine/Source/Runtime/UMG/Private/UMGPrivatePCH.h
+++ b/Engine/Source/Runtime/UMG/Private/UMGPrivatePCH.h
@@ -3,3 +3,5 @@
 #pragma once
 
 #include "UMG.h"
+
+#include "UMGStyle.h"
diff --git a/Engine/Source/Runtime/UMG/Private/UMGStyle.cpp b/Engine/Source/Runtime/UMG/Private/UMGStyle.cpp
new file mode 100644
--- /dev/null
+++ b/Engine/Source/Runtime/UMG/Private/UMGStyle.cpp
@@ -0,0 +1,75 @@
+// Copyright 1998-2014 Epic Games, Inc. All Rights Reserved.
+
+#include "UMGPrivatePCH.h"
+
+#include "UMGStyle.h"
+#include "SlateGameResources.h"
+
+TSharedPtr< FSlateStyleSet > FUMGStyle::UMGStyleInstance = NULL;
+
+void FUMGStyle::Initialize()
+{
+	if ( !UMGStyleInstance.IsValid() )
+	{
+		UMGStyleInstance = Create();
+		FSlateStyleRegistry::RegisterSlateStyle( *UMGStyleInstance );
+	}
+}
+
+void FUMGStyle::Shutdown()
+{
+	FSlateStyleRegistry::UnRegisterSlateStyle( *UMGStyleInstance );
+	ensure( UMGStyleInstance.IsUnique() );
+	UMGStyleInstance.Reset();
+}
+
+FName FUMGStyle::GetStyleSetName()
+{
+	static FName StyleSetName(TEXT("UMGStyle"));
+	return StyleSetName;
+}
+
+#define IMAGE_BRUSH( RelativePath, ... ) FSlateImageBrush( Style->RootToContentDir( RelativePath, TEXT(".png") ), __VA_ARGS__ )
+#define BOX_BRUSH( RelativePath, ... ) FSlateBoxBrush( Style->RootToContentDir( RelativePath, TEXT(".png") ), __VA_ARGS__ )
+#define BORDER_BRUSH( RelativePath, ... ) FSlateBorderBrush( Style->RootToContentDir( RelativePath, TEXT(".png") ), __VA_ARGS__ )
+#define TTF_FONT( RelativePath, ... ) FSlateFontInfo( Style->RootToContentDir( RelativePath, TEXT(".ttf") ), __VA_ARGS__ )
+#define OTF_FONT( RelativePath, ... ) FSlateFontInfo( Style->RootToContentDir( RelativePath, TEXT(".otf") ), __VA_ARGS__ )
+
+const FVector2D Icon16x16(16.0f, 16.0f);
+const FVector2D Icon20x20(20.0f, 20.0f);
+
+TSharedRef< FSlateStyleSet > FUMGStyle::Create()
+{
+	TSharedRef< FSlateStyleSet > Style = MakeShareable(new FSlateStyleSet("UMGStyle"));
+	Style->SetContentRoot(FPaths::EngineContentDir() / TEXT("Editor/Slate/UMG"));
+	
+	Style->Set("Widget", new IMAGE_BRUSH(TEXT("Widget"), Icon16x16));
+	Style->Set("Widget.CheckBox", new IMAGE_BRUSH(TEXT("CheckBox"), Icon16x16));
+	Style->Set("Widget.Button", new IMAGE_BRUSH(TEXT("Button"), Icon16x16));
+	Style->Set("Widget.EditableTextBlock", new IMAGE_BRUSH(TEXT("EditableTextBlock"), Icon16x16));
+	Style->Set("Widget.HorizontalBox", new IMAGE_BRUSH(TEXT("HorizontalBox"), Icon16x16));
+	Style->Set("Widget.VerticalBox", new IMAGE_BRUSH(TEXT("VerticalBox"), Icon16x16));
+	Style->Set("Widget.Image", new IMAGE_BRUSH(TEXT("Image"), Icon16x16));
+	Style->Set("Widget.Canvas", new IMAGE_BRUSH(TEXT("Canvas"), Icon16x16));
+	Style->Set("Widget.TextBlock", new IMAGE_BRUSH(TEXT("TextBlock"), Icon16x16));
+
+	Style->Set("Widget.UserWidget", new IMAGE_BRUSH(TEXT("UserWidget"), Icon16x16));
+
+	return Style;
+}
+
+#undef IMAGE_BRUSH
+#undef BOX_BRUSH
+#undef BORDER_BRUSH
+#undef TTF_FONT
+#undef OTF_FONT
+
+void FUMGStyle::ReloadTextures()
+{
+	FSlateApplication::Get().GetRenderer()->ReloadTextureResources();
+}
+
+const ISlateStyle& FUMGStyle::Get()
+{
+	return *UMGStyleInstance;
+}
diff --git a/Engine/Source/Runtime/UMG/Private/UMGStyle.h b/Engine/Source/Runtime/UMG/Private/UMGStyle.h
new file mode 100644
--- /dev/null
+++ b/Engine/Source/Runtime/UMG/Private/UMGStyle.h
@@ -0,0 +1,30 @@
+// Copyright 1998-2014 Epic Games, Inc. All Rights Reserved.
+
+#pragma once
+
+#include "Slate.h"
+
+class FUMGStyle
+{
+public:
+
+	static void Initialize();
+
+	static void Shutdown();
+
+	/** reloads textures used by slate renderer */
+	static void ReloadTextures();
+
+	/** @return The Slate style set for the Shooter game */
+	static const ISlateStyle& Get();
+
+	static FName GetStyleSetName();
+
+private:
+
+	static TSharedRef< class FSlateStyleSet > Create();
+
+private:
+
+	static TSharedPtr< class FSlateStyleSet > UMGStyleInstance;
+};
\ No newline at end of file
diff --git a/Engine/Source/Runtime/UMG/Public/Components/Button.h b/Engine/Source/Runtime/UMG/Public/Components/Button.h
--- a/Engine/Source/Runtime/UMG/Public/Components/Button.h
+++ b/Engine/Source/Runtime/UMG/Public/Components/Button.h
@@ -87,6 +87,10 @@ public:
 	virtual void SyncronizeProperties() override;
 	// End of UWidget interface
 
+#if WITH_EDITOR
+	virtual const FSlateBrush* GetEditorIcon() override;
+#endif
+
 protected:
 	// UWidget interface
 	virtual TSharedRef<SWidget> RebuildWidget() override;
diff --git a/Engine/Source/Runtime/UMG/Public/Components/CanvasPanel.h b/Engine/Source/Runtime/UMG/Public/Components/CanvasPanel.h
--- a/Engine/Source/Runtime/UMG/Public/Components/CanvasPanel.h
+++ b/Engine/Source/Runtime/UMG/Public/Components/CanvasPanel.h
@@ -42,6 +42,7 @@ class UMG_API UCanvasPanel : public UPanelWidget
 
 #if WITH_EDITOR
 	// UWidget interface
+	virtual const FSlateBrush* GetEditorIcon() override;
 	virtual void ConnectEditorData() override;
 	// End UWidget interface
 #endif
diff --git a/Engine/Source/Runtime/UMG/Public/Components/CheckBox.h b/Engine/Source/Runtime/UMG/Public/Components/CheckBox.h
--- a/Engine/Source/Runtime/UMG/Public/Components/CheckBox.h
+++ b/Engine/Source/Runtime/UMG/Public/Components/CheckBox.h
@@ -103,6 +103,10 @@ protected:
 	virtual void SyncronizeProperties() override;
 	// End of UWidget interface
 
+#if WITH_EDITOR
+	virtual const FSlateBrush* GetEditorIcon() override;
+#endif
+
 protected:
 	// UWidget interface
 	virtual TSharedRef<SWidget> RebuildWidget() override;
diff --git a/Engine/Source/Runtime/UMG/Public/Components/EditableTextBlock.h b/Engine/Source/Runtime/UMG/Public/Components/EditableTextBlock.h
--- a/Engine/Source/Runtime/UMG/Public/Components/EditableTextBlock.h
+++ b/Engine/Source/Runtime/UMG/Public/Components/EditableTextBlock.h
@@ -98,6 +98,10 @@ protected:
 	virtual void SyncronizeProperties() override;
 	// End of UWidget interface
 
+#if WITH_EDITOR
+	virtual const FSlateBrush* GetEditorIcon() override;
+#endif
+
 protected:
 	// UWidget interface
 	virtual TSharedRef<SWidget> RebuildWidget() override;
diff --git a/Engine/Source/Runtime/UMG/Public/Components/HorizontalBox.h b/Engine/Source/Runtime/UMG/Public/Components/HorizontalBox.h
--- a/Engine/Source/Runtime/UMG/Public/Components/HorizontalBox.h
+++ b/Engine/Source/Runtime/UMG/Public/Components/HorizontalBox.h
@@ -27,6 +27,7 @@ class UMG_API UHorizontalBox : public UPanelWidget
 
 #if WITH_EDITOR
 	// UWidget interface
+	virtual const FSlateBrush* GetEditorIcon() override;
 	virtual void ConnectEditorData() override;
 	// End UWidget interface
 #endif
diff --git a/Engine/Source/Runtime/UMG/Public/Components/Image.h b/Engine/Source/Runtime/UMG/Public/Components/Image.h
--- a/Engine/Source/Runtime/UMG/Public/Components/Image.h
+++ b/Engine/Source/Runtime/UMG/Public/Components/Image.h
@@ -37,6 +37,12 @@ public:
 	virtual void SyncronizeProperties() override;
 	// End of UWidget interface
 
+#if WITH_EDITOR
+	// UWidget interface
+	virtual const FSlateBrush* GetEditorIcon() override;
+	// End UWidget interface
+#endif
+
 protected:
 	// UWidget interface
 	virtual TSharedRef<SWidget> RebuildWidget() override;
diff --git a/Engine/Source/Runtime/UMG/Public/Components/TextBlock.h b/Engine/Source/Runtime/UMG/Public/Components/TextBlock.h
--- a/Engine/Source/Runtime/UMG/Public/Components/TextBlock.h
+++ b/Engine/Source/Runtime/UMG/Public/Components/TextBlock.h
@@ -63,6 +63,10 @@ public:
 	// End of UWidget interface
 
 #if WITH_EDITOR
+	// UWidget interface
+	virtual const FSlateBrush* GetEditorIcon() override;
+	// End UWidget interface
+
 	void HandleTextCommitted(const FText& InText, ETextCommit::Type CommitteType);
 
 	virtual void OnDesignerDoubleClicked() override;
diff --git a/Engine/Source/Runtime/UMG/Public/Components/VerticalBox.h b/Engine/Source/Runtime/UMG/Public/Components/VerticalBox.h
--- a/Engine/Source/Runtime/UMG/Public/Components/VerticalBox.h
+++ b/Engine/Source/Runtime/UMG/Public/Components/VerticalBox.h
@@ -33,6 +33,7 @@ class UMG_API UVerticalBox : public UPanelWidget
 
 #if WITH_EDITOR
 	// UWidget interface
+	virtual const FSlateBrush* GetEditorIcon() override;
 	virtual void ConnectEditorData() override;
 	// End UWidget interface
 #endif
diff --git a/Engine/Source/Runtime/UMG/Public/Components/Widget.h b/Engine/Source/Runtime/UMG/Public/Components/Widget.h
--- a/Engine/Source/Runtime/UMG/Public/Components/Widget.h
+++ b/Engine/Source/Runtime/UMG/Public/Components/Widget.h
@@ -149,6 +149,10 @@ public:
 	/** Gets the label to display to the user for this widget. */
 	FString GetLabel() const;
 
+	/** Gets the editor icon */
+	virtual const FSlateBrush* GetEditorIcon();
+
+	//TODO UMG These are not used yet, and may never be.
 	virtual void OnDesignerSelected() { }
 	virtual void OnDesignerDeselected() { }
 	virtual void OnDesignerClicked() { }
