UMG - More icons for the widgets in the tree.

diff --git a/Engine/Source/Runtime/UMG/Private/Components/Border.cpp b/Engine/Source/Runtime/UMG/Private/Components/Border.cpp
--- a/Engine/Source/Runtime/UMG/Private/Components/Border.cpp
+++ b/Engine/Source/Runtime/UMG/Private/Components/Border.cpp
@@ -150,6 +150,15 @@ FReply UBorder::HandleMouseDoubleClick(const FGeometry& Geometry, const FPointer
 	return FReply::Unhandled();
 }
 
+#if WITH_EDITOR
+
+const FSlateBrush* UBorder::GetEditorIcon()
+{
+	return FUMGStyle::Get().GetBrush("Widget.Border");
+}
+
+#endif
+
 /////////////////////////////////////////////////////
 
 #undef LOCTEXT_NAMESPACE
diff --git a/Engine/Source/Runtime/UMG/Private/Components/ScrollBox.cpp b/Engine/Source/Runtime/UMG/Private/Components/ScrollBox.cpp
--- a/Engine/Source/Runtime/UMG/Private/Components/ScrollBox.cpp
+++ b/Engine/Source/Runtime/UMG/Private/Components/ScrollBox.cpp
@@ -119,6 +119,11 @@ UScrollBoxSlot* UScrollBox::AddSlot(UWidget* Content)
 
 #if WITH_EDITOR
 
+const FSlateBrush* UScrollBox::GetEditorIcon()
+{
+	return FUMGStyle::Get().GetBrush("Widget.ScrollBox");
+}
+
 void UScrollBox::ConnectEditorData()
 {
 	for ( UScrollBoxSlot* Slot : Slots )
diff --git a/Engine/Source/Runtime/UMG/Private/Components/Slider.cpp b/Engine/Source/Runtime/UMG/Private/Components/Slider.cpp
--- a/Engine/Source/Runtime/UMG/Private/Components/Slider.cpp
+++ b/Engine/Source/Runtime/UMG/Private/Components/Slider.cpp
@@ -61,6 +61,15 @@ void USlider::SetValue(float InValue)
 	return MySlider->SetValue(InValue);
 }
 
+#if WITH_EDITOR
+
+const FSlateBrush* USlider::GetEditorIcon()
+{
+	return FUMGStyle::Get().GetBrush("Widget.Slider");
+}
+
+#endif
+
 /////////////////////////////////////////////////////
 
 #undef LOCTEXT_NAMESPACE
diff --git a/Engine/Source/Runtime/UMG/Private/Components/Spacer.cpp b/Engine/Source/Runtime/UMG/Private/Components/Spacer.cpp
--- a/Engine/Source/Runtime/UMG/Private/Components/Spacer.cpp
+++ b/Engine/Source/Runtime/UMG/Private/Components/Spacer.cpp
@@ -41,6 +41,15 @@ void USpacer::SyncronizeProperties()
 	MySpacer->SetSize(Size);
 }
 
+#if WITH_EDITOR
+
+const FSlateBrush* USpacer::GetEditorIcon()
+{
+	return FUMGStyle::Get().GetBrush("Widget.Spacer");
+}
+
+#endif
+
 /////////////////////////////////////////////////////
 
 #undef LOCTEXT_NAMESPACE
diff --git a/Engine/Source/Runtime/UMG/Private/UMGStyle.cpp b/Engine/Source/Runtime/UMG/Private/UMGStyle.cpp
--- a/Engine/Source/Runtime/UMG/Private/UMGStyle.cpp
+++ b/Engine/Source/Runtime/UMG/Private/UMGStyle.cpp
@@ -52,6 +52,10 @@ TSharedRef< FSlateStyleSet > FUMGStyle::Create()
 	Style->Set("Widget.Image", new IMAGE_BRUSH(TEXT("Image"), Icon16x16));
 	Style->Set("Widget.Canvas", new IMAGE_BRUSH(TEXT("Canvas"), Icon16x16));
 	Style->Set("Widget.TextBlock", new IMAGE_BRUSH(TEXT("TextBlock"), Icon16x16));
+	Style->Set("Widget.Border", new IMAGE_BRUSH(TEXT("Border"), Icon16x16));
+	Style->Set("Widget.Slider", new IMAGE_BRUSH(TEXT("Slider"), Icon16x16));
+	Style->Set("Widget.Spacer", new IMAGE_BRUSH(TEXT("Spacer"), Icon16x16));
+	Style->Set("Widget.ScrollBox", new IMAGE_BRUSH(TEXT("ScrollBox"), Icon16x16));
 
 	Style->Set("Widget.UserWidget", new IMAGE_BRUSH(TEXT("UserWidget"), Icon16x16));
 
diff --git a/Engine/Source/Runtime/UMG/Public/Components/Border.h b/Engine/Source/Runtime/UMG/Public/Components/Border.h
--- a/Engine/Source/Runtime/UMG/Public/Components/Border.h
+++ b/Engine/Source/Runtime/UMG/Public/Components/Border.h
@@ -95,6 +95,10 @@ class UMG_API UBorder : public UContentWidget
 	void SyncronizeProperties() override;
 	// End of UWidget interface
 
+#if WITH_EDITOR
+	virtual const FSlateBrush* GetEditorIcon() override;
+#endif
+
 protected:
 	TSharedPtr<SBorder> MyBorder;
 
diff --git a/Engine/Source/Runtime/UMG/Public/Components/ScrollBox.h b/Engine/Source/Runtime/UMG/Public/Components/ScrollBox.h
--- a/Engine/Source/Runtime/UMG/Public/Components/ScrollBox.h
+++ b/Engine/Source/Runtime/UMG/Public/Components/ScrollBox.h
@@ -30,6 +30,7 @@ class UMG_API UScrollBox : public UPanelWidget
 
 #if WITH_EDITOR
 	// UWidget interface
+	virtual const FSlateBrush* GetEditorIcon() override;
 	virtual void ConnectEditorData() override;
 	// End UWidget interface
 #endif
diff --git a/Engine/Source/Runtime/UMG/Public/Components/Slider.h b/Engine/Source/Runtime/UMG/Public/Components/Slider.h
--- a/Engine/Source/Runtime/UMG/Public/Components/Slider.h
+++ b/Engine/Source/Runtime/UMG/Public/Components/Slider.h
@@ -65,6 +65,10 @@ class UMG_API USlider : public UWidget
 	virtual void SyncronizeProperties() override;
 	// End of UWidget interface
 
+#if WITH_EDITOR
+	virtual const FSlateBrush* GetEditorIcon() override;
+#endif
+
 protected:
 	/** Native Slate Widget */
 	TSharedPtr<SSlider> MySlider;
diff --git a/Engine/Source/Runtime/UMG/Public/Components/Spacer.h b/Engine/Source/Runtime/UMG/Public/Components/Spacer.h
--- a/Engine/Source/Runtime/UMG/Public/Components/Spacer.h
+++ b/Engine/Source/Runtime/UMG/Public/Components/Spacer.h
@@ -22,6 +22,10 @@ class UMG_API USpacer : public UWidget
 	virtual void SyncronizeProperties() override;
 	// End of UWidget interface
 
+#if WITH_EDITOR
+	virtual const FSlateBrush* GetEditorIcon() override;
+#endif
+
 protected:
 	// UWidget interface
 	virtual TSharedRef<SWidget> RebuildWidget() override;
