UMG - Support for a few more common control icons.

diff --git a/Engine/Source/Runtime/UMG/Private/Components/EditableText.cpp b/Engine/Source/Runtime/UMG/Private/Components/EditableText.cpp
--- a/Engine/Source/Runtime/UMG/Private/Components/EditableText.cpp
+++ b/Engine/Source/Runtime/UMG/Private/Components/EditableText.cpp
@@ -101,6 +101,15 @@ void UEditableText::HandleOnTextCommitted(const FText& Text, ETextCommit::Type C
 	OnTextCommitted.Broadcast(Text, CommitMethod);
 }
 
+#if WITH_EDITOR
+
+const FSlateBrush* UEditableText::GetEditorIcon()
+{
+	return FUMGStyle::Get().GetBrush("Widget.EditableText");
+}
+
+#endif
+
 /////////////////////////////////////////////////////
 
 #undef LOCTEXT_NAMESPACE
diff --git a/Engine/Source/Runtime/UMG/Private/Components/ProcessBar.cpp b/Engine/Source/Runtime/UMG/Private/Components/ProcessBar.cpp
--- a/Engine/Source/Runtime/UMG/Private/Components/ProcessBar.cpp
+++ b/Engine/Source/Runtime/UMG/Private/Components/ProcessBar.cpp
@@ -47,6 +47,15 @@ void UProgressBar::SetPercent(float InPercent)
 	return MyProgressBar->SetPercent(InPercent);
 }
 
+#if WITH_EDITOR
+
+const FSlateBrush* UProgressBar::GetEditorIcon()
+{
+	return FUMGStyle::Get().GetBrush("Widget.ProgressBar");
+}
+
+#endif
+
 /////////////////////////////////////////////////////
 
 #undef LOCTEXT_NAMESPACE
diff --git a/Engine/Source/Runtime/UMG/Private/UMGStyle.cpp b/Engine/Source/Runtime/UMG/Private/UMGStyle.cpp
--- a/Engine/Source/Runtime/UMG/Private/UMGStyle.cpp
+++ b/Engine/Source/Runtime/UMG/Private/UMGStyle.cpp
@@ -47,6 +47,7 @@ TSharedRef< FSlateStyleSet > FUMGStyle::Create()
 	Style->Set("Widget.CheckBox", new IMAGE_BRUSH(TEXT("CheckBox"), Icon16x16));
 	Style->Set("Widget.Button", new IMAGE_BRUSH(TEXT("Button"), Icon16x16));
 	Style->Set("Widget.EditableTextBlock", new IMAGE_BRUSH(TEXT("EditableTextBlock"), Icon16x16));
+	Style->Set("Widget.EditableText", new IMAGE_BRUSH(TEXT("EditableText"), Icon16x16));
 	Style->Set("Widget.HorizontalBox", new IMAGE_BRUSH(TEXT("HorizontalBox"), Icon16x16));
 	Style->Set("Widget.VerticalBox", new IMAGE_BRUSH(TEXT("VerticalBox"), Icon16x16));
 	Style->Set("Widget.Image", new IMAGE_BRUSH(TEXT("Image"), Icon16x16));
@@ -56,6 +57,7 @@ TSharedRef< FSlateStyleSet > FUMGStyle::Create()
 	Style->Set("Widget.Slider", new IMAGE_BRUSH(TEXT("Slider"), Icon16x16));
 	Style->Set("Widget.Spacer", new IMAGE_BRUSH(TEXT("Spacer"), Icon16x16));
 	Style->Set("Widget.ScrollBox", new IMAGE_BRUSH(TEXT("ScrollBox"), Icon16x16));
+	Style->Set("Widget.ProgressBar", new IMAGE_BRUSH(TEXT("ProgressBar"), Icon16x16));
 
 	Style->Set("Widget.UserWidget", new IMAGE_BRUSH(TEXT("UserWidget"), Icon16x16));
 
diff --git a/Engine/Source/Runtime/UMG/Public/Blueprint/UserWidget.h b/Engine/Source/Runtime/UMG/Public/Blueprint/UserWidget.h
--- a/Engine/Source/Runtime/UMG/Public/Blueprint/UserWidget.h
+++ b/Engine/Source/Runtime/UMG/Public/Blueprint/UserWidget.h
@@ -246,15 +246,22 @@ class UMG_API UUserWidget : public UWidget
 	UFUNCTION(BlueprintImplementableEvent, Category="User Interface")
 	FSReply OnMotionDetected(FGeometry MyGeometry, FMotionEvent InMotionEvent);
 
-
+	/** @returns The UObject wrapper for a given SWidget */
 	UWidget* GetWidgetHandle(TSharedRef<SWidget> InWidget);
 
+	/** Creates this widget and all children and returns them. */
 	TSharedRef<SWidget> MakeWidget();
+
+	/** Creates a fullscreen host widget, that wraps this widget. */
 	TSharedRef<SWidget> MakeFullScreenWidget();
 
+	/** @returns The root UObject widget wrapper */
 	UWidget* GetRootWidgetComponent();
 
+	/** @returns The slate widget corresponding to a given name */
 	TSharedPtr<SWidget> GetWidgetFromName(const FString& Name) const;
+
+	/** @returns The uobject widget corresponding to a given name */
 	UWidget* GetHandleFromName(const FString& Name) const;
 
 protected:
diff --git a/Engine/Source/Runtime/UMG/Public/Components/EditableText.h b/Engine/Source/Runtime/UMG/Public/Components/EditableText.h
--- a/Engine/Source/Runtime/UMG/Public/Components/EditableText.h
+++ b/Engine/Source/Runtime/UMG/Public/Components/EditableText.h
@@ -103,6 +103,10 @@ protected:
 	virtual void SyncronizeProperties() override;
 	// End of UWidget interface
 
+#if WITH_EDITOR
+	virtual const FSlateBrush* GetEditorIcon() override;
+#endif
+
 protected:
 	// UWidget interface
 	virtual TSharedRef<SWidget> RebuildWidget() override;
diff --git a/Engine/Source/Runtime/UMG/Public/Components/ProgressBar.h b/Engine/Source/Runtime/UMG/Public/Components/ProgressBar.h
--- a/Engine/Source/Runtime/UMG/Public/Components/ProgressBar.h
+++ b/Engine/Source/Runtime/UMG/Public/Components/ProgressBar.h
@@ -55,6 +55,12 @@ public:
 	virtual void SyncronizeProperties() override;
 	// End of UWidget interface
 
+#if WITH_EDITOR
+	// UWidget interface
+	virtual const FSlateBrush* GetEditorIcon() override;
+	// End UWidget interface
+#endif
+
 protected:
 	/** Native Slate Widget */
 	TSharedPtr<SProgressBar> MyProgressBar;
