UMG - Added mutators to UTextBlock so it could be properly animated

diff --git a/Engine/Source/Runtime/UMG/Private/Components/TextBlock.cpp b/Engine/Source/Runtime/UMG/Private/Components/TextBlock.cpp
--- a/Engine/Source/Runtime/UMG/Private/Components/TextBlock.cpp
+++ b/Engine/Source/Runtime/UMG/Private/Components/TextBlock.cpp
@@ -23,6 +23,33 @@ UTextBlock::UTextBlock(const FPostConstructInitializeProperties& PCIP)
 	Font = FSlateFontInfo(TEXT("Slate/Fonts/Roboto-Bold.ttf"), 24);
 }
 
+void UTextBlock::SetColorAndOpacity(FLinearColor InColorAndOpacity)
+{
+	ColorAndOpacity = InColorAndOpacity;
+	if( MyTextBlock.IsValid() )
+	{
+		MyTextBlock->SetColorAndOpacity( InColorAndOpacity );
+	}
+}
+
+void UTextBlock::SetShadowColorAndOpacity(FLinearColor InShadowColorAndOpacity)
+{
+	ShadowColorAndOpacity = InShadowColorAndOpacity;
+	if(MyTextBlock.IsValid())
+	{
+		MyTextBlock->SetShadowColorAndOpacity(InShadowColorAndOpacity);
+	}
+}
+
+void UTextBlock::SetShadowOffset(FVector2D InShadowOffset)
+{
+	ShadowOffset = InShadowOffset;
+	if(MyTextBlock.IsValid())
+	{
+		MyTextBlock->SetShadowOffset(ShadowOffset);
+	}
+}
+
 TSharedRef<SWidget> UTextBlock::RebuildWidget()
 {
 	MyTextBlock = SNew(STextBlock);
diff --git a/Engine/Source/Runtime/UMG/Public/Components/TextBlock.h b/Engine/Source/Runtime/UMG/Public/Components/TextBlock.h
--- a/Engine/Source/Runtime/UMG/Public/Components/TextBlock.h
+++ b/Engine/Source/Runtime/UMG/Public/Components/TextBlock.h
@@ -10,6 +10,33 @@ class UMG_API UTextBlock : public UWidget
 {
 	GENERATED_UCLASS_BODY()
 
+public:
+	/**  
+ 	 * Sets the color and opacity of the text in this text block
+	 *
+	 * @param InColorAndOpacity		The new text color and opacity
+ 	 */
+	UFUNCTION(BlueprintCallable, Category="Appearance")
+	void SetColorAndOpacity(FLinearColor InColorAndOpacity);
+
+	/**  
+ 	 * Sets the color and opacity of the text drop shadow
+	 * Note: if opacity is zero no shadow will be drawn
+	 *
+	 * @param InShadowColorAndOpacity		The new drop shadow color and opacity
+ 	 */
+	UFUNCTION(BlueprintCallable, Category="Appearance")
+	void SetShadowColorAndOpacity(FLinearColor InShadowColorAndOpacity);
+
+	/**  
+ 	 * Sets the offset that the text drop shadow should be drawn at
+	 *
+	 * @param InShadowOffset		The new offset
+ 	 */
+	UFUNCTION(BlueprintCallable, Category="Appearance")
+	void SetShadowOffset(FVector2D InShadowOffset);
+
+
 public:
 	/** The style to use to render the text */
 	UPROPERTY(EditDefaultsOnly, Category=Style, meta=( DisplayThumbnail = "true" ))
