Fix non-unity builds

diff --git a/Engine/Source/Editor/Sequencer/Public/ISequencer.h b/Engine/Source/Editor/Sequencer/Public/ISequencer.h
--- a/Engine/Source/Editor/Sequencer/Public/ISequencer.h
+++ b/Engine/Source/Editor/Sequencer/Public/ISequencer.h
@@ -3,8 +3,6 @@
 #pragma once
 
 
-#include "Toolkits/AssetEditorToolkit.h"
-#include "Toolkits/IToolkit.h"
 #include "Runtime/MovieSceneCore/Public/IMovieScenePlayer.h"
 
 class UMovieScene;
diff --git a/Engine/Source/Editor/Sequencer/Public/ISequencerModule.h b/Engine/Source/Editor/Sequencer/Public/ISequencerModule.h
--- a/Engine/Source/Editor/Sequencer/Public/ISequencerModule.h
+++ b/Engine/Source/Editor/Sequencer/Public/ISequencerModule.h
@@ -5,6 +5,9 @@
 #include "ModuleInterface.h"
 #include "ISequencer.h"
 
+#include "Toolkits/AssetEditorToolkit.h"
+#include "Toolkits/IToolkit.h"
+
 class ISequencerObjectBindingManager;
 
 /**
diff --git a/Engine/Source/Editor/Sequencer/Public/MovieSceneTrackEditor.h b/Engine/Source/Editor/Sequencer/Public/MovieSceneTrackEditor.h
--- a/Engine/Source/Editor/Sequencer/Public/MovieSceneTrackEditor.h
+++ b/Engine/Source/Editor/Sequencer/Public/MovieSceneTrackEditor.h
@@ -4,6 +4,7 @@
 
 //@todo Sequencer - these have to be here for now because this class contains a small amount of implementation inside this header to avoid exporting this class
 #include "ISequencer.h"
+#include "MovieScene.h"
 
 class ISequencerSection;
 
diff --git a/Engine/Source/Editor/UMGEditor/Private/UMGSequencerObjectBindingManager.cpp b/Engine/Source/Editor/UMGEditor/Private/UMGSequencerObjectBindingManager.cpp
--- a/Engine/Source/Editor/UMGEditor/Private/UMGSequencerObjectBindingManager.cpp
+++ b/Engine/Source/Editor/UMGEditor/Private/UMGSequencerObjectBindingManager.cpp
@@ -2,6 +2,8 @@
 
 #include "UMGEditorPrivatePCH.h"
 #include "UMGSequencerObjectBindingManager.h"
+#include "WidgetBlueprintEditor.h"
+#include "ISequencer.h"
 
 FUMGSequencerObjectBindingManager::FUMGSequencerObjectBindingManager( FWidgetBlueprintEditor& InWidgetBlueprintEditor )
 	: WidgetBlueprintEditor( InWidgetBlueprintEditor )
diff --git a/Engine/Source/Runtime/UMG/Private/Animation/MovieSceneMarginTrackInstance.cpp b/Engine/Source/Runtime/UMG/Private/Animation/MovieSceneMarginTrackInstance.cpp
--- a/Engine/Source/Runtime/UMG/Private/Animation/MovieSceneMarginTrackInstance.cpp
+++ b/Engine/Source/Runtime/UMG/Private/Animation/MovieSceneMarginTrackInstance.cpp
@@ -2,6 +2,8 @@
 
 #include "UMGPrivatePCH.h"
 #include "MovieSceneMarginTrackInstance.h"
+#include "MovieSceneMarginTrack.h"
+#include "MovieSceneCommonHelpers.h"
 
 FMovieSceneMarginTrackInstance::FMovieSceneMarginTrackInstance( UMovieSceneMarginTrack& InMarginTrack )
 {
diff --git a/Engine/Source/Runtime/UMG/Private/Animation/MovieSceneMarginTrackInstance.h b/Engine/Source/Runtime/UMG/Private/Animation/MovieSceneMarginTrackInstance.h
--- a/Engine/Source/Runtime/UMG/Private/Animation/MovieSceneMarginTrackInstance.h
+++ b/Engine/Source/Runtime/UMG/Private/Animation/MovieSceneMarginTrackInstance.h
@@ -10,7 +10,7 @@
 class FMovieSceneMarginTrackInstance : public IMovieSceneTrackInstance
 {
 public:
-	FMovieSceneMarginTrackInstance( UMovieSceneMarginTrack& InMarginTrack );
+	FMovieSceneMarginTrackInstance( class UMovieSceneMarginTrack& InMarginTrack );
 
 	/** IMovieSceneTrackInstance interface */
 	virtual void Update( float Position, float LastPosition, const TArray<UObject*>& RuntimeObjects, class IMovieScenePlayer& Player ) override;
