UMG - Vastly simplifying the boiler plate code needed to wrap a slate panel.  This change will make it much easier to expose everything to blueprints.  Previous UMG assets may not work correctly after this change.

diff --git a/Engine/Source/Editor/UMGEditor/Private/Templates/WidgetTemplateButton.cpp b/Engine/Source/Editor/UMGEditor/Private/Templates/WidgetTemplateButton.cpp
--- a/Engine/Source/Editor/UMGEditor/Private/Templates/WidgetTemplateButton.cpp
+++ b/Engine/Source/Editor/UMGEditor/Private/Templates/WidgetTemplateButton.cpp
@@ -25,7 +25,7 @@ UWidget* FWidgetTemplateButton::Create(UWidgetTree* Tree)
 
 	Text->Text = FText::FromString(TEXT("Button Text"));
 	Text->Font.Size = 10;
-	Button->SetContent(Text);
+	Button->AddChild(Text, FVector2D(0, 0));
 
 	return Button;
 }
diff --git a/Engine/Source/Editor/UMGEditor/Private/Templates/WidgetTemplateCheckBox.cpp b/Engine/Source/Editor/UMGEditor/Private/Templates/WidgetTemplateCheckBox.cpp
--- a/Engine/Source/Editor/UMGEditor/Private/Templates/WidgetTemplateCheckBox.cpp
+++ b/Engine/Source/Editor/UMGEditor/Private/Templates/WidgetTemplateCheckBox.cpp
@@ -26,9 +26,9 @@ UWidget* FWidgetTemplateCheckBox::Create(UWidgetTree* Tree)
 	Text->Text = LOCTEXT("CheckboxText", "Checkbox Text");
 	Text->Font.Size = 10;
 
-	Horizontal->AddSlot(Checkbox)
+	Horizontal->AddChild(Checkbox)
 		->Size = FSlateChildSize(ESlateSizeRule::Automatic);
-	Horizontal->AddSlot(Text);
+	Horizontal->AddChild(Text);
 
 	return Horizontal;
 }
diff --git a/Engine/Source/Runtime/UMG/Private/Components/Border.cpp b/Engine/Source/Runtime/UMG/Private/Components/Border.cpp
--- a/Engine/Source/Runtime/UMG/Private/Components/Border.cpp
+++ b/Engine/Source/Runtime/UMG/Private/Components/Border.cpp
@@ -32,7 +32,11 @@ UBorder::UBorder(const FPostConstructInitializeProperties& PCIP)
 TSharedRef<SWidget> UBorder::RebuildWidget()
 {
 	MyBorder = SNew(SBorder);
-	MyBorder->SetContent(GetContentSlot()->Content ? GetContentSlot()->Content->GetWidget() : SNullWidget::NullWidget);
+	
+	if ( GetChildrenCount() > 0 )
+	{
+		MyBorder->SetContent(GetContentSlot()->Content ? GetContentSlot()->Content->GetWidget() : SNullWidget::NullWidget);
+	}
 
 	return MyBorder.ToSharedRef();
 }
@@ -62,13 +66,21 @@ void UBorder::SyncronizeProperties()
 	MyBorder->SetOnMouseDoubleClick(BIND_UOBJECT_DELEGATE(FPointerEventHandler, HandleMouseDoubleClick));
 }
 
-void UBorder::SetContent(UWidget* InContent)
+void UBorder::OnSlotAdded(UPanelSlot* Slot)
 {
-	Super::SetContent(InContent);
+	// Add the child to the live canvas if it already exists
+	if ( MyBorder.IsValid() )
+	{
+		MyBorder->SetContent(Slot->Content ? Slot->Content->GetWidget() : SNullWidget::NullWidget);
+	}
+}
 
+void UBorder::OnSlotRemoved(UPanelSlot* Slot)
+{
+	// Remove the widget from the live slot if it exists.
 	if ( MyBorder.IsValid() )
 	{
-		MyBorder->SetContent(InContent ? InContent->GetWidget() : SNullWidget::NullWidget);
+		MyBorder->SetContent(SNullWidget::NullWidget);
 	}
 }
 
diff --git a/Engine/Source/Runtime/UMG/Private/Components/Button.cpp b/Engine/Source/Runtime/UMG/Private/Components/Button.cpp
--- a/Engine/Source/Runtime/UMG/Private/Components/Button.cpp
+++ b/Engine/Source/Runtime/UMG/Private/Components/Button.cpp
@@ -29,7 +29,11 @@ UButton::UButton(const FPostConstructInitializeProperties& PCIP)
 TSharedRef<SWidget> UButton::RebuildWidget()
 {
 	MyButton = SNew(SButton);
-	MyButton->SetContent(GetContentSlot()->Content ? GetContentSlot()->Content->GetWidget() : SNullWidget::NullWidget);
+
+	if ( GetChildrenCount() > 0 )
+	{
+		MyButton->SetContent(GetContentSlot()->Content ? GetContentSlot()->Content->GetWidget() : SNullWidget::NullWidget);
+	}
 	
 	return MyButton.ToSharedRef();
 }
@@ -75,13 +79,21 @@ void UButton::SyncronizeProperties()
 	MyButton->SetOnClicked(BIND_UOBJECT_DELEGATE(FOnClicked, HandleOnClicked));
 }
 
-void UButton::SetContent(UWidget* InContent)
+void UButton::OnSlotAdded(UPanelSlot* Slot)
 {
-	Super::SetContent(InContent);
+	// Add the child to the live canvas if it already exists
+	if ( MyButton.IsValid() )
+	{
+		MyButton->SetContent(Slot->Content ? Slot->Content->GetWidget() : SNullWidget::NullWidget);
+	}
+}
 
+void UButton::OnSlotRemoved(UPanelSlot* Slot)
+{
+	// Remove the widget from the live slot if it exists.
 	if ( MyButton.IsValid() )
 	{
-		MyButton->SetContent(InContent ? InContent->GetWidget() : SNullWidget::NullWidget);
+		MyButton->SetContent(SNullWidget::NullWidget);
 	}
 }
 
diff --git a/Engine/Source/Runtime/UMG/Private/Components/CanvasPanel.cpp b/Engine/Source/Runtime/UMG/Private/Components/CanvasPanel.cpp
--- a/Engine/Source/Runtime/UMG/Private/Components/CanvasPanel.cpp
+++ b/Engine/Source/Runtime/UMG/Private/Components/CanvasPanel.cpp
@@ -46,99 +46,40 @@ UCanvasPanel::UCanvasPanel(const FPostConstructInitializeProperties& PCIP)
 	DesiredCanvasSize = FVector2D(128.0f, 128.0f);
 }
 
-int32 UCanvasPanel::GetChildrenCount() const
+UClass* UCanvasPanel::GetSlotClass() const
 {
-	return Slots.Num();
+	return UCanvasPanelSlot::StaticClass();
 }
 
-UWidget* UCanvasPanel::GetChildAt(int32 Index) const
+void UCanvasPanel::OnSlotAdded(UPanelSlot* Slot)
 {
-	return Slots[Index]->Content;
-}
-
-TSharedPtr<SConstraintCanvas> UCanvasPanel::GetCanvasWidget() const
-{
-	return MyCanvas;
-}
-
-int32 UCanvasPanel::GetChildIndex(UWidget* Content) const
-{
-	for ( int32 SlotIndex = 0; SlotIndex < Slots.Num(); ++SlotIndex )
-	{
-		UCanvasPanelSlot* Slot = Slots[SlotIndex];
-
-		if ( Slot->Content == Content )
-		{
-			return SlotIndex;
-		}
-	}
-
-	return -1;
-}
-
-bool UCanvasPanel::AddChild(UWidget* Child, FVector2D Position)
-{
-	UCanvasPanelSlot* Slot = AddSlot(Child);
-	Slot->LayoutData.Offsets = FMargin(Position.X, Position.Y, 100, 25);
-	
 	// Add the child to the live canvas if it already exists
-	if (MyCanvas.IsValid())
+	if ( MyCanvas.IsValid() )
 	{
-		Slot->BuildSlot(MyCanvas.ToSharedRef());
+		Cast<UCanvasPanelSlot>(Slot)->BuildSlot(MyCanvas.ToSharedRef());
 	}
-	
-	return true;
 }
 
-bool UCanvasPanel::RemoveChild(UWidget* Child)
+void UCanvasPanel::OnSlotRemoved(UPanelSlot* Slot)
 {
-	int32 SlotIndex = GetChildIndex(Child);
-	if ( SlotIndex != -1 )
+	// Remove the widget from the live slot if it exists.
+	if ( MyCanvas.IsValid() )
 	{
-		Slots.RemoveAt(SlotIndex);
-		
-		// Remove the widget from the live slot if it exists.
-		if (MyCanvas.IsValid())
-		{
-			MyCanvas->RemoveSlot(Child->GetWidget());
-		}
-		
-		return true;
+		MyCanvas->RemoveSlot(Slot->Content->GetWidget());
 	}
-
-	return false;
-}
-
-void UCanvasPanel::ReplaceChildAt(int32 Index, UWidget* Content)
-{
-	UCanvasPanelSlot* Slot = Slots[Index];
-	Slot->Content = Content;
-
-	Content->Slot = Slot;
-	
-	//TODO UMG How do we handle this swap for live controls?  It doesn't matter for canvas, but matters for flow controls it will matter
-}
-
-void UCanvasPanel::InsertChildAt(int32 Index, UWidget* Content)
-{
-	UCanvasPanelSlot* Slot = ConstructObject<UCanvasPanelSlot>(UCanvasPanelSlot::StaticClass(), this);
-	Slot->SetFlags(RF_Transactional);
-	Slot->Content = Content;
-	Slot->Parent = this;
-
-	Content->Slot = Slot;
-
-	Slots.Insert(Slot, Index);
 }
 
 TSharedRef<SWidget> UCanvasPanel::RebuildWidget()
 {
 	MyCanvas = SNew(SFixedSizeCanvas, DesiredCanvasSize);
 
-	for ( auto Slot : Slots )
+	for ( UPanelSlot* Slot : Slots )
 	{
-		Slot->Parent = this;
-		Slot->BuildSlot(MyCanvas.ToSharedRef());
+		if ( UCanvasPanelSlot* TypedSlot = Cast<UCanvasPanelSlot>(Slot) )
+		{
+			TypedSlot->Parent = this;
+			TypedSlot->BuildSlot(MyCanvas.ToSharedRef());
+		}
 	}
 
 	return BuildDesignTimeWidget( MyCanvas.ToSharedRef() );
@@ -151,18 +92,9 @@ void UCanvasPanel::SyncronizeProperties()
 	MyCanvas->SetDesiredSize(DesiredCanvasSize);
 }
 
-UCanvasPanelSlot* UCanvasPanel::AddSlot(UWidget* Content)
+TSharedPtr<SConstraintCanvas> UCanvasPanel::GetCanvasWidget() const
 {
-	UCanvasPanelSlot* Slot = ConstructObject<UCanvasPanelSlot>(UCanvasPanelSlot::StaticClass(), this);
-	Slot->SetFlags(RF_Transactional);
-	Slot->Content = Content;
-	Slot->Parent = this;
-
-	Content->Slot = Slot;
-	
-	Slots.Add(Slot);
-
-	return Slot;
+	return MyCanvas;
 }
 
 bool UCanvasPanel::GetGeometryForSlot(UCanvasPanelSlot* Slot, FGeometry& ArrangedGeometry) const
@@ -193,15 +125,6 @@ bool UCanvasPanel::GetGeometryForSlot(UCanvasPanelSlot* Slot, FGeometry& Arrange
 
 #if WITH_EDITOR
 
-void UCanvasPanel::ConnectEditorData()
-{
-	for ( UCanvasPanelSlot* Slot : Slots )
-	{
-		Slot->Parent = this;
-		Slot->Content->Slot = Slot;
-	}
-}
-
 const FSlateBrush* UCanvasPanel::GetEditorIcon()
 {
 	return FUMGStyle::Get().GetBrush("Widget.Canvas");
diff --git a/Engine/Source/Runtime/UMG/Private/Components/ContentWidget.cpp b/Engine/Source/Runtime/UMG/Private/Components/ContentWidget.cpp
--- a/Engine/Source/Runtime/UMG/Private/Components/ContentWidget.cpp
+++ b/Engine/Source/Runtime/UMG/Private/Components/ContentWidget.cpp
@@ -15,77 +15,17 @@
 UContentWidget::UContentWidget(const FPostConstructInitializeProperties& PCIP)
 	: Super(PCIP)
 {
-	ChildSlot = PCIP.CreateDefaultSubobject<UPanelSlot>(this, TEXT("ChildSlot"));
-	ChildSlot.Get()->SetFlags(RF_Transactional);
-	ChildSlot.Get()->Parent = this;
+	bCanHaveMultipleChildren = false;
 }
 
 UPanelSlot* UContentWidget::GetContentSlot() const
 {
-	//if ( ContentSlot == NULL )
-	//{
-	//	UContentWidget* MutableThis = const_cast<UContentWidget*>( this );
-	//	MutableThis->ContentSlot = ConstructObject<UPanelSlot>(UPanelSlot::StaticClass(), MutableThis);
-
-	//	MutableThis->ContentSlot->SetFlags(RF_Transactional);
-	//	//MutableThis->ContentSlot->SetFlags(RF_DefaultSubObject);
-	//	MutableThis->ContentSlot->Parent = MutableThis;
-	//}
-
-	return ChildSlot.Get();
-}
-
-UWidget* UContentWidget::GetContent()
-{
-	return GetContentSlot()->Content;
-}
-
-void UContentWidget::SetContent(UWidget* InContent)
-{
-	GetContentSlot()->Content = InContent;
-
-	if ( InContent )
-	{
-		InContent->Slot = GetContentSlot();
-	}
-}
-
-int32 UContentWidget::GetChildrenCount() const
-{
-	return GetContentSlot()->Content != NULL ? 1 : 0;
-}
-
-UWidget* UContentWidget::GetChildAt(int32 Index) const
-{
-	return GetContentSlot()->Content;
-}
-
-bool UContentWidget::AddChild(UWidget* InContent, FVector2D Position)
-{
-	if ( GetContentSlot()->Content == NULL )
-	{
-		SetContent(InContent);
-		return true;
-	}
-
-	return false;
-}
-
-bool UContentWidget::RemoveChild(UWidget* Child)
-{
-	if ( GetContentSlot()->Content == Child )
-	{
-		SetContent(NULL);
-		return true;
-	}
-
-	return false;
+	return Slots.Num() > 0 ? Slots[0] : NULL;
 }
 
-void UContentWidget::ReplaceChildAt(int32 Index, UWidget* Child)
+UClass* UContentWidget::GetSlotClass() const
 {
-	check(Index == 0);
-	SetContent(Child);
+	return UPanelSlot::StaticClass();
 }
 
 /////////////////////////////////////////////////////
diff --git a/Engine/Source/Runtime/UMG/Private/Components/HorizontalBox.cpp b/Engine/Source/Runtime/UMG/Private/Components/HorizontalBox.cpp
--- a/Engine/Source/Runtime/UMG/Private/Components/HorizontalBox.cpp
+++ b/Engine/Source/Runtime/UMG/Private/Components/HorizontalBox.cpp
@@ -11,120 +11,52 @@ UHorizontalBox::UHorizontalBox(const FPostConstructInitializeProperties& PCIP)
 	bIsVariable = false;
 }
 
-int32 UHorizontalBox::GetChildrenCount() const
+UClass* UHorizontalBox::GetSlotClass() const
 {
-	return Slots.Num();
+	return UHorizontalBoxSlot::StaticClass();
 }
 
-UWidget* UHorizontalBox::GetChildAt(int32 Index) const
+void UHorizontalBox::OnSlotAdded(UPanelSlot* Slot)
 {
-	return Slots[Index]->Content;
-}
-
-int32 UHorizontalBox::GetChildIndex(UWidget* Content) const
-{
-	for ( int32 SlotIndex = 0; SlotIndex < Slots.Num(); ++SlotIndex )
+	// Add the child to the live canvas if it already exists
+	if ( MyHorizontalBox.IsValid() )
 	{
-		UHorizontalBoxSlot* Slot = Slots[SlotIndex];
-
-		if ( Slot->Content == Content )
-		{
-			return SlotIndex;
-		}
+		Cast<UHorizontalBoxSlot>(Slot)->BuildSlot(MyHorizontalBox.ToSharedRef());
 	}
-
-	return -1;
-}
-
-bool UHorizontalBox::AddChild(UWidget* Child, FVector2D Position)
-{
-	UHorizontalBoxSlot* Slot = AddSlot(Child);
-	return true;
 }
 
-bool UHorizontalBox::RemoveChild(UWidget* Child)
+void UHorizontalBox::OnSlotRemoved(UPanelSlot* Slot)
 {
-	int32 SlotIndex = GetChildIndex(Child);
-	if ( SlotIndex != -1 )
+	// Remove the widget from the live slot if it exists.
+	if ( MyHorizontalBox.IsValid() )
 	{
-		Slots.RemoveAt(SlotIndex);
-		
-		// Remove the widget from the live slot if it exists.
-		if (MyHorizontalBox.IsValid())
-		{
-			MyHorizontalBox->RemoveSlot(Child->GetWidget());
-		}
-		
-		return true;
+		MyHorizontalBox->RemoveSlot(Slot->Content->GetWidget());
 	}
-
-	return false;
 }
 
-void UHorizontalBox::ReplaceChildAt(int32 Index, UWidget* Content)
+UHorizontalBoxSlot* UHorizontalBox::AddChild(UWidget* Content)
 {
-	UHorizontalBoxSlot* Slot = Slots[Index];
-	Slot->Content = Content;
-
-	Content->Slot = Slot;
-}
-
-void UHorizontalBox::InsertChildAt(int32 Index, UWidget* Content)
-{
-	UHorizontalBoxSlot* Slot = ConstructObject<UHorizontalBoxSlot>(UHorizontalBoxSlot::StaticClass(), this);
-	Slot->SetFlags(RF_Transactional);
-	Slot->Content = Content;
-	Slot->Parent = this;
-
-	Content->Slot = Slot;
-
-	Slots.Insert(Slot, Index);
+	return Cast<UHorizontalBoxSlot>( Super::AddChild(Content, FVector2D(0, 0)) );
 }
 
 TSharedRef<SWidget> UHorizontalBox::RebuildWidget()
 {
 	MyHorizontalBox = SNew(SHorizontalBox);
 
-	for ( auto Slot : Slots )
+	for ( UPanelSlot* Slot : Slots )
 	{
-		Slot->Parent = this;
-		Slot->BuildSlot( MyHorizontalBox.ToSharedRef() );
+		if ( UHorizontalBoxSlot* TypedSlot = Cast<UHorizontalBoxSlot>(Slot) )
+		{
+			TypedSlot->Parent = this;
+			TypedSlot->BuildSlot(MyHorizontalBox.ToSharedRef());
+		}
 	}
 
 	return BuildDesignTimeWidget( MyHorizontalBox.ToSharedRef() );
 }
 
-UHorizontalBoxSlot* UHorizontalBox::AddSlot(UWidget* Content)
-{
-	UHorizontalBoxSlot* Slot = ConstructObject<UHorizontalBoxSlot>(UHorizontalBoxSlot::StaticClass(), this);
-	Slot->SetFlags(RF_Transactional);
-	Slot->Content = Content;
-	Slot->Parent = this;
-
-	Content->Slot = Slot;
-	
-	Slots.Add(Slot);
-
-	// Add the child to the live canvas if it already exists
-	if ( MyHorizontalBox.IsValid() )
-	{
-		Slot->BuildSlot(MyHorizontalBox.ToSharedRef());
-	}
-
-	return Slot;
-}
-
 #if WITH_EDITOR
 
-void UHorizontalBox::ConnectEditorData()
-{
-	for ( UHorizontalBoxSlot* Slot : Slots )
-	{
-		Slot->Content->Slot = Slot;
-		Slot->Parent = this;
-	}
-}
-
 const FSlateBrush* UHorizontalBox::GetEditorIcon()
 {
 	return FUMGStyle::Get().GetBrush("Widget.HorizontalBox");
diff --git a/Engine/Source/Runtime/UMG/Private/Components/Overlay.cpp b/Engine/Source/Runtime/UMG/Private/Components/Overlay.cpp
--- a/Engine/Source/Runtime/UMG/Private/Components/Overlay.cpp
+++ b/Engine/Source/Runtime/UMG/Private/Components/Overlay.cpp
@@ -11,99 +11,41 @@ UOverlay::UOverlay(const FPostConstructInitializeProperties& PCIP)
 	bIsVariable = false;
 }
 
-int32 UOverlay::GetChildrenCount() const
+UClass* UOverlay::GetSlotClass() const
 {
-	return Slots.Num();
+	return UOverlaySlot::StaticClass();
 }
 
-UWidget* UOverlay::GetChildAt(int32 Index) const
+void UOverlay::OnSlotAdded(UPanelSlot* Slot)
 {
-	return Slots[Index]->Content;
-}
-
-bool UOverlay::AddChild(UWidget* Child, FVector2D Position)
-{
-	UOverlaySlot* Slot = AddSlot(Child);
-	
 	// Add the child to the live canvas if it already exists
-	if (MyOverlay.IsValid())
+	if ( MyOverlay.IsValid() )
 	{
-		Slot->BuildSlot(MyOverlay.ToSharedRef());
+		Cast<UOverlaySlot>(Slot)->BuildSlot(MyOverlay.ToSharedRef());
 	}
-	
-	return true;
 }
 
-bool UOverlay::RemoveChild(UWidget* Child)
+void UOverlay::OnSlotRemoved(UPanelSlot* Slot)
 {
-	for ( int32 SlotIndex = 0; SlotIndex < Slots.Num(); ++SlotIndex )
+	// Remove the widget from the live slot if it exists.
+	if ( MyOverlay.IsValid() )
 	{
-		UOverlaySlot* Slot = Slots[SlotIndex];
-
-		if ( Slot->Content == Child )
-		{
-			Slots.RemoveAt(SlotIndex);
-			
-			// Remove the widget from the live slot if it exists.
-			if (MyOverlay.IsValid())
-			{
-				MyOverlay->RemoveSlot(Child->GetWidget());
-			}
-			
-			return true;
-		}
+		MyOverlay->RemoveSlot(Slot->Content->GetWidget());
 	}
-
-	return false;
-}
-
-void UOverlay::ReplaceChildAt(int32 Index, UWidget* Content)
-{
-	UOverlaySlot* Slot = Slots[Index];
-	Slot->Content = Content;
-
-	Content->Slot = Slot;
 }
 
 TSharedRef<SWidget> UOverlay::RebuildWidget()
 {
 	MyOverlay = SNew(SOverlay);
 
-	for ( auto Slot : Slots )
+	for ( UPanelSlot* Slot : Slots )
 	{
-		Slot->Parent = this;
-		Slot->BuildSlot(MyOverlay.ToSharedRef());
-	}
-
-	return BuildDesignTimeWidget( MyOverlay.ToSharedRef() );
-}
-
-UOverlaySlot* UOverlay::AddSlot(UWidget* Content)
-{
-	UOverlaySlot* Slot = ConstructObject<UOverlaySlot>(UOverlaySlot::StaticClass(), this);
-	Slot->SetFlags(RF_Transactional);
-	Slot->Content = Content;
-	Slot->Parent = this;
-
-	Content->Slot = Slot;
-	
-	Slots.Add(Slot);
-
-	return Slot;
-}
-
-#if WITH_EDITOR
-
-void UOverlay::ConnectEditorData()
-{
-	for ( UOverlaySlot* Slot : Slots )
-	{
-		if ( Slot->Content )
+		if ( UOverlaySlot* TypedSlot = Cast<UOverlaySlot>(Slot) )
 		{
-			Slot->Content->Slot = Slot;
+			TypedSlot->Parent = this;
+			TypedSlot->BuildSlot(MyOverlay.ToSharedRef());
 		}
-		//TODO UMG Should we auto delete empty slots?
 	}
-}
 
-#endif
+	return BuildDesignTimeWidget( MyOverlay.ToSharedRef() );
+}
diff --git a/Engine/Source/Runtime/UMG/Private/Components/PanelWidget.cpp b/Engine/Source/Runtime/UMG/Private/Components/PanelWidget.cpp
--- a/Engine/Source/Runtime/UMG/Private/Components/PanelWidget.cpp
+++ b/Engine/Source/Runtime/UMG/Private/Components/PanelWidget.cpp
@@ -14,6 +14,7 @@
 
 UPanelWidget::UPanelWidget(const FPostConstructInitializeProperties& PCIP)
 	: Super(PCIP)
+	, bCanHaveMultipleChildren(true)
 {
 }
 
diff --git a/Engine/Source/Runtime/UMG/Private/Components/ScrollBox.cpp b/Engine/Source/Runtime/UMG/Private/Components/ScrollBox.cpp
--- a/Engine/Source/Runtime/UMG/Private/Components/ScrollBox.cpp
+++ b/Engine/Source/Runtime/UMG/Private/Components/ScrollBox.cpp
@@ -13,110 +13,45 @@ UScrollBox::UScrollBox(const FPostConstructInitializeProperties& PCIP)
 	bIsVariable = false;
 }
 
-int32 UScrollBox::GetChildrenCount() const
+UClass* UScrollBox::GetSlotClass() const
 {
-	return Slots.Num();
+	return UScrollBoxSlot::StaticClass();
 }
 
-UWidget* UScrollBox::GetChildAt(int32 Index) const
+void UScrollBox::OnSlotAdded(UPanelSlot* Slot)
 {
-	return Slots[Index]->Content;
-}
-
-int32 UScrollBox::GetChildIndex(UWidget* Content) const
-{
-	for ( int32 SlotIndex = 0; SlotIndex < Slots.Num(); ++SlotIndex )
-	{
-		UScrollBoxSlot* Slot = Slots[SlotIndex];
-
-		if ( Slot->Content == Content )
-		{
-			return SlotIndex;
-		}
-	}
-
-	return -1;
-}
-
-bool UScrollBox::AddChild(UWidget* Child, FVector2D Position)
-{
-	UScrollBoxSlot* Slot = AddSlot(Child);
-	
 	// Add the child to the live canvas if it already exists
-	if (MyScrollBox.IsValid())
+	if ( MyScrollBox.IsValid() )
 	{
-		Slot->BuildSlot(MyScrollBox.ToSharedRef());
+		Cast<UScrollBoxSlot>(Slot)->BuildSlot(MyScrollBox.ToSharedRef());
 	}
-	
-	return true;
 }
 
-bool UScrollBox::RemoveChild(UWidget* Child)
+void UScrollBox::OnSlotRemoved(UPanelSlot* Slot)
 {
-	int32 SlotIndex = GetChildIndex(Child);
-	if ( SlotIndex != -1 )
+	// Remove the widget from the live slot if it exists.
+	if ( MyScrollBox.IsValid() )
 	{
-		Slots.RemoveAt(SlotIndex);
-		
-		// Remove the widget from the live slot if it exists.
-		if (MyScrollBox.IsValid())
-		{
-			MyScrollBox->RemoveSlot(Child->GetWidget());
-		}
-		
-		return true;
+		MyScrollBox->RemoveSlot(Slot->Content->GetWidget());
 	}
-
-	return false;
-}
-
-void UScrollBox::ReplaceChildAt(int32 Index, UWidget* Content)
-{
-	UScrollBoxSlot* Slot = Slots[Index];
-	Slot->Content = Content;
-
-	Content->Slot = Slot;
-}
-
-void UScrollBox::InsertChildAt(int32 Index, UWidget* Content)
-{
-	UScrollBoxSlot* Slot = ConstructObject<UScrollBoxSlot>(UScrollBoxSlot::StaticClass(), this);
-	Slot->SetFlags(RF_Transactional);
-	Slot->Content = Content;
-	Slot->Parent = this;
-
-	Content->Slot = Slot;
-
-	Slots.Insert(Slot, Index);
 }
 
 TSharedRef<SWidget> UScrollBox::RebuildWidget()
 {
 	MyScrollBox = SNew(SScrollBox);
 
-	for ( auto Slot : Slots )
+	for ( UPanelSlot* Slot : Slots )
 	{
-		Slot->Parent = this;
-		Slot->BuildSlot(MyScrollBox.ToSharedRef());
+		if ( UScrollBoxSlot* TypedSlot = Cast<UScrollBoxSlot>(Slot) )
+		{
+			TypedSlot->Parent = this;
+			TypedSlot->BuildSlot(MyScrollBox.ToSharedRef());
+		}
 	}
 	
 	return BuildDesignTimeWidget( MyScrollBox.ToSharedRef() );
 }
 
-UScrollBoxSlot* UScrollBox::AddSlot(UWidget* Content)
-{
-	UScrollBoxSlot* Slot = ConstructObject<UScrollBoxSlot>(UScrollBoxSlot::StaticClass(), this);
-	Slot->SetFlags(RF_Transactional);
-	Slot->Content = Content;
-	Slot->Parent = this;
-
-	Content->Slot = Slot;
-	
-	Slots.Add(Slot);
-
-	return Slot;
-}
-
 #if WITH_EDITOR
 
 const FSlateBrush* UScrollBox::GetEditorIcon()
@@ -124,21 +59,6 @@ const FSlateBrush* UScrollBox::GetEditorIcon()
 	return FUMGStyle::Get().GetBrush("Widget.ScrollBox");
 }
 
-void UScrollBox::ConnectEditorData()
-{
-	for ( UScrollBoxSlot* Slot : Slots )
-	{
-		// TODO UMG Find a better way to assign the parent pointers for slots.
-		Slot->Parent = this;
-
-		if ( Slot->Content )
-		{
-			Slot->Content->Slot = Slot;
-		}
-		//TODO UMG Should we auto delete empty slots?
-	}
-}
-
 #endif
 
 /////////////////////////////////////////////////////
diff --git a/Engine/Source/Runtime/UMG/Private/Components/UniformGridPanel.cpp b/Engine/Source/Runtime/UMG/Private/Components/UniformGridPanel.cpp
--- a/Engine/Source/Runtime/UMG/Private/Components/UniformGridPanel.cpp
+++ b/Engine/Source/Runtime/UMG/Private/Components/UniformGridPanel.cpp
@@ -11,58 +11,27 @@ UUniformGridPanel::UUniformGridPanel(const FPostConstructInitializeProperties& P
 	bIsVariable = false;
 }
 
-int32 UUniformGridPanel::GetChildrenCount() const
+UClass* UUniformGridPanel::GetSlotClass() const
 {
-	return Slots.Num();
+	return UUniformGridSlot::StaticClass();
 }
 
-UWidget* UUniformGridPanel::GetChildAt(int32 Index) const
+void UUniformGridPanel::OnSlotAdded(UPanelSlot* Slot)
 {
-	return Slots[Index]->Content;
-}
-
-bool UUniformGridPanel::AddChild(UWidget* Child, FVector2D Position)
-{
-	UUniformGridSlot* Slot = AddSlot(Child);
-	
-	// Add the child to the live panel if it already exists
-	if (MyUniformGridPanel.IsValid())
+	// Add the child to the live canvas if it already exists
+	if ( MyUniformGridPanel.IsValid() )
 	{
-		Slot->BuildSlot(MyUniformGridPanel.ToSharedRef());
+		Cast<UUniformGridSlot>(Slot)->BuildSlot(MyUniformGridPanel.ToSharedRef());
 	}
-	
-	return true;
 }
 
-bool UUniformGridPanel::RemoveChild(UWidget* Child)
+void UUniformGridPanel::OnSlotRemoved(UPanelSlot* Slot)
 {
-	for ( int32 SlotIndex = 0; SlotIndex < Slots.Num(); ++SlotIndex )
+	// Remove the widget from the live slot if it exists.
+	if ( MyUniformGridPanel.IsValid() )
 	{
-		UUniformGridSlot* Slot = Slots[SlotIndex];
-
-		if ( Slot->Content == Child )
-		{
-			Slots.RemoveAt(SlotIndex);
-			
-			// Remove the widget from the live slot if it exists.
-			if (MyUniformGridPanel.IsValid())
-			{
-				MyUniformGridPanel->RemoveSlot(Child->GetWidget());
-			}
-			
-			return true;
-		}
+		MyUniformGridPanel->RemoveSlot(Slot->Content->GetWidget());
 	}
-
-	return false;
-}
-
-void UUniformGridPanel::ReplaceChildAt(int32 Index, UWidget* Content)
-{
-	UUniformGridSlot* Slot = Slots[Index];
-	Slot->Content = Content;
-
-	Content->Slot = Slot;
 }
 
 TSharedRef<SWidget> UUniformGridPanel::RebuildWidget()
@@ -73,38 +42,14 @@ TSharedRef<SWidget> UUniformGridPanel::RebuildWidget()
 		.MinDesiredSlotWidth(MinDesiredSlotWidth)
 		.MinDesiredSlotHeight(MinDesiredSlotHeight);
 
-	for ( auto Slot : Slots )
+	for ( UPanelSlot* Slot : Slots )
 	{
-		Slot->Parent = this;
-		Slot->BuildSlot(MyUniformGridPanel.ToSharedRef());
+		if ( UUniformGridSlot* TypedSlot = Cast<UUniformGridSlot>(Slot) )
+		{
+			TypedSlot->Parent = this;
+			TypedSlot->BuildSlot(MyUniformGridPanel.ToSharedRef());
+		}
 	}
 
 	return BuildDesignTimeWidget( MyUniformGridPanel.ToSharedRef() );
 }
-
-UUniformGridSlot* UUniformGridPanel::AddSlot(UWidget* Content)
-{
-	UUniformGridSlot* Slot = ConstructObject<UUniformGridSlot>(UUniformGridSlot::StaticClass(), this);
-	Slot->SetFlags(RF_Transactional);
-	Slot->Content = Content;
-	Slot->Parent = this;
-
-	Content->Slot = Slot;
-	
-	Slots.Add(Slot);
-
-	return Slot;
-}
-
-#if WITH_EDITOR
-
-void UUniformGridPanel::ConnectEditorData()
-{
-	for ( UUniformGridSlot* Slot : Slots )
-	{
-		Slot->Content->Slot = Slot;
-		Slot->Parent = this;
-	}
-}
-
-#endif
diff --git a/Engine/Source/Runtime/UMG/Private/Components/VerticalBox.cpp b/Engine/Source/Runtime/UMG/Private/Components/VerticalBox.cpp
--- a/Engine/Source/Runtime/UMG/Private/Components/VerticalBox.cpp
+++ b/Engine/Source/Runtime/UMG/Private/Components/VerticalBox.cpp
@@ -11,126 +11,52 @@ UVerticalBox::UVerticalBox(const FPostConstructInitializeProperties& PCIP)
 	bIsVariable = false;
 }
 
-int32 UVerticalBox::GetChildrenCount() const
+UClass* UVerticalBox::GetSlotClass() const
 {
-	return Slots.Num();
+	return UVerticalBoxSlot::StaticClass();
 }
 
-UWidget* UVerticalBox::GetChildAt(int32 Index) const
+void UVerticalBox::OnSlotAdded(UPanelSlot* Slot)
 {
-	return Slots[Index]->Content;
-}
-
-int32 UVerticalBox::GetChildIndex(UWidget* Content) const
-{
-	for ( int32 SlotIndex = 0; SlotIndex < Slots.Num(); ++SlotIndex )
+	// Add the child to the live canvas if it already exists
+	if ( MyVerticalBox.IsValid() )
 	{
-		UVerticalBoxSlot* Slot = Slots[SlotIndex];
-
-		if ( Slot->Content == Content )
-		{
-			return SlotIndex;
-		}
+		Cast<UVerticalBoxSlot>(Slot)->BuildSlot(MyVerticalBox.ToSharedRef());
 	}
-
-	return -1;
 }
 
-bool UVerticalBox::AddChild(UWidget* Child, FVector2D Position)
+void UVerticalBox::OnSlotRemoved(UPanelSlot* Slot)
 {
-	UVerticalBoxSlot* Slot = AddSlot(Child);
-	return true;
-}
-
-bool UVerticalBox::RemoveChild(UWidget* Child)
-{
-	int32 SlotIndex = GetChildIndex(Child);
-	if ( SlotIndex != -1 )
+	// Remove the widget from the live slot if it exists.
+	if ( MyVerticalBox.IsValid() )
 	{
-		Slots.RemoveAt(SlotIndex);
-		
-		// Remove the widget from the live slot if it exists.
-		if (MyVerticalBox.IsValid())
-		{
-			MyVerticalBox->RemoveSlot(Child->GetWidget());
-		}
-		
-		return true;
+		MyVerticalBox->RemoveSlot(Slot->Content->GetWidget());
 	}
-
-	return false;
 }
 
-void UVerticalBox::ReplaceChildAt(int32 Index, UWidget* Content)
+UVerticalBoxSlot* UVerticalBox::AddChild(UWidget* Content)
 {
-	UVerticalBoxSlot* Slot = Slots[Index];
-	Slot->Content = Content;
-
-	Content->Slot = Slot;
-}
-
-void UVerticalBox::InsertChildAt(int32 Index, UWidget* Content)
-{
-	UVerticalBoxSlot* Slot = ConstructObject<UVerticalBoxSlot>(UVerticalBoxSlot::StaticClass(), this);
-	Slot->SetFlags(RF_Transactional);
-	Slot->Content = Content;
-	Slot->Parent = this;
-
-	Content->Slot = Slot;
-
-	Slots.Insert(Slot, Index);
+	return Cast<UVerticalBoxSlot>( Super::AddChild(Content, FVector2D(0, 0)) );
 }
 
 TSharedRef<SWidget> UVerticalBox::RebuildWidget()
 {
 	MyVerticalBox = SNew(SVerticalBox);
 
-	for ( auto Slot : Slots )
+	for ( UPanelSlot* Slot : Slots )
 	{
-		Slot->Parent = this;
-		Slot->BuildSlot(MyVerticalBox.ToSharedRef());
+		if ( UVerticalBoxSlot* TypedSlot = Cast<UVerticalBoxSlot>(Slot) )
+		{
+			TypedSlot->Parent = this;
+			TypedSlot->BuildSlot(MyVerticalBox.ToSharedRef());
+		}
 	}
 
 	return BuildDesignTimeWidget( MyVerticalBox.ToSharedRef() );
 }
 
-UVerticalBoxSlot* UVerticalBox::AddSlot(UWidget* Content)
-{
-	UVerticalBoxSlot* Slot = ConstructObject<UVerticalBoxSlot>(UVerticalBoxSlot::StaticClass(), this);
-	Slot->SetFlags(RF_Transactional);
-	Slot->Content = Content;
-	Slot->Parent = this;
-
-	Content->Slot = Slot;
-	
-	Slots.Add(Slot);
-
-	// Add the child to the live canvas if it already exists
-	if ( MyVerticalBox.IsValid() )
-	{
-		Slot->BuildSlot(MyVerticalBox.ToSharedRef());
-	}
-
-	return Slot;
-}
-
 #if WITH_EDITOR
 
-void UVerticalBox::ConnectEditorData()
-{
-	for ( UVerticalBoxSlot* Slot : Slots )
-	{
-		// TODO UMG Find a better way to assign the parent pointers for slots.
-		Slot->Parent = this;
-
-		if ( Slot->Content )
-		{
-			Slot->Content->Slot = Slot;
-		}
-		//TODO UMG Should we auto delete empty slots?
-	}
-}
-
 const FSlateBrush* UVerticalBox::GetEditorIcon()
 {
 	return FUMGStyle::Get().GetBrush("Widget.VerticalBox");
diff --git a/Engine/Source/Runtime/UMG/Private/Components/WidgetSwitcher.cpp b/Engine/Source/Runtime/UMG/Private/Components/WidgetSwitcher.cpp
--- a/Engine/Source/Runtime/UMG/Private/Components/WidgetSwitcher.cpp
+++ b/Engine/Source/Runtime/UMG/Private/Components/WidgetSwitcher.cpp
@@ -42,104 +42,40 @@ void UWidgetSwitcher::SetActiveWidgetIndex(int32 Index)
 	}
 }
 
-int32 UWidgetSwitcher::GetChildrenCount() const
+UClass* UWidgetSwitcher::GetSlotClass() const
 {
-	return Slots.Num();
-}
-
-UWidget* UWidgetSwitcher::GetChildAt(int32 Index) const
-{
-	return Slots[Index]->Content;
-}
-
-int32 UWidgetSwitcher::GetChildIndex(UWidget* Content) const
-{
-	for ( int32 SlotIndex = 0; SlotIndex < Slots.Num(); ++SlotIndex )
-	{
-		UWidgetSwitcherSlot* Slot = Slots[SlotIndex];
-
-		if ( Slot->Content == Content )
-		{
-			return SlotIndex;
-		}
-	}
-
-	return -1;
-}
-
-bool UWidgetSwitcher::AddChild(UWidget* Child, FVector2D Position)
-{
-	UWidgetSwitcherSlot* Slot = AddSlot(Child);
-	return true;
+	return UWidgetSwitcherSlot::StaticClass();
 }
 
-bool UWidgetSwitcher::RemoveChild(UWidget* Child)
+void UWidgetSwitcher::OnSlotAdded(UPanelSlot* Slot)
 {
-	int32 SlotIndex = GetChildIndex(Child);
-	if ( SlotIndex != -1 )
+	// Add the child to the live canvas if it already exists
+	if ( MyWidgetSwitcher.IsValid() )
 	{
-		Slots.RemoveAt(SlotIndex);
-		
-		// Remove the widget from the live slot if it exists.
-		if (MyWidgetSwitcher.IsValid())
-		{
-			MyWidgetSwitcher->RemoveSlot(Child->GetWidget());
-		}
-		
-		return true;
+		Cast<UWidgetSwitcherSlot>(Slot)->BuildSlot(MyWidgetSwitcher.ToSharedRef());
 	}
-
-	return false;
 }
 
-void UWidgetSwitcher::ReplaceChildAt(int32 Index, UWidget* Content)
+void UWidgetSwitcher::OnSlotRemoved(UPanelSlot* Slot)
 {
-	UWidgetSwitcherSlot* Slot = Slots[Index];
-	Slot->Content = Content;
-
-	Content->Slot = Slot;
-}
-
-void UWidgetSwitcher::InsertChildAt(int32 Index, UWidget* Content)
-{
-	UWidgetSwitcherSlot* Slot = ConstructObject<UWidgetSwitcherSlot>(UWidgetSwitcherSlot::StaticClass(), this);
-	Slot->SetFlags(RF_Transactional);
-	Slot->Content = Content;
-	Slot->Parent = this;
-
-	Content->Slot = Slot;
-
-	Slots.Insert(Slot, Index);
-}
-
-UWidgetSwitcherSlot* UWidgetSwitcher::AddSlot(UWidget* Content)
-{
-	UWidgetSwitcherSlot* Slot = ConstructObject<UWidgetSwitcherSlot>(UWidgetSwitcherSlot::StaticClass(), this);
-	Slot->SetFlags(RF_Transactional);
-	Slot->Content = Content;
-	Slot->Parent = this;
-
-	Content->Slot = Slot;
-
-	Slots.Add(Slot);
-
-	// Add the child to the live canvas if it already exists
+	// Remove the widget from the live slot if it exists.
 	if ( MyWidgetSwitcher.IsValid() )
 	{
-		Slot->BuildSlot(MyWidgetSwitcher.ToSharedRef());
+		MyWidgetSwitcher->RemoveSlot(Slot->Content->GetWidget());
 	}
-
-	return Slot;
 }
 
 TSharedRef<SWidget> UWidgetSwitcher::RebuildWidget()
 {
 	MyWidgetSwitcher = SNew(SWidgetSwitcher);
 
-	for ( auto Slot : Slots )
+	for ( UPanelSlot* Slot : Slots )
 	{
-		Slot->Parent = this;
-		Slot->BuildSlot(MyWidgetSwitcher.ToSharedRef());
+		if ( UWidgetSwitcherSlot* TypedSlot = Cast<UWidgetSwitcherSlot>(Slot) )
+		{
+			TypedSlot->Parent = this;
+			TypedSlot->BuildSlot(MyWidgetSwitcher.ToSharedRef());
+		}
 	}
 
 	return BuildDesignTimeWidget( MyWidgetSwitcher.ToSharedRef() );
@@ -151,22 +87,3 @@ void UWidgetSwitcher::SyncronizeProperties()
 
 	SetActiveWidgetIndex(ActiveWidgetIndex);
 }
-
-#if WITH_EDITOR
-
-void UWidgetSwitcher::ConnectEditorData()
-{
-	for ( UWidgetSwitcherSlot* Slot : Slots )
-	{
-		// TODO UMG Find a better way to assign the parent pointers for slots.
-		Slot->Parent = this;
-
-		if ( Slot->Content )
-		{
-			Slot->Content->Slot = Slot;
-		}
-		//TODO UMG Should we auto delete empty slots?
-	}
-}
-
-#endif
diff --git a/Engine/Source/Runtime/UMG/Public/Components/Border.h b/Engine/Source/Runtime/UMG/Public/Components/Border.h
--- a/Engine/Source/Runtime/UMG/Public/Components/Border.h
+++ b/Engine/Source/Runtime/UMG/Public/Components/Border.h
@@ -87,10 +87,6 @@ class UMG_API UBorder : public UContentWidget
 	UFUNCTION(BlueprintCallable, Category="Appearance")
 	void SetContentPadding(FMargin InContentPadding);
 
-	// UContentWidget interface
-	virtual void SetContent(UWidget* Content) override;
-	// End UContentWidget interface
-
 	// UWidget interface
 	void SyncronizeProperties() override;
 	// End of UWidget interface
@@ -99,6 +95,13 @@ class UMG_API UBorder : public UContentWidget
 	virtual const FSlateBrush* GetEditorIcon() override;
 #endif
 
+protected:
+
+	// UPanelWidget
+	virtual void OnSlotAdded(UPanelSlot* Slot) override;
+	virtual void OnSlotRemoved(UPanelSlot* Slot) override;
+	// End UPanelWidget
+
 protected:
 	TSharedPtr<SBorder> MyBorder;
 
diff --git a/Engine/Source/Runtime/UMG/Public/Components/Button.h b/Engine/Source/Runtime/UMG/Public/Components/Button.h
--- a/Engine/Source/Runtime/UMG/Public/Components/Button.h
+++ b/Engine/Source/Runtime/UMG/Public/Components/Button.h
@@ -79,10 +79,6 @@ public:
 	UFUNCTION(BlueprintCallable, Category="Widget")
 	bool IsPressed() const;
 
-	// UContentWidget interface
-	virtual void SetContent(UWidget* Content) override;
-	// End UContentWidget interface
-
 	// UWidget interface
 	virtual void SyncronizeProperties() override;
 	// End of UWidget interface
@@ -91,6 +87,13 @@ public:
 	virtual const FSlateBrush* GetEditorIcon() override;
 #endif
 
+protected:
+
+	// UPanelWidget
+	virtual void OnSlotAdded(UPanelSlot* Slot) override;
+	virtual void OnSlotRemoved(UPanelSlot* Slot) override;
+	// End UPanelWidget
+
 protected:
 	// UWidget interface
 	virtual TSharedRef<SWidget> RebuildWidget() override;
diff --git a/Engine/Source/Runtime/UMG/Public/Components/CanvasPanel.h b/Engine/Source/Runtime/UMG/Public/Components/CanvasPanel.h
--- a/Engine/Source/Runtime/UMG/Public/Components/CanvasPanel.h
+++ b/Engine/Source/Runtime/UMG/Public/Components/CanvasPanel.h
@@ -9,27 +9,10 @@ class UMG_API UCanvasPanel : public UPanelWidget
 {
 	GENERATED_UCLASS_BODY()
 
-	/** The items placed on the canvas */
-	UPROPERTY()
-	TArray<UCanvasPanelSlot*> Slots;
-
 	/** The desired size of the canvas */
 	UPROPERTY(EditAnywhere, Category=Appearance)
 	FVector2D DesiredCanvasSize;
 
-	/** Adds content to a new slot.  The slot will automatically be allocated and returned. */
-	UCanvasPanelSlot* AddSlot(UWidget* Content);
-
-	// UPanelWidget
-	virtual int32 GetChildrenCount() const override;
-	virtual UWidget* GetChildAt(int32 Index) const override;
-	virtual int32 GetChildIndex(UWidget* Content) const override;
-	virtual bool AddChild(UWidget* Child, FVector2D Position) override;
-	virtual bool RemoveChild(UWidget* Child) override;
-	virtual void ReplaceChildAt(int32 Index, UWidget* Child) override;
-	virtual void InsertChildAt(int32 Index, UWidget* Child) override;
-	// End UPanelWidget
-
 	/** Gets the underlying native canvas widget if it has been constructed */
 	TSharedPtr<class SConstraintCanvas> GetCanvasWidget() const;
 
@@ -43,10 +26,17 @@ class UMG_API UCanvasPanel : public UPanelWidget
 #if WITH_EDITOR
 	// UWidget interface
 	virtual const FSlateBrush* GetEditorIcon() override;
-	virtual void ConnectEditorData() override;
 	// End UWidget interface
 #endif
 
+protected:
+
+	// UPanelWidget
+	virtual UClass* GetSlotClass() const override;
+	virtual void OnSlotAdded(UPanelSlot* Slot) override;
+	virtual void OnSlotRemoved(UPanelSlot* Slot) override;
+	// End UPanelWidget
+
 protected:
 
 	TSharedPtr<class SFixedSizeCanvas> MyCanvas;
diff --git a/Engine/Source/Runtime/UMG/Public/Components/ContentWidget.h b/Engine/Source/Runtime/UMG/Public/Components/ContentWidget.h
--- a/Engine/Source/Runtime/UMG/Public/Components/ContentWidget.h
+++ b/Engine/Source/Runtime/UMG/Public/Components/ContentWidget.h
@@ -9,25 +9,18 @@ class UMG_API UContentWidget : public UPanelWidget
 {
 	GENERATED_UCLASS_BODY()
 
-	UFUNCTION()
-	virtual UWidget* GetContent();
-	
-	UFUNCTION()
-	virtual void SetContent(UWidget* InContent);
+	//UFUNCTION()
+	//virtual UWidget* GetContent();
+	//
+	//UFUNCTION()
+	//virtual void SetContent(UWidget* InContent);
+
+protected:
 
 	// UPanelWidget
-	virtual int32 GetChildrenCount() const override;
-	virtual UWidget* GetChildAt(int32 Index) const override;
-	virtual bool AddChild(UWidget* Child, FVector2D Position) override;
-	virtual bool RemoveChild(UWidget* Child) override;
-	virtual void ReplaceChildAt(int32 Index, UWidget* Child) override;
+	virtual UClass* GetSlotClass() const override;
 	// End UPanelWidget
 
 protected:
 	UPanelSlot* GetContentSlot() const;
-
-private:
-
-	UPROPERTY(Instanced)
-	TSubobjectPtr<UPanelSlot> ChildSlot;
 };
diff --git a/Engine/Source/Runtime/UMG/Public/Components/HorizontalBox.h b/Engine/Source/Runtime/UMG/Public/Components/HorizontalBox.h
--- a/Engine/Source/Runtime/UMG/Public/Components/HorizontalBox.h
+++ b/Engine/Source/Runtime/UMG/Public/Components/HorizontalBox.h
@@ -9,29 +9,24 @@ class UMG_API UHorizontalBox : public UPanelWidget
 {
 	GENERATED_UCLASS_BODY()
 
-	/** The items placed on the canvas */
-	UPROPERTY()
-	TArray<UHorizontalBoxSlot*> Slots;
-
-	UHorizontalBoxSlot* AddSlot(UWidget* Content);
-
-	// UPanelWidget
-	virtual int32 GetChildrenCount() const override;
-	virtual UWidget* GetChildAt(int32 Index) const override;
-	virtual int32 GetChildIndex(UWidget* Content) const override;
-	virtual bool AddChild(UWidget* Child, FVector2D Position) override;
-	virtual bool RemoveChild(UWidget* Child) override;
-	virtual void ReplaceChildAt(int32 Index, UWidget* Child) override;
-	virtual void InsertChildAt(int32 Index, UWidget* Child) override;
-	// End UPanelWidget
+	/**  */
+	UFUNCTION(BlueprintCallable, Category="Widget")
+	UHorizontalBoxSlot* AddChild(UWidget* Content);
 
 #if WITH_EDITOR
 	// UWidget interface
 	virtual const FSlateBrush* GetEditorIcon() override;
-	virtual void ConnectEditorData() override;
 	// End UWidget interface
 #endif
 
+protected:
+
+	// UPanelWidget
+	virtual UClass* GetSlotClass() const override;
+	virtual void OnSlotAdded(UPanelSlot* Slot) override;
+	virtual void OnSlotRemoved(UPanelSlot* Slot) override;
+	// End UPanelWidget
+
 protected:
 
 	TSharedPtr<class SHorizontalBox> MyHorizontalBox;
diff --git a/Engine/Source/Runtime/UMG/Public/Components/Overlay.h b/Engine/Source/Runtime/UMG/Public/Components/Overlay.h
--- a/Engine/Source/Runtime/UMG/Public/Components/Overlay.h
+++ b/Engine/Source/Runtime/UMG/Public/Components/Overlay.h
@@ -9,26 +9,14 @@ class UMG_API UOverlay : public UPanelWidget
 {
 	GENERATED_UCLASS_BODY()
 
-	/** The items placed on the canvas */
-	UPROPERTY()
-	TArray<UOverlaySlot*> Slots;
-
-	UOverlaySlot* AddSlot(UWidget* Content);
+protected:
 
 	// UPanelWidget
-	virtual int32 GetChildrenCount() const override;
-	virtual UWidget* GetChildAt(int32 Index) const override;
-	virtual bool AddChild(UWidget* Child, FVector2D Position) override;
-	virtual bool RemoveChild(UWidget* Child) override;
-	virtual void ReplaceChildAt(int32 Index, UWidget* Child) override;
+	virtual UClass* GetSlotClass() const override;
+	virtual void OnSlotAdded(UPanelSlot* Slot) override;
+	virtual void OnSlotRemoved(UPanelSlot* Slot) override;
 	// End UPanelWidget
 
-#if WITH_EDITOR
-	// UWidget interface
-	virtual void ConnectEditorData() override;
-	// End UWidget interface
-#endif
-
 protected:
 
 	TSharedPtr<class SOverlay> MyOverlay;
diff --git a/Engine/Source/Runtime/UMG/Public/Components/PanelWidget.h b/Engine/Source/Runtime/UMG/Public/Components/PanelWidget.h
--- a/Engine/Source/Runtime/UMG/Public/Components/PanelWidget.h
+++ b/Engine/Source/Runtime/UMG/Public/Components/PanelWidget.h
@@ -12,26 +12,133 @@ class UMG_API UPanelWidget : public UWidget
 
 public:
 
-	virtual int32 GetChildrenCount() const { return 0; }
-	virtual UWidget* GetChildAt(int32 Index) const { return NULL; }
+	/** The items placed in the panel */
+	UPROPERTY()
+	TArray<UPanelSlot*> Slots;
 
-	virtual int32 GetChildIndex(UWidget* Content) const { return -1; }
+public:
+
+	//TODO UMG Add ways to make adding slots callable by blueprints.
 
-	virtual void ReplaceChildAt(int32 Index, UWidget* Child)
+	int32 GetChildrenCount() const
 	{
+		return Slots.Num();
 	}
 
-	virtual void InsertChildAt(int32 Index, UWidget* Child)
+	UWidget* GetChildAt(int32 Index) const
 	{
+		return Slots[Index]->Content;
 	}
-	
-	virtual bool AddChild(UWidget* Child, FVector2D Position)
+
+	int32 GetChildIndex(UWidget* Content) const
 	{
-		return false;
+		const int32 ChildCount = GetChildrenCount();
+		for ( int32 ChildIndex = 0; ChildIndex < ChildCount; ChildIndex++ )
+		{
+			if ( GetChildAt(ChildIndex) == Content )
+			{
+				return ChildIndex;
+			}
+		}
+
+		return -1;
+	}
+
+	bool RemoveChildAt(int32 Index)
+	{
+		UPanelSlot* Slot = Slots[Index];
+		Slots.RemoveAt(Index);
+		
+		OnSlotRemoved(Slot);
+
+		return true;
 	}
 
-	virtual bool RemoveChild(UWidget* Child)
+	UPanelSlot* AddChild(UWidget* Content, FVector2D Position)
 	{
+		if ( !bCanHaveMultipleChildren && GetChildrenCount() > 0 )
+		{
+			return false;
+		}
+
+		UPanelSlot* Slot = ConstructObject<UPanelSlot>(GetSlotClass(), this);
+		Slot->SetFlags(RF_Transactional);
+		Slot->Content = Content;
+		Slot->Parent = this;
+
+		Content->Slot = Slot;
+
+		Slots.Add(Slot);
+
+		OnSlotAdded(Slot);
+
+		return Slot;
+	}
+
+	void ReplaceChildAt(int32 Index, UWidget* Content)
+	{
+		UPanelSlot* Slot = Slots[Index];
+		Slot->Content = Content;
+
+		Content->Slot = Slot;
+
+		Slot->Refresh();
+	}
+
+	void InsertChildAt(int32 Index, UWidget* Content)
+	{
+		UPanelSlot* Slot = ConstructObject<UPanelSlot>(GetSlotClass(), this);
+		Slot->SetFlags(RF_Transactional);
+		Slot->Content = Content;
+		Slot->Parent = this;
+
+		Content->Slot = Slot;
+
+		Slots.Insert(Slot, Index);
+
+		OnSlotAdded(Slot);
+	}
+
+	bool RemoveChild(UWidget* Content)
+	{
+		int32 ChildIndex = GetChildIndex(Content);
+		if ( ChildIndex != -1 )
+		{
+			return RemoveChildAt(ChildIndex);
+		}
+
 		return false;
 	}
+
+#if WITH_EDITOR
+	virtual void ConnectEditorData() override
+	{
+		for ( UPanelSlot* Slot : Slots )
+		{
+			Slot->Parent = this;
+			Slot->Content->Slot = Slot;
+		}
+	}
+#endif
+
+protected:
+
+	virtual UClass* GetSlotClass() const
+	{
+		return UPanelSlot::StaticClass();
+	}
+
+	virtual void OnSlotAdded(UPanelSlot* Slot)
+	{
+
+	}
+
+	virtual void OnSlotRemoved(UPanelSlot* Slot)
+	{
+
+	}
+
+protected:
+	/** Can this panel allow for multiple children? */
+	bool bCanHaveMultipleChildren;
 };
diff --git a/Engine/Source/Runtime/UMG/Public/Components/ScrollBox.h b/Engine/Source/Runtime/UMG/Public/Components/ScrollBox.h
--- a/Engine/Source/Runtime/UMG/Public/Components/ScrollBox.h
+++ b/Engine/Source/Runtime/UMG/Public/Components/ScrollBox.h
@@ -10,31 +10,20 @@ class UMG_API UScrollBox : public UPanelWidget
 {
 	GENERATED_UCLASS_BODY()
 
-	/** The slots containing the widgets that are scrollable. */
-	UPROPERTY()
-	TArray<UScrollBoxSlot*> Slots;
-
-	//TODO UMG Add ways to make adding slots callable by blueprints.
-
-	UScrollBoxSlot* AddSlot(UWidget* Content);
-
-	// UPanelWidget
-	virtual int32 GetChildrenCount() const override;
-	virtual UWidget* GetChildAt(int32 Index) const override;
-	virtual int32 GetChildIndex(UWidget* Content) const override;
-	virtual bool AddChild(UWidget* Child, FVector2D Position) override;
-	virtual bool RemoveChild(UWidget* Child) override;
-	virtual void ReplaceChildAt(int32 Index, UWidget* Child) override;
-	virtual void InsertChildAt(int32 Index, UWidget* Child) override;
-	// End UPanelWidget
-
 #if WITH_EDITOR
 	// UWidget interface
 	virtual const FSlateBrush* GetEditorIcon() override;
-	virtual void ConnectEditorData() override;
 	// End UWidget interface
 #endif
 
+protected:
+
+	// UPanelWidget
+	virtual UClass* GetSlotClass() const override;
+	virtual void OnSlotAdded(UPanelSlot* Slot) override;
+	virtual void OnSlotRemoved(UPanelSlot* Slot) override;
+	// End UPanelWidget
+
 protected:
 
 	TSharedPtr<class SScrollBox> MyScrollBox;
diff --git a/Engine/Source/Runtime/UMG/Public/Components/UniformGridPanel.h b/Engine/Source/Runtime/UMG/Public/Components/UniformGridPanel.h
--- a/Engine/Source/Runtime/UMG/Public/Components/UniformGridPanel.h
+++ b/Engine/Source/Runtime/UMG/Public/Components/UniformGridPanel.h
@@ -10,10 +10,6 @@ class UMG_API UUniformGridPanel : public UPanelWidget
 {
 	GENERATED_UCLASS_BODY()
 
-	/** The items placed on the grid */
-	UPROPERTY()
-	TArray<UUniformGridSlot*> Slots;
-
 	/** Padding given to each slot */
 	UPROPERTY(EditDefaultsOnly, Category=Layout)
 	FMargin SlotPadding;
@@ -26,25 +22,14 @@ class UMG_API UUniformGridPanel : public UPanelWidget
 	UPROPERTY(EditDefaultsOnly, Category=Layout)
 	float MinDesiredSlotHeight;
 
-	//TODO UMG Add ways to make adding slots callable by blueprints.
-
-	/** Adds content to the panel, and constructs a new slot to hold it. */
-	UUniformGridSlot* AddSlot(UWidget* Content);
+protected:
 
 	// UPanelWidget
-	virtual int32 GetChildrenCount() const override;
-	virtual UWidget* GetChildAt(int32 Index) const override;
-	virtual bool AddChild(UWidget* Child, FVector2D Position) override;
-	virtual bool RemoveChild(UWidget* Child) override;
-	virtual void ReplaceChildAt(int32 Index, UWidget* Child) override;
+	virtual UClass* GetSlotClass() const override;
+	virtual void OnSlotAdded(UPanelSlot* Slot) override;
+	virtual void OnSlotRemoved(UPanelSlot* Slot) override;
 	// End UPanelWidget
 
-#if WITH_EDITOR
-	// UWidget interface
-	virtual void ConnectEditorData() override;
-	// End UWidget interface
-#endif
-
 protected:
 
 	TSharedPtr<SUniformGridPanel> MyUniformGridPanel;
diff --git a/Engine/Source/Runtime/UMG/Public/Components/VerticalBox.h b/Engine/Source/Runtime/UMG/Public/Components/VerticalBox.h
--- a/Engine/Source/Runtime/UMG/Public/Components/VerticalBox.h
+++ b/Engine/Source/Runtime/UMG/Public/Components/VerticalBox.h
@@ -13,31 +13,24 @@ class UMG_API UVerticalBox : public UPanelWidget
 {
 	GENERATED_UCLASS_BODY()
 
-	/** The slots containing the widgets that are flowed vertically. */
-	UPROPERTY()
-	TArray<UVerticalBoxSlot*> Slots;
-
-	//TODO UMG Add ways to make adding slots callable by blueprints.
-
-	UVerticalBoxSlot* AddSlot(UWidget* Content);
-
-	// UPanelWidget
-	virtual int32 GetChildrenCount() const override;
-	virtual UWidget* GetChildAt(int32 Index) const override;
-	virtual int32 GetChildIndex(UWidget* Content) const override;
-	virtual bool AddChild(UWidget* Child, FVector2D Position) override;
-	virtual bool RemoveChild(UWidget* Child) override;
-	virtual void ReplaceChildAt(int32 Index, UWidget* Child) override;
-	virtual void InsertChildAt(int32 Index, UWidget* Child) override;
-	// End UPanelWidget
+	/**  */
+	UFUNCTION(BlueprintCallable, Category="Widget")
+	UVerticalBoxSlot* AddChild(UWidget* Content);
 
 #if WITH_EDITOR
 	// UWidget interface
 	virtual const FSlateBrush* GetEditorIcon() override;
-	virtual void ConnectEditorData() override;
 	// End UWidget interface
 #endif
 
+protected:
+
+	// UPanelWidget
+	virtual UClass* GetSlotClass() const override;
+	virtual void OnSlotAdded(UPanelSlot* Slot) override;
+	virtual void OnSlotRemoved(UPanelSlot* Slot) override;
+	// End UPanelWidget
+
 protected:
 
 	TSharedPtr<class SVerticalBox> MyVerticalBox;
diff --git a/Engine/Source/Runtime/UMG/Public/Components/WidgetSwitcher.h b/Engine/Source/Runtime/UMG/Public/Components/WidgetSwitcher.h
--- a/Engine/Source/Runtime/UMG/Public/Components/WidgetSwitcher.h
+++ b/Engine/Source/Runtime/UMG/Public/Components/WidgetSwitcher.h
@@ -13,10 +13,6 @@ class UMG_API UWidgetSwitcher : public UPanelWidget
 {
 	GENERATED_UCLASS_BODY()
 
-	/** The slots containing the widgets that are flowed vertically. */
-	UPROPERTY()
-	TArray<UWidgetSwitcherSlot*> Slots;
-
 	/** Image to draw */
 	UPROPERTY(EditDefaultsOnly, Category="Switcher", meta=( UIMin=0, ClampMin=0 ))
 	int32 ActiveWidgetIndex;
@@ -33,28 +29,16 @@ class UMG_API UWidgetSwitcher : public UPanelWidget
 	UFUNCTION(BlueprintCallable, Category="Switcher")
 	void SetActiveWidgetIndex( int32 Index );
 	
-	//TODO UMG Add ways to make adding slots callable by blueprints.
+	void SyncronizeProperties();
 
-	UWidgetSwitcherSlot* AddSlot(UWidget* Content);
+protected:
 
 	// UPanelWidget
-	virtual int32 GetChildrenCount() const override;
-	virtual UWidget* GetChildAt(int32 Index) const override;
-	virtual int32 GetChildIndex(UWidget* Content) const override;
-	virtual bool AddChild(UWidget* Child, FVector2D Position) override;
-	virtual bool RemoveChild(UWidget* Child) override;
-	virtual void ReplaceChildAt(int32 Index, UWidget* Child) override;
-	virtual void InsertChildAt(int32 Index, UWidget* Child) override;
+	virtual UClass* GetSlotClass() const override;
+	virtual void OnSlotAdded(UPanelSlot* Slot) override;
+	virtual void OnSlotRemoved(UPanelSlot* Slot) override;
 	// End UPanelWidget
 
-	void SyncronizeProperties();
-
-#if WITH_EDITOR
-	// UWidget interface
-	virtual void ConnectEditorData() override;
-	// End UWidget interface
-#endif
-
 protected:
 
 	TSharedPtr<class SWidgetSwitcher> MyWidgetSwitcher;
