UMG - Moving code out of the header for PanelWidget, fixing a build error.

diff --git a/Engine/Source/Runtime/UMG/Private/Components/PanelWidget.cpp b/Engine/Source/Runtime/UMG/Private/Components/PanelWidget.cpp
--- a/Engine/Source/Runtime/UMG/Private/Components/PanelWidget.cpp
+++ b/Engine/Source/Runtime/UMG/Private/Components/PanelWidget.cpp
@@ -18,6 +18,96 @@ UPanelWidget::UPanelWidget(const FPostConstructInitializeProperties& PCIP)
 {
 }
 
+int32 UPanelWidget::GetChildrenCount() const
+{
+	return Slots.Num();
+}
+
+UWidget* UPanelWidget::GetChildAt(int32 Index) const
+{
+	return Slots[Index]->Content;
+}
+
+int32 UPanelWidget::GetChildIndex(UWidget* Content) const
+{
+	const int32 ChildCount = GetChildrenCount();
+	for ( int32 ChildIndex = 0; ChildIndex < ChildCount; ChildIndex++ )
+	{
+		if ( GetChildAt(ChildIndex) == Content )
+		{
+			return ChildIndex;
+		}
+	}
+
+	return -1;
+}
+
+bool UPanelWidget::RemoveChildAt(int32 Index)
+{
+	UPanelSlot* Slot = Slots[Index];
+	Slots.RemoveAt(Index);
+
+	OnSlotRemoved(Slot);
+
+	return true;
+}
+
+UPanelSlot* UPanelWidget::AddChild(UWidget* Content, FVector2D Position)
+{
+	if ( !bCanHaveMultipleChildren && GetChildrenCount() > 0 )
+	{
+		return NULL;
+	}
+
+	UPanelSlot* Slot = ConstructObject<UPanelSlot>(GetSlotClass(), this);
+	Slot->SetFlags(RF_Transactional);
+	Slot->Content = Content;
+	Slot->Parent = this;
+
+	Content->Slot = Slot;
+
+	Slots.Add(Slot);
+
+	OnSlotAdded(Slot);
+
+	return Slot;
+}
+
+void UPanelWidget::ReplaceChildAt(int32 Index, UWidget* Content)
+{
+	UPanelSlot* Slot = Slots[Index];
+	Slot->Content = Content;
+
+	Content->Slot = Slot;
+
+	Slot->Refresh();
+}
+
+void UPanelWidget::InsertChildAt(int32 Index, UWidget* Content)
+{
+	UPanelSlot* Slot = ConstructObject<UPanelSlot>(GetSlotClass(), this);
+	Slot->SetFlags(RF_Transactional);
+	Slot->Content = Content;
+	Slot->Parent = this;
+
+	Content->Slot = Slot;
+
+	Slots.Insert(Slot, Index);
+
+	OnSlotAdded(Slot);
+}
+
+bool UPanelWidget::RemoveChild(UWidget* Content)
+{
+	int32 ChildIndex = GetChildIndex(Content);
+	if ( ChildIndex != -1 )
+	{
+		return RemoveChildAt(ChildIndex);
+	}
+
+	return false;
+}
+
 /////////////////////////////////////////////////////
 
 #undef LOCTEXT_NAMESPACE
diff --git a/Engine/Source/Runtime/UMG/Public/Components/PanelWidget.h b/Engine/Source/Runtime/UMG/Public/Components/PanelWidget.h
--- a/Engine/Source/Runtime/UMG/Public/Components/PanelWidget.h
+++ b/Engine/Source/Runtime/UMG/Public/Components/PanelWidget.h
@@ -20,95 +20,21 @@ public:
 
 	//TODO UMG Add ways to make adding slots callable by blueprints.
 
-	int32 GetChildrenCount() const
-	{
-		return Slots.Num();
-	}
-
-	UWidget* GetChildAt(int32 Index) const
-	{
-		return Slots[Index]->Content;
-	}
-
-	int32 GetChildIndex(UWidget* Content) const
-	{
-		const int32 ChildCount = GetChildrenCount();
-		for ( int32 ChildIndex = 0; ChildIndex < ChildCount; ChildIndex++ )
-		{
-			if ( GetChildAt(ChildIndex) == Content )
-			{
-				return ChildIndex;
-			}
-		}
+	int32 GetChildrenCount() const;
 
-		return -1;
-	}
+	UWidget* GetChildAt(int32 Index) const;
 
-	bool RemoveChildAt(int32 Index)
-	{
-		UPanelSlot* Slot = Slots[Index];
-		Slots.RemoveAt(Index);
-		
-		OnSlotRemoved(Slot);
+	int32 GetChildIndex(UWidget* Content) const;
 
-		return true;
-	}
+	bool RemoveChildAt(int32 Index);
 
-	UPanelSlot* AddChild(UWidget* Content, FVector2D Position)
-	{
-		if ( !bCanHaveMultipleChildren && GetChildrenCount() > 0 )
-		{
-			return false;
-		}
-
-		UPanelSlot* Slot = ConstructObject<UPanelSlot>(GetSlotClass(), this);
-		Slot->SetFlags(RF_Transactional);
-		Slot->Content = Content;
-		Slot->Parent = this;
+	UPanelSlot* AddChild(UWidget* Content, FVector2D Position);
 
-		Content->Slot = Slot;
+	void ReplaceChildAt(int32 Index, UWidget* Content);
 
-		Slots.Add(Slot);
+	void InsertChildAt(int32 Index, UWidget* Content);
 
-		OnSlotAdded(Slot);
-
-		return Slot;
-	}
-
-	void ReplaceChildAt(int32 Index, UWidget* Content)
-	{
-		UPanelSlot* Slot = Slots[Index];
-		Slot->Content = Content;
-
-		Content->Slot = Slot;
-
-		Slot->Refresh();
-	}
-
-	void InsertChildAt(int32 Index, UWidget* Content)
-	{
-		UPanelSlot* Slot = ConstructObject<UPanelSlot>(GetSlotClass(), this);
-		Slot->SetFlags(RF_Transactional);
-		Slot->Content = Content;
-		Slot->Parent = this;
-
-		Content->Slot = Slot;
-
-		Slots.Insert(Slot, Index);
-
-		OnSlotAdded(Slot);
-	}
-
-	bool RemoveChild(UWidget* Content)
-	{
-		int32 ChildIndex = GetChildIndex(Content);
-		if ( ChildIndex != -1 )
-		{
-			return RemoveChildAt(ChildIndex);
-		}
-
-		return false;
-	}
+	bool RemoveChild(UWidget* Content);
 
 #if WITH_EDITOR
 	virtual void ConnectEditorData() override
