UMG - Refresh on the slots is SyncronizeProperties to keep the naming common.

diff --git a/Engine/Source/Runtime/UMG/Private/Components/CanvasPanelSlot.cpp b/Engine/Source/Runtime/UMG/Private/Components/CanvasPanelSlot.cpp
--- a/Engine/Source/Runtime/UMG/Private/Components/CanvasPanelSlot.cpp
+++ b/Engine/Source/Runtime/UMG/Private/Components/CanvasPanelSlot.cpp
@@ -123,7 +123,7 @@ void UCanvasPanelSlot::SetAlignment(FVector2D InAlignment)
 	}
 }
 
-void UCanvasPanelSlot::Refresh()
+void UCanvasPanelSlot::SyncronizeProperties()
 {
 	SetOffset(LayoutData.Offsets);
 	SetAnchors(LayoutData.Anchors);
@@ -150,7 +150,7 @@ void UCanvasPanelSlot::PreEditChange(UProperty* PropertyAboutToChange)
 
 void UCanvasPanelSlot::PostEditChangeChainProperty(struct FPropertyChangedChainEvent& PropertyChangedEvent)
 {
-	Refresh();
+	SyncronizeProperties();
 
 	static FName AnchorsProperty(TEXT("Anchors"));
 	
@@ -231,7 +231,7 @@ void UCanvasPanelSlot::PostEditChangeChainProperty(struct FPropertyChangedChainE
 			}
 
 			// Apply the changes to the properties.
-			Refresh();
+			SyncronizeProperties();
 		}
 	}
 
diff --git a/Engine/Source/Runtime/UMG/Private/Components/HorizontalBoxSlot.cpp b/Engine/Source/Runtime/UMG/Private/Components/HorizontalBoxSlot.cpp
--- a/Engine/Source/Runtime/UMG/Private/Components/HorizontalBoxSlot.cpp
+++ b/Engine/Source/Runtime/UMG/Private/Components/HorizontalBoxSlot.cpp
@@ -61,7 +61,7 @@ void UHorizontalBoxSlot::SetVerticalAlignment(EVerticalAlignment InVerticalAlign
 	}
 }
 
-void UHorizontalBoxSlot::Refresh()
+void UHorizontalBoxSlot::SyncronizeProperties()
 {
 	SetPadding(Padding);
 	SetSize(Size);
diff --git a/Engine/Source/Runtime/UMG/Private/Components/OverlaySlot.cpp b/Engine/Source/Runtime/UMG/Private/Components/OverlaySlot.cpp
--- a/Engine/Source/Runtime/UMG/Private/Components/OverlaySlot.cpp
+++ b/Engine/Source/Runtime/UMG/Private/Components/OverlaySlot.cpp
@@ -50,7 +50,7 @@ void UOverlaySlot::SetVerticalAlignment(EVerticalAlignment InVerticalAlignment)
 	}
 }
 
-void UOverlaySlot::Refresh()
+void UOverlaySlot::SyncronizeProperties()
 {
 	SetPadding(Padding);
 	SetHorizontalAlignment(HorizontalAlignment);
diff --git a/Engine/Source/Runtime/UMG/Private/Components/PanelWidget.cpp b/Engine/Source/Runtime/UMG/Private/Components/PanelWidget.cpp
--- a/Engine/Source/Runtime/UMG/Private/Components/PanelWidget.cpp
+++ b/Engine/Source/Runtime/UMG/Private/Components/PanelWidget.cpp
@@ -80,7 +80,7 @@ void UPanelWidget::ReplaceChildAt(int32 Index, UWidget* Content)
 
 	Content->Slot = Slot;
 
-	Slot->Refresh();
+	Slot->SyncronizeProperties();
 }
 
 void UPanelWidget::InsertChildAt(int32 Index, UWidget* Content)
diff --git a/Engine/Source/Runtime/UMG/Private/Components/ScrollBoxSlot.cpp b/Engine/Source/Runtime/UMG/Private/Components/ScrollBoxSlot.cpp
--- a/Engine/Source/Runtime/UMG/Private/Components/ScrollBoxSlot.cpp
+++ b/Engine/Source/Runtime/UMG/Private/Components/ScrollBoxSlot.cpp
@@ -40,7 +40,7 @@ void UScrollBoxSlot::SetHorizontalAlignment(EHorizontalAlignment InHorizontalAli
 	}
 }
 
-void UScrollBoxSlot::Refresh()
+void UScrollBoxSlot::SyncronizeProperties()
 {
 	SetPadding(Padding);
 	SetHorizontalAlignment(HorizontalAlignment);
diff --git a/Engine/Source/Runtime/UMG/Private/Components/UniformGridSlot.cpp b/Engine/Source/Runtime/UMG/Private/Components/UniformGridSlot.cpp
--- a/Engine/Source/Runtime/UMG/Private/Components/UniformGridSlot.cpp
+++ b/Engine/Source/Runtime/UMG/Private/Components/UniformGridSlot.cpp
@@ -41,7 +41,7 @@ void UUniformGridSlot::SetColumn(int32 InColumn)
 	}
 }
 
-void UUniformGridSlot::Refresh()
+void UUniformGridSlot::SyncronizeProperties()
 {
 	SetRow(Row);
 	SetColumn(Column);
diff --git a/Engine/Source/Runtime/UMG/Private/Components/VerticalBoxSlot.cpp b/Engine/Source/Runtime/UMG/Private/Components/VerticalBoxSlot.cpp
--- a/Engine/Source/Runtime/UMG/Private/Components/VerticalBoxSlot.cpp
+++ b/Engine/Source/Runtime/UMG/Private/Components/VerticalBoxSlot.cpp
@@ -62,7 +62,7 @@ void UVerticalBoxSlot::SetVerticalAlignment(EVerticalAlignment InVerticalAlignme
 	}
 }
 
-void UVerticalBoxSlot::Refresh()
+void UVerticalBoxSlot::SyncronizeProperties()
 {
 	SetPadding(Padding);
 	SetSize(Size);
diff --git a/Engine/Source/Runtime/UMG/Private/Components/WidgetSwitcherSlot.cpp b/Engine/Source/Runtime/UMG/Private/Components/WidgetSwitcherSlot.cpp
--- a/Engine/Source/Runtime/UMG/Private/Components/WidgetSwitcherSlot.cpp
+++ b/Engine/Source/Runtime/UMG/Private/Components/WidgetSwitcherSlot.cpp
@@ -51,7 +51,7 @@ void UWidgetSwitcherSlot::SetVerticalAlignment(EVerticalAlignment InVerticalAlig
 	}
 }
 
-void UWidgetSwitcherSlot::Refresh()
+void UWidgetSwitcherSlot::SyncronizeProperties()
 {
 	SetPadding(Padding);
 	SetHorizontalAlignment(HorizontalAlignment);
diff --git a/Engine/Source/Runtime/UMG/Public/Components/CanvasPanelSlot.h b/Engine/Source/Runtime/UMG/Public/Components/CanvasPanelSlot.h
--- a/Engine/Source/Runtime/UMG/Public/Components/CanvasPanelSlot.h
+++ b/Engine/Source/Runtime/UMG/Public/Components/CanvasPanelSlot.h
@@ -57,7 +57,7 @@ class UMG_API UCanvasPanelSlot : public UPanelSlot
 	void SetAlignment(FVector2D InAlignment);
 
 	// UPanelSlot interface
-	virtual void Refresh() override;
+	virtual void SyncronizeProperties() override;
 	// End of UPanelSlot interface
 
 #if WITH_EDITOR
diff --git a/Engine/Source/Runtime/UMG/Public/Components/HorizontalBoxSlot.h b/Engine/Source/Runtime/UMG/Public/Components/HorizontalBoxSlot.h
--- a/Engine/Source/Runtime/UMG/Public/Components/HorizontalBoxSlot.h
+++ b/Engine/Source/Runtime/UMG/Public/Components/HorizontalBoxSlot.h
@@ -40,7 +40,7 @@ class UMG_API UHorizontalBoxSlot : public UPanelSlot
 	void BuildSlot(TSharedRef<SHorizontalBox> HorizontalBox);
 
 	// UPanelSlot interface
-	virtual void Refresh() override;
+	virtual void SyncronizeProperties() override;
 	// End of UPanelSlot interface
 
 private:
diff --git a/Engine/Source/Runtime/UMG/Public/Components/OverlaySlot.h b/Engine/Source/Runtime/UMG/Public/Components/OverlaySlot.h
--- a/Engine/Source/Runtime/UMG/Public/Components/OverlaySlot.h
+++ b/Engine/Source/Runtime/UMG/Public/Components/OverlaySlot.h
@@ -36,7 +36,7 @@ class UMG_API UOverlaySlot : public UPanelSlot
 	void BuildSlot(TSharedRef<SHorizontalBox> HorizontalBox);
 
 	// UPanelSlot interface
-	virtual void Refresh() override;
+	virtual void SyncronizeProperties() override;
 	// End of UPanelSlot interface
 
 	/** Builds the underlying FSlot for the Slate layout panel. */
diff --git a/Engine/Source/Runtime/UMG/Public/Components/PanelSlot.h b/Engine/Source/Runtime/UMG/Public/Components/PanelSlot.h
--- a/Engine/Source/Runtime/UMG/Public/Components/PanelSlot.h
+++ b/Engine/Source/Runtime/UMG/Public/Components/PanelSlot.h
@@ -33,7 +33,7 @@ public:
 	virtual bool CanMove() const;
 
 	/** Applies all properties to the live slot if possible. */
-	virtual void Refresh()
+	virtual void SyncronizeProperties()
 	{
 	}
 
@@ -43,7 +43,7 @@ public:
 	{
 		Super::PostEditChangeProperty(PropertyChangedEvent);
 
-		Refresh();
+		SyncronizeProperties();
 	}
 
 #endif
diff --git a/Engine/Source/Runtime/UMG/Public/Components/ScrollBoxSlot.h b/Engine/Source/Runtime/UMG/Public/Components/ScrollBoxSlot.h
--- a/Engine/Source/Runtime/UMG/Public/Components/ScrollBoxSlot.h
+++ b/Engine/Source/Runtime/UMG/Public/Components/ScrollBoxSlot.h
@@ -27,7 +27,7 @@ class UMG_API UScrollBoxSlot : public UPanelSlot
 	void SetHorizontalAlignment(EHorizontalAlignment InHorizontalAlignment);
 
 	// UPanelSlot interface
-	virtual void Refresh() override;
+	virtual void SyncronizeProperties() override;
 	// End of UPanelSlot interface
 
 	/** Builds the underlying FSlot for the Slate layout panel. */
diff --git a/Engine/Source/Runtime/UMG/Public/Components/UniformGridSlot.h b/Engine/Source/Runtime/UMG/Public/Components/UniformGridSlot.h
--- a/Engine/Source/Runtime/UMG/Public/Components/UniformGridSlot.h
+++ b/Engine/Source/Runtime/UMG/Public/Components/UniformGridSlot.h
@@ -42,7 +42,7 @@ public:
 	void SetColumn(int32 InColumn);
 
 	// UPanelSlot interface
-	virtual void Refresh() override;
+	virtual void SyncronizeProperties() override;
 	// End of UPanelSlot interface
 
 	/** Builds the underlying FSlot for the Slate layout panel. */
diff --git a/Engine/Source/Runtime/UMG/Public/Components/VerticalBoxSlot.h b/Engine/Source/Runtime/UMG/Public/Components/VerticalBoxSlot.h
--- a/Engine/Source/Runtime/UMG/Public/Components/VerticalBoxSlot.h
+++ b/Engine/Source/Runtime/UMG/Public/Components/VerticalBoxSlot.h
@@ -41,7 +41,7 @@ class UMG_API UVerticalBoxSlot : public UPanelSlot
 	void SetVerticalAlignment(EVerticalAlignment InVerticalAlignment);
 
 	// UPanelSlot interface
-	virtual void Refresh() override;
+	virtual void SyncronizeProperties() override;
 	// End of UPanelSlot interface
 
 	/** Builds the underlying FSlot for the Slate layout panel. */
diff --git a/Engine/Source/Runtime/UMG/Public/Components/WidgetSwitcherSlot.h b/Engine/Source/Runtime/UMG/Public/Components/WidgetSwitcherSlot.h
--- a/Engine/Source/Runtime/UMG/Public/Components/WidgetSwitcherSlot.h
+++ b/Engine/Source/Runtime/UMG/Public/Components/WidgetSwitcherSlot.h
@@ -34,7 +34,7 @@ class UMG_API UWidgetSwitcherSlot : public UPanelSlot
 	void SetVerticalAlignment(EVerticalAlignment InVerticalAlignment);
 
 	// UPanelSlot interface
-	virtual void Refresh() override;
+	virtual void SyncronizeProperties() override;
 	// End of UPanelSlot interface
 
 	/** Builds the underlying FSlot for the Slate layout panel. */
