UMG - Renaming some widgets that were still referred to as components.

diff --git a/Engine/Source/Runtime/UMG/Private/Components/CircularThrobber.cpp b/Engine/Source/Runtime/UMG/Private/Components/CircularThrobber.cpp
new file mode 100644
--- /dev/null
+++ b/Engine/Source/Runtime/UMG/Private/Components/CircularThrobber.cpp
@@ -0,0 +1,36 @@
+// Copyright 1998-2014 Epic Games, Inc. All Rights Reserved.
+
+#include "UMGPrivatePCH.h"
+#include "Slate/SlateBrushAsset.h"
+
+#define LOCTEXT_NAMESPACE "UMG"
+
+/////////////////////////////////////////////////////
+// UCircularThrobber
+
+UCircularThrobber::UCircularThrobber(const FPostConstructInitializeProperties& PCIP)
+	: Super(PCIP)
+{
+	SCircularThrobber::FArguments DefaultArgs;
+
+	NumberOfPieces = DefaultArgs._NumPieces;
+	Period = DefaultArgs._Period;
+	Radius = DefaultArgs._Radius;
+}
+
+TSharedRef<SWidget> UCircularThrobber::RebuildWidget()
+{
+	SCircularThrobber::FArguments DefaultArgs;
+
+	const FSlateBrush* Image = PieceImage ? &PieceImage->Brush : DefaultArgs._PieceImage;
+	
+	return SNew(SCircularThrobber)
+		.PieceImage(Image)
+		.NumPieces(NumberOfPieces)
+		.Period(Period)
+		.Radius(Radius);
+}
+
+/////////////////////////////////////////////////////
+
+#undef LOCTEXT_NAMESPACE
diff --git a/Engine/Source/Runtime/UMG/Private/Components/CircularThrobberComponent.cpp b/Engine/Source/Runtime/UMG/Private/Components/CircularThrobberComponent.cpp
deleted file mode 100644
diff --git a/Engine/Source/Runtime/UMG/Private/Components/Throbber.cpp b/Engine/Source/Runtime/UMG/Private/Components/Throbber.cpp
new file mode 100644
--- /dev/null
+++ b/Engine/Source/Runtime/UMG/Private/Components/Throbber.cpp
@@ -0,0 +1,41 @@
+// Copyright 1998-2014 Epic Games, Inc. All Rights Reserved.
+
+#include "UMGPrivatePCH.h"
+#include "Slate/SlateBrushAsset.h"
+
+#define LOCTEXT_NAMESPACE "UMG"
+
+/////////////////////////////////////////////////////
+// UThrobber
+
+UThrobber::UThrobber(const FPostConstructInitializeProperties& PCIP)
+	: Super(PCIP)
+{
+	SThrobber::FArguments DefaultArgs;
+
+	NumberOfPieces = DefaultArgs._NumPieces;
+
+	bAnimateVertically = (DefaultArgs._Animate & SThrobber::Vertical) != 0;
+	bAnimateHorizontally = (DefaultArgs._Animate & SThrobber::Horizontal) != 0;
+	bAnimateOpacity = (DefaultArgs._Animate & SThrobber::Opacity) != 0;
+}
+
+TSharedRef<SWidget> UThrobber::RebuildWidget()
+{
+	SThrobber::FArguments DefaultArgs;
+
+	const FSlateBrush* Image = PieceImage ? &PieceImage->Brush : DefaultArgs._PieceImage;
+
+	const int32 AnimationParams = (bAnimateVertically ? SThrobber::Vertical : 0) |
+		(bAnimateHorizontally ? SThrobber::Horizontal : 0) |
+		(bAnimateOpacity ? SThrobber::Opacity : 0);
+
+	return SNew(SThrobber)
+		.PieceImage(Image)
+		.NumPieces(NumberOfPieces)
+		.Animate(static_cast<SThrobber::EAnimation>(AnimationParams));
+}
+
+/////////////////////////////////////////////////////
+
+#undef LOCTEXT_NAMESPACE
diff --git a/Engine/Source/Runtime/UMG/Private/Components/ThrobberComponent.cpp b/Engine/Source/Runtime/UMG/Private/Components/ThrobberComponent.cpp
deleted file mode 100644
diff --git a/Engine/Source/Runtime/UMG/Public/Components/CircularThrobber.h b/Engine/Source/Runtime/UMG/Public/Components/CircularThrobber.h
new file mode 100644
--- /dev/null
+++ b/Engine/Source/Runtime/UMG/Public/Components/CircularThrobber.h
@@ -0,0 +1,37 @@
+// Copyright 1998-2014 Epic Games, Inc. All Rights Reserved.
+
+#pragma once
+
+#include "CircularThrobber.generated.h"
+
+class USlateBrushAsset;
+
+/**  A throbber widget that orients images in a spinning circle. */
+UCLASS(meta=( Category="Misc" ), ClassGroup=UserInterface)
+class UMG_API UCircularThrobber : public UWidget
+{
+	GENERATED_UCLASS_BODY()
+
+protected:
+
+	/** How many pieces there are */
+	UPROPERTY(EditAnywhere, BlueprintReadOnly, Category=Appearance, meta=(ClampMin = "1", UIMin = "1", UIMax = "8"))
+	int32 NumberOfPieces;
+
+	/** The amount of time for a full circle (in seconds) */
+	UPROPERTY(EditAnywhere, BlueprintReadOnly, Category=Appearance)
+	float Period;
+
+	/** The radius of the circle */
+	UPROPERTY(EditAnywhere, BlueprintReadOnly, Category=Appearance)
+	float Radius;
+
+	/** Image to use for each segment of the throbber */
+	UPROPERTY(EditAnywhere, BlueprintReadWrite, Category=Appearance, meta=( DisplayThumbnail = "true" ))
+	USlateBrushAsset* PieceImage;
+
+protected:
+	// UWidget interface
+	virtual TSharedRef<SWidget> RebuildWidget() override;
+	// End of UWidget interface
+};
diff --git a/Engine/Source/Runtime/UMG/Public/Components/CircularThrobberComponent.h b/Engine/Source/Runtime/UMG/Public/Components/CircularThrobberComponent.h
deleted file mode 100644
diff --git a/Engine/Source/Runtime/UMG/Public/Components/Throbber.h b/Engine/Source/Runtime/UMG/Public/Components/Throbber.h
new file mode 100644
--- /dev/null
+++ b/Engine/Source/Runtime/UMG/Public/Components/Throbber.h
@@ -0,0 +1,39 @@
+// Copyright 1998-2014 Epic Games, Inc. All Rights Reserved.
+
+#pragma once
+
+#include "Throbber.generated.h"
+
+class USlateBrushAsset;
+
+/** A Throbber widget that shows several zooming circles in a row. */
+UCLASS(meta=( Category="Misc" ), ClassGroup=UserInterface)
+class UMG_API UThrobber : public UWidget
+{
+	GENERATED_UCLASS_BODY()
+
+	/** How many pieces there are */
+	UPROPERTY(EditAnywhere, BlueprintReadOnly, Category=Appearance, meta=(ClampMin = "1", UIMin = "1", UIMax = "8"))
+	int32 NumberOfPieces;
+
+	/** Should the pieces animate horizontally? */
+	UPROPERTY(EditAnywhere, BlueprintReadWrite, Category=Appearance)
+	bool bAnimateHorizontally;
+
+	/** Should the pieces animate vertically? */
+	UPROPERTY(EditAnywhere, BlueprintReadWrite, Category=Appearance)
+	bool bAnimateVertically;
+
+	/** Should the pieces animate their opacity? */
+	UPROPERTY(EditAnywhere, BlueprintReadWrite, Category=Appearance)
+	bool bAnimateOpacity;
+
+	/** Image to use for each segment of the throbber */
+	UPROPERTY(EditAnywhere, BlueprintReadWrite, Category=Appearance, meta=( DisplayThumbnail = "true" ))
+	USlateBrushAsset* PieceImage;
+
+protected:
+	// UWidget interface
+	virtual TSharedRef<SWidget> RebuildWidget() override;
+	// End of UWidget interface
+};
diff --git a/Engine/Source/Runtime/UMG/Public/Components/ThrobberComponent.h b/Engine/Source/Runtime/UMG/Public/Components/ThrobberComponent.h
deleted file mode 100644
diff --git a/Engine/Source/Runtime/UMG/Public/UMG.h b/Engine/Source/Runtime/UMG/Public/UMG.h
--- a/Engine/Source/Runtime/UMG/Public/UMG.h
+++ b/Engine/Source/Runtime/UMG/Public/UMG.h
@@ -67,8 +67,8 @@
 
 #include "Spacer.h"
 
-#include "ThrobberComponent.h"
-#include "CircularThrobberComponent.h"
+#include "Throbber.h"
+#include "CircularThrobber.h"
 
 #include "TableViewBase.h"
 #include "ListView.h"
