UMG - Now providing a way for widget classes to add additional metadata to items displayed in the treeview, like a small snippet of text of a textblock.

diff --git a/Engine/Source/Runtime/UMG/Private/Components/TextBlock.cpp b/Engine/Source/Runtime/UMG/Private/Components/TextBlock.cpp
--- a/Engine/Source/Runtime/UMG/Private/Components/TextBlock.cpp
+++ b/Engine/Source/Runtime/UMG/Private/Components/TextBlock.cpp
@@ -105,6 +105,15 @@ void UTextBlock::SyncronizeProperties()
 
 #if WITH_EDITOR
 
+FString UTextBlock::GetLabelMetadata() const
+{
+	const int32 MaxSampleLength = 15;
+
+	FString TextStr = Text.ToString();
+	TextStr = TextStr.Len() <= MaxSampleLength ? TextStr : TextStr.Left(MaxSampleLength - 2) + TEXT("..");
+	return TEXT(" \"") + TextStr + TEXT("\"");
+}
+
 void UTextBlock::HandleTextCommitted(const FText& InText, ETextCommit::Type CommitteType)
 {
 	//TODO UMG How will this migrate to the template?  Seems to me we need the previews to have access to their templates!
diff --git a/Engine/Source/Runtime/UMG/Private/Components/Widget.cpp b/Engine/Source/Runtime/UMG/Private/Components/Widget.cpp
--- a/Engine/Source/Runtime/UMG/Private/Components/Widget.cpp
+++ b/Engine/Source/Runtime/UMG/Private/Components/Widget.cpp
@@ -145,11 +145,16 @@ bool UWidget::IsGeneratedName() const
 	return false;
 }
 
+FString UWidget::GetLabelMetadata() const
+{
+	return "";
+}
+
 FString UWidget::GetLabel() const
 {
 	if ( IsGeneratedName() && !bIsVariable )
 	{
-		return "[" + GetClass()->GetName() + "]";
+		return "[" + GetClass()->GetName() + "]" + GetLabelMetadata();
 	}
 	else
 	{
diff --git a/Engine/Source/Runtime/UMG/Public/Components/TextBlock.h b/Engine/Source/Runtime/UMG/Public/Components/TextBlock.h
--- a/Engine/Source/Runtime/UMG/Public/Components/TextBlock.h
+++ b/Engine/Source/Runtime/UMG/Public/Components/TextBlock.h
@@ -12,27 +12,27 @@ class UMG_API UTextBlock : public UWidget
 
 public:
 	/**  
- 	 * Sets the color and opacity of the text in this text block
+	 * Sets the color and opacity of the text in this text block
 	 *
 	 * @param InColorAndOpacity		The new text color and opacity
- 	 */
+	 */
 	UFUNCTION(BlueprintCallable, Category="Appearance")
 	void SetColorAndOpacity(FLinearColor InColorAndOpacity);
 
 	/**  
- 	 * Sets the color and opacity of the text drop shadow
+	 * Sets the color and opacity of the text drop shadow
 	 * Note: if opacity is zero no shadow will be drawn
 	 *
 	 * @param InShadowColorAndOpacity		The new drop shadow color and opacity
- 	 */
+	 */
 	UFUNCTION(BlueprintCallable, Category="Appearance")
 	void SetShadowColorAndOpacity(FLinearColor InShadowColorAndOpacity);
 
 	/**  
- 	 * Sets the offset that the text drop shadow should be drawn at
+	 * Sets the offset that the text drop shadow should be drawn at
 	 *
 	 * @param InShadowOffset		The new offset
- 	 */
+	 */
 	UFUNCTION(BlueprintCallable, Category="Appearance")
 	void SetShadowOffset(FVector2D InShadowOffset);
 
@@ -94,6 +94,8 @@ public:
 	virtual const FSlateBrush* GetEditorIcon() override;
 	// End UWidget interface
 
+	virtual FString GetLabelMetadata() const override;
+
 	void HandleTextCommitted(const FText& InText, ETextCommit::Type CommitteType);
 
 	virtual void OnDesignerDoubleClicked() override;
diff --git a/Engine/Source/Runtime/UMG/Public/Components/Widget.h b/Engine/Source/Runtime/UMG/Public/Components/Widget.h
--- a/Engine/Source/Runtime/UMG/Public/Components/Widget.h
+++ b/Engine/Source/Runtime/UMG/Public/Components/Widget.h
@@ -151,6 +151,9 @@ public:
 	/** Is the label generated or provided by the user? */
 	bool IsGeneratedName() const;
 
+	/** Get Label Metadata, which may be as simple as a bit of string data to help identify an anonymous text block. */
+	virtual FString GetLabelMetadata() const;
+
 	/** Gets the label to display to the user for this widget. */
 	FString GetLabel() const;
 
