UMG - The tree view searching now looks at the label instead of the name.  Adding additional ufunctions to the scrollbox to manipulate scrolling.

diff --git a/Engine/Source/Editor/UMGEditor/Private/SUMGEditorTree.cpp b/Engine/Source/Editor/UMGEditor/Private/SUMGEditorTree.cpp
--- a/Engine/Source/Editor/UMGEditor/Private/SUMGEditorTree.cpp
+++ b/Engine/Source/Editor/UMGEditor/Private/SUMGEditorTree.cpp
@@ -107,7 +107,7 @@ FReply SUMGEditorTree::OnKeyDown(const FGeometry& MyGeometry, const FKeyboardEve
 
 void SUMGEditorTree::TransformWidgetToString(const UWidget* Widget, OUT TArray< FString >& Array)
 {
-	Array.Add(Widget->GetName());
+	Array.Add( Widget->GetLabel() );
 }
 
 void SUMGEditorTree::OnSearchChanged(const FText& InFilterText)
diff --git a/Engine/Source/Runtime/UMG/Private/Components/ScrollBox.cpp b/Engine/Source/Runtime/UMG/Private/Components/ScrollBox.cpp
--- a/Engine/Source/Runtime/UMG/Private/Components/ScrollBox.cpp
+++ b/Engine/Source/Runtime/UMG/Private/Components/ScrollBox.cpp
@@ -52,6 +52,31 @@ TSharedRef<SWidget> UScrollBox::RebuildWidget()
 	return BuildDesignTimeWidget( MyScrollBox.ToSharedRef() );
 }
 
+void UScrollBox::SyncronizeProperties()
+{
+	MyScrollBox->SetScrollOffset(DesiredScrollOffset);
+}
+
+void UScrollBox::ClearChildren()
+{
+	Slots.Reset();
+
+	if ( MyScrollBox.IsValid() )
+	{
+		MyScrollBox->ClearChildren();
+	}
+}
+
+void UScrollBox::SetScrollOffset(float NewScrollOffset)
+{
+	DesiredScrollOffset = NewScrollOffset;
+
+	if ( MyScrollBox.IsValid() )
+	{
+		MyScrollBox->SetScrollOffset(NewScrollOffset);
+	}
+}
+
 #if WITH_EDITOR
 
 const FSlateBrush* UScrollBox::GetEditorIcon()
diff --git a/Engine/Source/Runtime/UMG/Public/Components/Border.h b/Engine/Source/Runtime/UMG/Public/Components/Border.h
--- a/Engine/Source/Runtime/UMG/Public/Components/Border.h
+++ b/Engine/Source/Runtime/UMG/Public/Components/Border.h
@@ -88,7 +88,7 @@ class UMG_API UBorder : public UContentWidget
 	void SetContentPadding(FMargin InContentPadding);
 
 	// UWidget interface
-	void SyncronizeProperties() override;
+	virtual void SyncronizeProperties() override;
 	// End of UWidget interface
 
 #if WITH_EDITOR
diff --git a/Engine/Source/Runtime/UMG/Public/Components/ScrollBox.h b/Engine/Source/Runtime/UMG/Public/Components/ScrollBox.h
--- a/Engine/Source/Runtime/UMG/Public/Components/ScrollBox.h
+++ b/Engine/Source/Runtime/UMG/Public/Components/ScrollBox.h
@@ -10,6 +10,18 @@ class UMG_API UScrollBox : public UPanelWidget
 {
 	GENERATED_UCLASS_BODY()
 
+	/** Removes all child widget from the scrollbox */
+	UFUNCTION(BlueprintCallable, Category="Widget")
+	void ClearChildren();
+
+	/** Updates the scroll offset of the scrollbox */
+	UFUNCTION(BlueprintCallable, Category="Widget")
+	void SetScrollOffset(float NewScrollOffset);
+
+	// UWidget interface
+	virtual void SyncronizeProperties() override;
+	// End of UWidget interface
+
 #if WITH_EDITOR
 	// UWidget interface
 	virtual const FSlateBrush* GetEditorIcon() override;
@@ -25,6 +37,8 @@ protected:
 	// End UPanelWidget
 
 protected:
+	/** The desired scroll offset for the underlying scrollbox.  This is a cache so that it can be set before the widget is constructed. */
+	float DesiredScrollOffset;
 
 	TSharedPtr<class SScrollBox> MyScrollBox;
 
