UMG - Backwards compatability breakage.  The widget tree is now actually a tree, with a root widget and the tree is dynamically searched.  Currently there are no nice iterators, we just build up a list of widgets upon request.  This should fix a lot of bugs with widgets being abandoned from the tree but remaining in the widget array that the tree had previously.

diff --git a/Engine/Source/Editor/UMGEditor/Private/Customizations/CanvasSlotCustomization.h b/Engine/Source/Editor/UMGEditor/Private/Customizations/CanvasSlotCustomization.h
--- a/Engine/Source/Editor/UMGEditor/Private/Customizations/CanvasSlotCustomization.h
+++ b/Engine/Source/Editor/UMGEditor/Private/Customizations/CanvasSlotCustomization.h
@@ -12,7 +12,7 @@ public:
 	}
 
 	FCanvasSlotCustomization(UBlueprint* InBlueprint)
-		: Blueprint(CastChecked<UWidgetBlueprint>(InBlueprint))
+		: Blueprint(Cast<UWidgetBlueprint>(InBlueprint))
 	{
 	}
 	
diff --git a/Engine/Source/Editor/UMGEditor/Private/SUMGDesigner.cpp b/Engine/Source/Editor/UMGEditor/Private/SUMGDesigner.cpp
--- a/Engine/Source/Editor/UMGEditor/Private/SUMGDesigner.cpp
+++ b/Engine/Source/Editor/UMGEditor/Private/SUMGDesigner.cpp
@@ -1029,19 +1029,19 @@ UWidget* SUMGDesigner::ProcessDropAndAddWidget(const FGeometry& MyGeometry, cons
 					Transaction.Cancel();
 				}
 			}
-			else if ( BP->WidgetTree->WidgetTemplates.Num() == 1 )
-			{
-				UWidget* Widget = DragDropOp->Template->Create(BP->WidgetTree);
-				Widget->IsDesignTime(true);
-				
-				DropPreviewParent = NULL;
-				
-				return Widget;
-			}
+			//else if ( BP->WidgetTree->RootWidget != NULL && BP->WidgetTree->RootWidget.Child == 1 )
+			//{
+			//	UWidget* Widget = DragDropOp->Template->Create(BP->WidgetTree);
+			//	Widget->IsDesignTime(true);
+			//	
+			//	DropPreviewParent = NULL;
+			//	
+			//	return Widget;
+			//}
 		}
 		else
 		{
-			if ( BP->WidgetTree->WidgetTemplates.Num() == 0 )
+			if ( BP->WidgetTree->RootWidget == NULL )
 			{
 				FScopedTransaction Transaction(LOCTEXT("Designer_AddWidget", "Add Widget"));
 
@@ -1057,6 +1057,8 @@ UWidget* SUMGDesigner::ProcessDropAndAddWidget(const FGeometry& MyGeometry, cons
 				UWidget* Widget = DragDropOp->Template->Create(BP->WidgetTree);
 				Widget->IsDesignTime(true);
 
+				BP->WidgetTree->RootWidget = Widget;
+
 				SelectedWidget = FWidgetReference::FromTemplate(BlueprintEditor.Pin(), Widget);
 
 				DropPreviewParent = NULL;
diff --git a/Engine/Source/Editor/UMGEditor/Private/SUMGEditorTree.cpp b/Engine/Source/Editor/UMGEditor/Private/SUMGEditorTree.cpp
--- a/Engine/Source/Editor/UMGEditor/Private/SUMGEditorTree.cpp
+++ b/Engine/Source/Editor/UMGEditor/Private/SUMGEditorTree.cpp
@@ -295,9 +295,8 @@ void SUMGEditorTree::RefreshTree()
 	RootWidgets.Reset();
 
 	UWidgetBlueprint* Blueprint = GetBlueprint();
-	TArray<UWidget*>& WidgetTemplates = Blueprint->WidgetTree->WidgetTemplates;
 
-	if ( WidgetTemplates.Num() > 0 )
+	if ( Blueprint->WidgetTree->RootWidget )
 	{
 		bool bRootPassed = true;
 
@@ -322,15 +321,15 @@ void SUMGEditorTree::RefreshTree()
 		
 		bIsFilterActive = bWillFilterBeActive;
 
-		WidgetsPassingFilter.Empty( WidgetTemplates.Num() );
+		WidgetsPassingFilter.Empty();
 		if ( bIsFilterActive )
 		{
-			bRootPassed = FilterWidgetHierarchy(WidgetTemplates[0]);
+			bRootPassed = FilterWidgetHierarchy(Blueprint->WidgetTree->RootWidget);
 		}
 
 		if ( bRootPassed )
 		{
-			RootWidgets.Add(WidgetTemplates[0]);
+			RootWidgets.Add(Blueprint->WidgetTree->RootWidget);
 		}
 	}
 
diff --git a/Engine/Source/Editor/UMGEditor/Private/WidgetBlueprintCompiler.cpp b/Engine/Source/Editor/UMGEditor/Private/WidgetBlueprintCompiler.cpp
--- a/Engine/Source/Editor/UMGEditor/Private/WidgetBlueprintCompiler.cpp
+++ b/Engine/Source/Editor/UMGEditor/Private/WidgetBlueprintCompiler.cpp
@@ -537,7 +537,9 @@ void FWidgetBlueprintCompiler::ValidateWidgetNames()
 		}
 	}
 
-	for ( UWidget* Widget : Blueprint->WidgetTree->WidgetTemplates )
+	TArray<UWidget*> Widgets;
+	Blueprint->WidgetTree->GetAllWidgets(Widgets);
+	for ( UWidget* Widget : Widgets )
 	{
 		if ( ParentBPNameValidator.IsValid() && ParentBPNameValidator->IsValid(Widget->GetName()) != EValidatorResult::Ok )
 		{
@@ -580,12 +582,12 @@ void FWidgetBlueprintCompiler::CleanAndSanitizeClass(UBlueprintGeneratedClass* C
 	TArray<UObject*> ClassSubObjects;
 	GetObjectsWithOuter(ClassToClean, ClassSubObjects, true);
 
-	UWidgetBlueprint* Blueprint = WidgetBlueprint();
+	//UWidgetBlueprint* Blueprint = WidgetBlueprint();
 
-	if ( 0 != Blueprint->WidgetTree->WidgetTemplates.Num() )
-	{
-		ClassSubObjects.RemoveAllSwap(FCullTemplateObjectsHelper<UWidget>(Blueprint->WidgetTree->WidgetTemplates));
-	}
+	//if ( 0 != Blueprint->WidgetTree->WidgetTemplates.Num() )
+	//{
+	//	ClassSubObjects.RemoveAllSwap(FCullTemplateObjectsHelper<UWidget>(Blueprint->WidgetTree->WidgetTemplates));
+	//}
 }
 
 void FWidgetBlueprintCompiler::CreateClassVariablesFromBlueprint()
@@ -597,7 +599,10 @@ void FWidgetBlueprintCompiler::CreateClassVariablesFromBlueprint()
 	ValidateWidgetNames();
 
 	// Build the set of variables based on the variable widgets in the widget tree.
-	for ( UWidget* Widget : Blueprint->WidgetTree->WidgetTemplates )
+	TArray<UWidget*> Widgets;
+	Blueprint->WidgetTree->GetAllWidgets(Widgets);
+
+	for ( UWidget* Widget : Widgets )
 	{
 		// Skip non-variable widgets
 		if ( !Widget->bIsVariable )
diff --git a/Engine/Source/Editor/UMGEditor/Private/WidgetBlueprintEditor.cpp b/Engine/Source/Editor/UMGEditor/Private/WidgetBlueprintEditor.cpp
--- a/Engine/Source/Editor/UMGEditor/Private/WidgetBlueprintEditor.cpp
+++ b/Engine/Source/Editor/UMGEditor/Private/WidgetBlueprintEditor.cpp
@@ -79,9 +79,11 @@ void FWidgetBlueprintEditor::InitWidgetBlueprintEditor(const EToolkitMode::Type
 	UWidgetBlueprint* Blueprint = GetWidgetBlueprintObj();
 
 	// If this blueprint is empty, add a canvas panel as the root widget.
-	if ( Blueprint->WidgetTree->WidgetTemplates.Num() == 0 )
+	if ( Blueprint->WidgetTree->RootWidget == NULL )
 	{
-		Blueprint->WidgetTree->ConstructWidget<UCanvasPanel>(UCanvasPanel::StaticClass());
+		 UWidget* RootWidget = Blueprint->WidgetTree->ConstructWidget<UCanvasPanel>(UCanvasPanel::StaticClass());
+		 RootWidget->IsDesignTime(true);
+		 Blueprint->WidgetTree->RootWidget = RootWidget;
 	}
 
 	UpdatePreview(GetWidgetBlueprintObj(), true);
diff --git a/Engine/Source/Editor/UMGEditor/Private/WidgetBlueprintEditorUtils.cpp b/Engine/Source/Editor/UMGEditor/Private/WidgetBlueprintEditorUtils.cpp
--- a/Engine/Source/Editor/UMGEditor/Private/WidgetBlueprintEditorUtils.cpp
+++ b/Engine/Source/Editor/UMGEditor/Private/WidgetBlueprintEditorUtils.cpp
@@ -297,11 +297,10 @@ void FWidgetBlueprintEditorUtils::PasteWidgets(UWidgetBlueprint* BP, FWidgetRefe
 		{
 			UPanelSlot* Slot = ParentWidget->AddChild(NewWidget);
 			Slot->SetDesiredPosition(PasteLocation);
+			//TODO UMG - The paste location needs to be relative from the most upper left hand corner of other widgets in their container.
 		}
 
-		//TODO UMG HACK Because the WidgetTree isn't a real tree, add every widget.
 		NewWidget->SetFlags( RF_Transactional );
-		BP->WidgetTree->WidgetTemplates.Add( NewWidget );
 	}
 
 	//TODO UMG Reorder the live slot without rebuilding the structure
diff --git a/Engine/Source/Runtime/UMG/Private/WidgetBlueprint.cpp b/Engine/Source/Runtime/UMG/Private/WidgetBlueprint.cpp
--- a/Engine/Source/Runtime/UMG/Private/WidgetBlueprint.cpp
+++ b/Engine/Source/Runtime/UMG/Private/WidgetBlueprint.cpp
@@ -30,7 +30,9 @@ void UWidgetBlueprint::PostLoad()
 	Super::PostLoad();
 
 #if WITH_EDITOR
-	for ( UWidget* Widget : WidgetTree->WidgetTemplates )
+	TArray<UWidget*> Widgets;
+	WidgetTree->GetAllWidgets(Widgets);
+	for ( UWidget* Widget : Widgets )
 	{
 		Widget->ConnectEditorData();
 	}
diff --git a/Engine/Source/Runtime/UMG/Private/WidgetBlueprintGeneratedClass.cpp b/Engine/Source/Runtime/UMG/Private/WidgetBlueprintGeneratedClass.cpp
--- a/Engine/Source/Runtime/UMG/Private/WidgetBlueprintGeneratedClass.cpp
+++ b/Engine/Source/Runtime/UMG/Private/WidgetBlueprintGeneratedClass.cpp
@@ -54,7 +54,10 @@ void UWidgetBlueprintGeneratedClass::InitializeWidget(UUserWidget* UserWidget) c
 
 	UClass* ActorClass = UserWidget->GetClass();
 
-	for ( UWidget* Widget : ClonedTree->WidgetTemplates )
+	TArray<UWidget*> ClonedWidgets;
+	ClonedTree->GetAllWidgets(ClonedWidgets);
+
+	for ( UWidget* Widget : ClonedWidgets )
 	{
 		// Not fatal if NULL, but shouldn't happen
 		if ( !ensure(Widget != NULL) )
@@ -62,7 +65,7 @@ void UWidgetBlueprintGeneratedClass::InitializeWidget(UUserWidget* UserWidget) c
 			continue;
 		}
 
-		FName NewName = *( FString::Printf(TEXT("WidgetComponent__%d"), UserWidget->Components.Num()) );
+		//FName NewName = *( FString::Printf(TEXT("WidgetComponent__%d"), UserWidget->Components.Num()) );
 
 		FString VariableName = Widget->GetName();
 
diff --git a/Engine/Source/Runtime/UMG/Private/WidgetTree.cpp b/Engine/Source/Runtime/UMG/Private/WidgetTree.cpp
--- a/Engine/Source/Runtime/UMG/Private/WidgetTree.cpp
+++ b/Engine/Source/Runtime/UMG/Private/WidgetTree.cpp
@@ -13,7 +13,12 @@ UWidgetTree::UWidgetTree(const FPostConstructInitializeProperties& PCIP)
 UWidget* UWidgetTree::FindWidget(const FString& Name) const
 {
 	FString ExistingName;
-	for ( UWidget* Widget : WidgetTemplates )
+
+	// TODO UMG Hacky, remove this find widget function, or make it faster.
+	TArray<UWidget*> Widgets;
+	GetAllWidgets(Widgets);
+
+	for ( UWidget* Widget : Widgets )
 	{
 		Widget->GetName(ExistingName);
 		if ( ExistingName.Equals(Name, ESearchCase::IgnoreCase) )
@@ -27,24 +32,17 @@ UWidget* UWidgetTree::FindWidget(const FString& Name) const
 
 UPanelWidget* UWidgetTree::FindWidgetParent(UWidget* Widget, int32& OutChildIndex)
 {
-	for ( UWidget* Template : WidgetTemplates )
+	UPanelWidget* Parent = Widget->GetParent();
+	if ( Parent != NULL )
 	{
-		UPanelWidget* NonLeafTemplate = Cast<UPanelWidget>(Template);
-		if ( NonLeafTemplate )
-		{
-			for ( int32 ChildIndex = 0; ChildIndex < NonLeafTemplate->GetChildrenCount(); ChildIndex++ )
-			{
-				if ( NonLeafTemplate->GetChildAt(ChildIndex) == Widget )
-				{
-					OutChildIndex = ChildIndex;
-					return NonLeafTemplate;
-				}
-			}
-		}
+		OutChildIndex = Parent->GetChildIndex(Widget);
+	}
+	else
+	{
+		OutChildIndex = 0;
 	}
 
-	OutChildIndex = -1;
-	return NULL;
+	return Parent;
 }
 
 bool UWidgetTree::RemoveWidget(UWidget* InRemovedWidget, bool bIsRecursive)
@@ -75,7 +73,28 @@ bool UWidgetTree::RemoveWidget(UWidget* InRemovedWidget, bool bIsRecursive)
 		}
 	}
 
-	int32 IndexRemoved = WidgetTemplates.Remove(InRemovedWidget);
-	
-	return bRemoved || IndexRemoved != -1;
+	return bRemoved;
+}
+
+void UWidgetTree::GetAllWidgets(TArray<UWidget*>& Widgets) const
+{
+	if ( RootWidget )
+	{
+		Widgets.Add(RootWidget);
+		GetChildWidgets(RootWidget, Widgets);
+	}
+}
+
+void UWidgetTree::GetChildWidgets(UWidget* Parent, TArray<UWidget*>& Widgets) const
+{
+	if ( UPanelWidget* PanelParent = Cast<UPanelWidget>(Parent) )
+	{
+		for ( int32 ChildIndex = 0; ChildIndex < PanelParent->GetChildrenCount(); ChildIndex++ )
+		{
+			UWidget* ChildWidget = PanelParent->GetChildAt(ChildIndex);
+			Widgets.Add(ChildWidget);
+
+			GetChildWidgets(ChildWidget, Widgets);
+		}
+	}
 }
diff --git a/Engine/Source/Runtime/UMG/Public/Blueprint/WidgetTree.h b/Engine/Source/Runtime/UMG/Public/Blueprint/WidgetTree.h
--- a/Engine/Source/Runtime/UMG/Public/Blueprint/WidgetTree.h
+++ b/Engine/Source/Runtime/UMG/Public/Blueprint/WidgetTree.h
@@ -11,9 +11,9 @@ class UMG_API UWidgetTree : public UObject
 	GENERATED_UCLASS_BODY()
 
 public:
-	/** Array of templates for widgets */
+	/** The root widget of the tree */
 	UPROPERTY()
-	TArray<UWidget*> WidgetTemplates;
+	UWidget* RootWidget;
 
 public:
 
@@ -23,6 +23,10 @@ public:
 
 	class UPanelWidget* FindWidgetParent(UWidget* Widget, int32& OutChildIndex);
 
+	void GetAllWidgets(TArray<UWidget*>& Widgets) const;
+
+	void GetChildWidgets(UWidget* Parent, TArray<UWidget*>& Widgets) const;
+
 	/** Constructs the widget, and adds it to the tree. */
 	template< class T >
 	T* ConstructWidget(TSubclassOf<UWidget> WidgetType)
@@ -34,7 +38,6 @@ public:
 
 		UWidget* Widget = (UWidget*)ConstructObject<UWidget>(WidgetType, this);
 		Widget->SetFlags(RF_Transactional);
-		WidgetTemplates.Add(Widget);
 
 		return (T*)Widget;
 	}
