UMG - Checking for null types when constructing widgets from blueprints.

diff --git a/Engine/Source/Runtime/UMG/Private/WidgetBlueprintLibrary.cpp b/Engine/Source/Runtime/UMG/Private/WidgetBlueprintLibrary.cpp
--- a/Engine/Source/Runtime/UMG/Private/WidgetBlueprintLibrary.cpp
+++ b/Engine/Source/Runtime/UMG/Private/WidgetBlueprintLibrary.cpp
@@ -16,9 +16,9 @@ UWidgetBlueprintLibrary::UWidgetBlueprintLibrary(const FPostConstructInitializeP
 
 UUserWidget* UWidgetBlueprintLibrary::Create(UObject* WorldContextObject, TSubclassOf<UUserWidget> WidgetType)
 {
-	if ( WidgetType->HasAnyClassFlags(CLASS_Abstract) )
+	if ( WidgetType == NULL || WidgetType->HasAnyClassFlags(CLASS_Abstract) )
 	{
-		// TODO Error?
+		// TODO script Error?
 		return NULL;
 	}
 
