UMG - Typo fixed.

diff --git a/Engine/Source/Runtime/UMG/Private/Components/Button.cpp b/Engine/Source/Runtime/UMG/Private/Components/Button.cpp
--- a/Engine/Source/Runtime/UMG/Private/Components/Button.cpp
+++ b/Engine/Source/Runtime/UMG/Private/Components/Button.cpp
@@ -81,7 +81,7 @@ void UButton::SyncronizeProperties()
 
 void UButton::OnSlotAdded(UPanelSlot* Slot)
 {
-	// Add the child to the live canvas if it already exists
+	// Add the child to the live slot if it already exists
 	if ( MyButton.IsValid() )
 	{
 		MyButton->SetContent(Slot->Content ? Slot->Content->GetWidget() : SNullWidget::NullWidget);
