UMG - UCheckbox is now a ContentWidget like button and border, so it now supports a single child slot, just like the underlying SWidget.

diff --git a/Engine/Source/Runtime/UMG/Private/Components/CheckBox.cpp b/Engine/Source/Runtime/UMG/Private/Components/CheckBox.cpp
--- a/Engine/Source/Runtime/UMG/Private/Components/CheckBox.cpp
+++ b/Engine/Source/Runtime/UMG/Private/Components/CheckBox.cpp
@@ -23,13 +23,6 @@ UCheckBox::UCheckBox(const FPostConstructInitializeProperties& PCIP)
 
 TSharedRef<SWidget> UCheckBox::RebuildWidget()
 {
-	const FCheckBoxStyle* StylePtr = ( Style != NULL ) ? Style->GetStyle<FCheckBoxStyle>() : NULL;
-	if ( StylePtr == NULL )
-	{
-		SCheckBox::FArguments Defaults;
-		StylePtr = Defaults._Style;
-	}
-
 	TOptional<FSlateSound> OptionalCheckedSound;
 	if ( CheckedSound.GetResourceObject() )
 	{
@@ -49,7 +42,6 @@ TSharedRef<SWidget> UCheckBox::RebuildWidget()
 	}
 
 	MyCheckbox = SNew(SCheckBox)
-		.Style(StylePtr)
 		.OnCheckStateChanged( BIND_UOBJECT_DELEGATE(FOnCheckStateChanged, SlateOnCheckStateChangedCallback) )
 		.HAlign( HorizontalAlignment )
 		.Padding( Padding )
@@ -59,6 +51,11 @@ TSharedRef<SWidget> UCheckBox::RebuildWidget()
 		.UncheckedSoundOverride(OptionalUncheckedSound)
 		.HoveredSoundOverride(OptionalHoveredSound)
 		;
+
+	if ( GetChildrenCount() > 0 )
+	{
+		MyCheckbox->SetContent(GetContentSlot()->Content ? GetContentSlot()->Content->GetWidget() : SNullWidget::NullWidget);
+	}
 	
 	return MyCheckbox.ToSharedRef();
 }
@@ -67,6 +64,15 @@ void UCheckBox::SyncronizeProperties()
 {
 	Super::SyncronizeProperties();
 
+	const FCheckBoxStyle* StylePtr = ( Style != NULL ) ? Style->GetStyle<FCheckBoxStyle>() : NULL;
+	if ( StylePtr == NULL )
+	{
+		SCheckBox::FArguments Defaults;
+		StylePtr = Defaults._Style;
+	}
+
+	MyCheckbox->SetStyle(StylePtr);
+
 	MyCheckbox->SetIsChecked( OPTIONAL_BINDING(ESlateCheckBoxState::Type, CheckedState) );
 	
 	MyCheckbox->SetUncheckedImage(UncheckedImage ? &UncheckedImage->Brush : nullptr);
@@ -82,6 +88,24 @@ void UCheckBox::SyncronizeProperties()
 	MyCheckbox->SetUndeterminedPressedImage(UndeterminedPressedImage ? &UndeterminedPressedImage->Brush : nullptr);
 }
 
+void UCheckBox::OnSlotAdded(UPanelSlot* Slot)
+{
+	// Add the child to the live slot if it already exists
+	if ( MyCheckbox.IsValid() )
+	{
+		MyCheckbox->SetContent(Slot->Content ? Slot->Content->GetWidget() : SNullWidget::NullWidget);
+	}
+}
+
+void UCheckBox::OnSlotRemoved(UPanelSlot* Slot)
+{
+	// Remove the widget from the live slot if it exists.
+	if ( MyCheckbox.IsValid() )
+	{
+		MyCheckbox->SetContent(SNullWidget::NullWidget);
+	}
+}
+
 bool UCheckBox::IsPressed() const
 {
 	if ( MyCheckbox.IsValid() )
diff --git a/Engine/Source/Runtime/UMG/Public/Components/CheckBox.h b/Engine/Source/Runtime/UMG/Public/Components/CheckBox.h
--- a/Engine/Source/Runtime/UMG/Public/Components/CheckBox.h
+++ b/Engine/Source/Runtime/UMG/Public/Components/CheckBox.h
@@ -8,7 +8,7 @@ DECLARE_DYNAMIC_MULTICAST_DELEGATE_OneParam( FOnCheckBoxComponentStateChanged, b
 
 /** Check box widget */
 UCLASS(meta=( Category="Common" ), ClassGroup=UserInterface)
-class UMG_API UCheckBox : public UWidget
+class UMG_API UCheckBox : public UContentWidget
 {
 	GENERATED_UCLASS_BODY()
 
@@ -123,6 +123,13 @@ protected:
 	virtual const FSlateBrush* GetEditorIcon() override;
 #endif
 
+protected:
+
+	// UPanelWidget
+	virtual void OnSlotAdded(UPanelSlot* Slot) override;
+	virtual void OnSlotRemoved(UPanelSlot* Slot) override;
+	// End UPanelWidget
+
 protected:
 	// UWidget interface
 	virtual TSharedRef<SWidget> RebuildWidget() override;
