UMG - Renaming EditableTextBlock to EditableTextBox.  Making UMG more stable when widgets fail to be loaded in the hierarchy, like when a Widget BP is removed that was referenced in the UI.  The designer still needs to load the asset as best it can.  Empty slots are now purged on load.

diff --git a/Engine/Source/Runtime/UMG/Private/Components/EditableTextBlock.cpp b/Engine/Source/Runtime/UMG/Private/Components/EditableTextBlock.cpp
deleted file mode 100644
diff --git a/Engine/Source/Runtime/UMG/Private/Components/EditableTextBox.cpp b/Engine/Source/Runtime/UMG/Private/Components/EditableTextBox.cpp
new file mode 100644
--- /dev/null
+++ b/Engine/Source/Runtime/UMG/Private/Components/EditableTextBox.cpp
@@ -0,0 +1,103 @@
+// Copyright 1998-2014 Epic Games, Inc. All Rights Reserved.
+
+#include "UMGPrivatePCH.h"
+
+#define LOCTEXT_NAMESPACE "UMG"
+
+/////////////////////////////////////////////////////
+// UEditableTextBox
+
+UEditableTextBox::UEditableTextBox(const FPostConstructInitializeProperties& PCIP)
+	: Super(PCIP)
+{
+	ForegroundColor = FLinearColor::Black;
+	BackgroundColor = FLinearColor::White;
+	ReadOnlyForegroundColor = FLinearColor::Black;
+
+	// HACK Special font initialization hack since there are no font assets yet for slate.
+	Font = FSlateFontInfo(TEXT("Slate/Fonts/Roboto-Bold.ttf"), 12);
+
+	// Grab other defaults from slate arguments.
+	SEditableTextBox::FArguments Defaults;
+	IsReadOnly = Defaults._IsReadOnly.Get();
+	IsPassword = Defaults._IsReadOnly.Get();
+	MinimumDesiredWidth = Defaults._MinDesiredWidth.Get();
+	Padding = Defaults._Padding.Get();
+	IsCaretMovedWhenGainFocus = Defaults._IsCaretMovedWhenGainFocus.Get();
+	SelectAllTextWhenFocused = Defaults._SelectAllTextWhenFocused.Get();
+	RevertTextOnEscape = Defaults._RevertTextOnEscape.Get();
+	ClearKeyboardFocusOnCommit = Defaults._ClearKeyboardFocusOnCommit.Get();
+	SelectAllTextOnCommit = Defaults._SelectAllTextOnCommit.Get();
+}
+
+TSharedRef<SWidget> UEditableTextBox::RebuildWidget()
+{
+	FString FontPath = FPaths::EngineContentDir() / Font.FontName.ToString();
+
+	MyEditableTextBlock = SNew(SEditableTextBox)
+		.Font(FSlateFontInfo(FontPath, Font.Size))
+		.ForegroundColor(ForegroundColor)
+		.BackgroundColor(BackgroundColor)
+		.ReadOnlyForegroundColor(ReadOnlyForegroundColor)
+		.MinDesiredWidth(MinimumDesiredWidth)
+		.Padding(Padding)
+		.IsCaretMovedWhenGainFocus(IsCaretMovedWhenGainFocus)
+		.SelectAllTextWhenFocused(SelectAllTextWhenFocused)
+		.RevertTextOnEscape(RevertTextOnEscape)
+		.ClearKeyboardFocusOnCommit(ClearKeyboardFocusOnCommit)
+		.SelectAllTextOnCommit(SelectAllTextOnCommit)
+		.OnTextChanged(BIND_UOBJECT_DELEGATE(FOnTextChanged, SlateOnTextChanged))
+		;
+
+	return MyEditableTextBlock.ToSharedRef();
+}
+
+void UEditableTextBox::SyncronizeProperties()
+{
+	Super::SyncronizeProperties();
+
+	MyEditableTextBlock->SetText(Text);
+	MyEditableTextBlock->SetHintText(HintText);
+	MyEditableTextBlock->SetIsReadOnly(IsReadOnly);
+	MyEditableTextBlock->SetIsPassword(IsPassword);
+//	MyEditableTextBlock->SetColorAndOpacity(ColorAndOpacity);
+
+	// TODO UMG Complete making all properties settable on SEditableTextBox
+}
+
+FText UEditableTextBox::GetText() const
+{
+	if ( MyEditableTextBlock.IsValid() )
+	{
+		return MyEditableTextBlock->GetText();
+	}
+
+	return Text;
+}
+
+void UEditableTextBox::SetText(FText InText)
+{
+	Text = InText;
+	if ( MyEditableTextBlock.IsValid() )
+	{
+		MyEditableTextBlock->SetText(Text);
+	}
+}
+
+void UEditableTextBox::SlateOnTextChanged(const FText& Text)
+{
+	OnTextChanged.Broadcast(Text);
+}
+
+#if WITH_EDITOR
+
+const FSlateBrush* UEditableTextBox::GetEditorIcon()
+{
+	return FUMGStyle::Get().GetBrush("Widget.EditableTextBlock");
+}
+
+#endif
+
+/////////////////////////////////////////////////////
+
+#undef LOCTEXT_NAMESPACE
diff --git a/Engine/Source/Runtime/UMG/Private/Components/PanelWidget.cpp b/Engine/Source/Runtime/UMG/Private/Components/PanelWidget.cpp
--- a/Engine/Source/Runtime/UMG/Private/Components/PanelWidget.cpp
+++ b/Engine/Source/Runtime/UMG/Private/Components/PanelWidget.cpp
@@ -123,6 +123,21 @@ bool UPanelWidget::RemoveChild(UWidget* Content)
 	return false;
 }
 
+void UPanelWidget::PostLoad()
+{
+	Super::PostLoad();
+
+	for ( int32 SlotIndex = 0; SlotIndex < Slots.Num(); SlotIndex++ )
+	{
+		// Remove any slots where their content is null, we don't support content-less slots.
+		if ( Slots[SlotIndex]->Content == NULL )
+		{
+			Slots.RemoveAt(SlotIndex);
+			SlotIndex--;
+		}
+	}
+}
+
 /////////////////////////////////////////////////////
 
 #undef LOCTEXT_NAMESPACE
diff --git a/Engine/Source/Runtime/UMG/Private/WidgetTree.cpp b/Engine/Source/Runtime/UMG/Private/WidgetTree.cpp
--- a/Engine/Source/Runtime/UMG/Private/WidgetTree.cpp
+++ b/Engine/Source/Runtime/UMG/Private/WidgetTree.cpp
@@ -86,10 +86,12 @@ void UWidgetTree::GetChildWidgets(UWidget* Parent, TArray<UWidget*>& Widgets) co
 	{
 		for ( int32 ChildIndex = 0; ChildIndex < PanelParent->GetChildrenCount(); ChildIndex++ )
 		{
-			UWidget* ChildWidget = PanelParent->GetChildAt(ChildIndex);
-			Widgets.Add(ChildWidget);
+			if ( UWidget* ChildWidget = PanelParent->GetChildAt(ChildIndex) )
+			{
+				Widgets.Add(ChildWidget);
 
-			GetChildWidgets(ChildWidget, Widgets);
+				GetChildWidgets(ChildWidget, Widgets);
+			}
 		}
 	}
 }
diff --git a/Engine/Source/Runtime/UMG/Public/Components/EditableTextBlock.h b/Engine/Source/Runtime/UMG/Public/Components/EditableTextBlock.h
deleted file mode 100644
diff --git a/Engine/Source/Runtime/UMG/Public/Components/EditableTextBox.h b/Engine/Source/Runtime/UMG/Public/Components/EditableTextBox.h
new file mode 100644
--- /dev/null
+++ b/Engine/Source/Runtime/UMG/Public/Components/EditableTextBox.h
@@ -0,0 +1,114 @@
+// Copyright 1998-2014 Epic Games, Inc. All Rights Reserved.
+
+#pragma once
+
+#include "EditableTextBox.generated.h"
+
+DECLARE_DYNAMIC_MULTICAST_DELEGATE_OneParam(FOnEditableTextBlock_TextChangedEvent, const FText&, Text);
+//DECLARE_DYNAMIC_MULTICAST_DELEGATE_TwoParams(FOnEditableTextBlock_TextCommittedEvent, const FText&, Text, ETextCommit::Type, Type);
+
+/** Editable text box widget */
+UCLASS(meta=( Category="Common" ), ClassGroup=UserInterface)
+class UMG_API UEditableTextBox : public UWidget
+{
+	GENERATED_UCLASS_BODY()
+
+protected:
+
+	/** The text content for this editable text box widget */
+	UPROPERTY(EditDefaultsOnly, Category=Content)
+	FText Text;
+
+	/** Hint text that appears when there is no text in the text box */
+	UPROPERTY(EditDefaultsOnly, Category=Content)
+	FText HintText;
+
+	/** Font color and opacity (overrides Style) */
+	UPROPERTY(EditDefaultsOnly, Category=Appearance)
+	FSlateFontInfo Font;
+
+	/** Text color and opacity (overrides Style) */
+	UPROPERTY(EditDefaultsOnly, Category=Appearance)
+	FSlateColor ForegroundColor;
+
+	/** The color of the background/border around the editable text (overrides Style) */
+	UPROPERTY(EditDefaultsOnly, Category=Appearance)
+	FSlateColor BackgroundColor;
+
+	/** Text color and opacity when read-only (overrides Style) */
+	UPROPERTY(EditDefaultsOnly, Category=Appearance)
+	FSlateColor ReadOnlyForegroundColor;
+
+	/** Sets whether this text box can actually be modified interactively by the user */
+	UPROPERTY(EditDefaultsOnly, Category=Appearance)
+	bool IsReadOnly;
+
+	/** Sets whether this text box is for storing a password */
+	UPROPERTY(EditDefaultsOnly, Category=Appearance)
+	bool IsPassword;
+
+	/** Minimum width that a text block should be */
+	UPROPERTY(EditDefaultsOnly, Category=Appearance)
+	float MinimumDesiredWidth;
+
+	/** Padding between the box/border and the text widget inside (overrides Style) */
+	UPROPERTY(EditDefaultsOnly, Category=Appearance)
+	FMargin Padding;
+
+	/** Workaround as we loose focus when the auto completion closes. */
+	UPROPERTY(EditDefaultsOnly, Category=Behavior, AdvancedDisplay)
+	bool IsCaretMovedWhenGainFocus;
+
+	/** Whether to select all text when the user clicks to give focus on the widget */
+	UPROPERTY(EditDefaultsOnly, Category=Behavior, AdvancedDisplay)
+	bool SelectAllTextWhenFocused;
+
+	/** Whether to allow the user to back out of changes when they press the escape key */
+	UPROPERTY(EditDefaultsOnly, Category=Behavior, AdvancedDisplay)
+	bool RevertTextOnEscape;
+
+	/** Whether to clear keyboard focus when pressing enter to commit changes */
+	UPROPERTY(EditDefaultsOnly, Category=Behavior, AdvancedDisplay)
+	bool ClearKeyboardFocusOnCommit;
+
+	/** Whether to select all text when pressing enter to commit changes */
+	UPROPERTY(EditDefaultsOnly, Category=Behavior, AdvancedDisplay)
+	bool SelectAllTextOnCommit;
+
+	/** Called whenever the text is changed interactively by the user */
+	UPROPERTY(BlueprintAssignable)
+	FOnEditableTextBlock_TextChangedEvent OnTextChanged;
+
+	/** Called whenever the text is committed.  This happens when the user presses enter or the text box loses focus. *///
+	//UPROPERTY(BlueprintAssignable)
+	//FOnTextCommittedEvent OnTextCommitted;
+
+	/** Provide a alternative mechanism for error reporting. */
+	//SLATE_ARGUMENT(TSharedPtr<class IErrorReportingWidget>, ErrorReporting)
+
+	/**  */
+	UFUNCTION(BlueprintCallable, Category="Widget")
+	FText GetText() const;
+
+	/**  */
+	UFUNCTION(BlueprintCallable, Category="Widget")
+	void SetText(FText InText);
+	
+	// UWidget interface
+	virtual void SyncronizeProperties() override;
+	// End of UWidget interface
+
+#if WITH_EDITOR
+	virtual const FSlateBrush* GetEditorIcon() override;
+#endif
+
+protected:
+	// UWidget interface
+	virtual TSharedRef<SWidget> RebuildWidget() override;
+	// End of UWidget
+
+	void SlateOnTextChanged(const FText& Text);
+
+protected:
+	TSharedPtr<SEditableTextBox> MyEditableTextBlock;
+};
diff --git a/Engine/Source/Runtime/UMG/Public/Components/PanelWidget.h b/Engine/Source/Runtime/UMG/Public/Components/PanelWidget.h
--- a/Engine/Source/Runtime/UMG/Public/Components/PanelWidget.h
+++ b/Engine/Source/Runtime/UMG/Public/Components/PanelWidget.h
@@ -53,11 +53,18 @@ public:
 		for ( UPanelSlot* Slot : Slots )
 		{
 			Slot->Parent = this;
-			Slot->Content->Slot = Slot;
+			if ( Slot->Content )
+			{
+				Slot->Content->Slot = Slot;
+			}
 		}
 	}
 #endif
 
+	// Begin UObject
+	virtual void PostLoad() override;
+	// End UObject
+
 protected:
 
 	virtual UClass* GetSlotClass() const
diff --git a/Engine/Source/Runtime/UMG/Public/UMG.h b/Engine/Source/Runtime/UMG/Public/UMG.h
--- a/Engine/Source/Runtime/UMG/Public/UMG.h
+++ b/Engine/Source/Runtime/UMG/Public/UMG.h
@@ -59,7 +59,7 @@
 #include "Button.h"
 #include "CheckBox.h"
 #include "EditableText.h"
-#include "EditableTextBlock.h"
+#include "EditableTextBox.h"
 #include "Image.h"
 #include "Slider.h"
 #include "ProgressBar.h"
