UMG - User widgets have a Tick now!

diff --git a/Engine/Source/Runtime/UMG/Private/Slate/SObjectWidget.cpp b/Engine/Source/Runtime/UMG/Private/Slate/SObjectWidget.cpp
--- a/Engine/Source/Runtime/UMG/Private/Slate/SObjectWidget.cpp
+++ b/Engine/Source/Runtime/UMG/Private/Slate/SObjectWidget.cpp
@@ -14,15 +14,23 @@ void SObjectWidget::Construct(const FArguments& InArgs, UUserWidget* InWidgetObj
 
 SObjectWidget::~SObjectWidget(void)
 {
-	printf("");
 }
 
 void SObjectWidget::AddReferencedObjects(FReferenceCollector& Collector)
 {
-	printf("");
 	Collector.AddReferencedObject(WidgetObject);
 }
 
+void SObjectWidget::Tick(const FGeometry& AllottedGeometry, const double InCurrentTime, const float InDeltaTime)
+{
+	return WidgetObject->Tick(AllottedGeometry, InCurrentTime, InDeltaTime);
+}
+
+//bool OnHitTest(const FGeometry& MyGeometry, FVector2D InAbsoluteCursorPosition)
+//{
+
+//}
+
 int32 SObjectWidget::OnPaint(const FGeometry& AllottedGeometry, const FSlateRect& MyClippingRect, FSlateWindowElementList& OutDrawElements, int32 LayerId, const FWidgetStyle& InWidgetStyle, bool bParentEnabled) const
 {
 	int32 MaxLayer = SCompoundWidget::OnPaint(AllottedGeometry, MyClippingRect, OutDrawElements, LayerId, InWidgetStyle, bParentEnabled);
diff --git a/Engine/Source/Runtime/UMG/Public/Blueprint/UserWidget.h b/Engine/Source/Runtime/UMG/Public/Blueprint/UserWidget.h
--- a/Engine/Source/Runtime/UMG/Public/Blueprint/UserWidget.h
+++ b/Engine/Source/Runtime/UMG/Public/Blueprint/UserWidget.h
@@ -115,6 +115,11 @@ public:
 	UFUNCTION(BlueprintPure, Category="Appearance")
 	TEnumAsByte<ESlateVisibility::Type> GetVisiblity();
 
+	UFUNCTION(BlueprintImplementableEvent, Category="User Interface")
+	void Tick(FGeometry MyGeometry, double InCurrentTime, float InDeltaTime);
+
+	//TODO UMG HitTest
+
 	UFUNCTION(BlueprintImplementableEvent, Category="User Interface")
 	void OnPaint(UPARAM(ref) FPaintContext& Context) const;
 
diff --git a/Engine/Source/Runtime/UMG/Public/Slate/SObjectWidget.h b/Engine/Source/Runtime/UMG/Public/Slate/SObjectWidget.h
--- a/Engine/Source/Runtime/UMG/Public/Slate/SObjectWidget.h
+++ b/Engine/Source/Runtime/UMG/Public/Slate/SObjectWidget.h
@@ -23,6 +23,8 @@ class SObjectWidget : public SCompoundWidget, public FGCObject
 	virtual void AddReferencedObjects(FReferenceCollector& Collector) override;
 	// End of FGCObject interface
 
+	virtual void Tick(const FGeometry& AllottedGeometry, const double InCurrentTime, const float InDeltaTime) override;
+	//virtual bool OnHitTest(const FGeometry& MyGeometry, FVector2D InAbsoluteCursorPosition) override;
 	virtual int32 OnPaint(const FGeometry& AllottedGeometry, const FSlateRect& MyClippingRect, FSlateWindowElementList& OutDrawElements, int32 LayerId, const FWidgetStyle& InWidgetStyle, bool bParentEnabled) const override;
 
 	virtual FReply OnKeyboardFocusReceived(const FGeometry& MyGeometry, const FKeyboardFocusEvent& InKeyboardFocusEvent) override;
