UMG - Fixing a bug with the widget Tick.  No double support in blueprints.

diff --git a/Engine/Source/Runtime/UMG/Private/Slate/SObjectWidget.cpp b/Engine/Source/Runtime/UMG/Private/Slate/SObjectWidget.cpp
--- a/Engine/Source/Runtime/UMG/Private/Slate/SObjectWidget.cpp
+++ b/Engine/Source/Runtime/UMG/Private/Slate/SObjectWidget.cpp
@@ -23,7 +23,7 @@ void SObjectWidget::AddReferencedObjects(FReferenceCollector& Collector)
 
 void SObjectWidget::Tick(const FGeometry& AllottedGeometry, const double InCurrentTime, const float InDeltaTime)
 {
-	return WidgetObject->Tick(AllottedGeometry, InCurrentTime, InDeltaTime);
+	return WidgetObject->Tick(AllottedGeometry, InDeltaTime);
 }
 
 //bool OnHitTest(const FGeometry& MyGeometry, FVector2D InAbsoluteCursorPosition)
diff --git a/Engine/Source/Runtime/UMG/Public/Blueprint/UserWidget.h b/Engine/Source/Runtime/UMG/Public/Blueprint/UserWidget.h
--- a/Engine/Source/Runtime/UMG/Public/Blueprint/UserWidget.h
+++ b/Engine/Source/Runtime/UMG/Public/Blueprint/UserWidget.h
@@ -116,7 +116,7 @@ public:
 	TEnumAsByte<ESlateVisibility::Type> GetVisiblity();
 
 	UFUNCTION(BlueprintImplementableEvent, Category="User Interface")
-	void Tick(FGeometry MyGeometry, double InCurrentTime, float InDeltaTime);
+	void Tick(FGeometry MyGeometry, float InDeltaTime);
 
 	//TODO UMG HitTest
 
