UMG - Adding a prototype wrapper of the MenuAnchor.  It still has several complexitites that need to be simplified.

diff --git a/Engine/Source/Runtime/Slate/Private/Widgets/Input/SMenuAnchor.cpp b/Engine/Source/Runtime/Slate/Private/Widgets/Input/SMenuAnchor.cpp
--- a/Engine/Source/Runtime/Slate/Private/Widgets/Input/SMenuAnchor.cpp
+++ b/Engine/Source/Runtime/Slate/Private/Widgets/Input/SMenuAnchor.cpp
@@ -120,6 +120,20 @@ void SMenuAnchor::Tick( const FGeometry& AllottedGeometry, const double InCurren
 	bDismissedThisTick = false;
 }
 
+void SMenuAnchor::SetContent(TSharedRef<SWidget> InContent)
+{
+	ChildSlot
+		.Padding(0)
+		[
+			InContent
+		];
+}
+
+void SMenuAnchor::SetMenuContent(TSharedRef<SWidget> InMenuContent)
+{
+	MenuContent = InMenuContent;
+}
+
 /**
  * Open or close the popup
  *
diff --git a/Engine/Source/Runtime/Slate/Public/Widgets/Input/SMenuAnchor.h b/Engine/Source/Runtime/Slate/Public/Widgets/Input/SMenuAnchor.h
--- a/Engine/Source/Runtime/Slate/Public/Widgets/Input/SMenuAnchor.h
+++ b/Engine/Source/Runtime/Slate/Public/Widgets/Input/SMenuAnchor.h
@@ -53,6 +53,11 @@ public:
 	virtual void Tick( const FGeometry& AllottedGeometry, const double InCurrentTime, const float InDeltaTime ) override;
 	// End of SWidget interface
 
+	/** See Content Slot attribute */
+	void SetContent(TSharedRef<SWidget> InContent);
+
+	/** See MenuContent attribute */
+	void SetMenuContent(TSharedRef<SWidget> InMenuContent);
 	
 	/**
 	 * Open or close the popup
diff --git a/Engine/Source/Runtime/UMG/Private/Components/MenuAnchor.cpp b/Engine/Source/Runtime/UMG/Private/Components/MenuAnchor.cpp
new file mode 100644
--- /dev/null
+++ b/Engine/Source/Runtime/UMG/Private/Components/MenuAnchor.cpp
@@ -0,0 +1,108 @@
+// Copyright 1998-2014 Epic Games, Inc. All Rights Reserved.
+
+#include "UMGPrivatePCH.h"
+
+#define LOCTEXT_NAMESPACE "UMG"
+
+/////////////////////////////////////////////////////
+// UMenuAnchor
+
+UMenuAnchor::UMenuAnchor(const FPostConstructInitializeProperties& PCIP)
+	: Super(PCIP)
+{
+	Placement = MenuPlacement_ComboBox;
+}
+
+TSharedRef<SWidget> UMenuAnchor::RebuildWidget()
+{
+	MyMenuAnchor = SNew(SMenuAnchor)
+		.Method(SMenuAnchor::UseCurrentWindow)
+		.Placement(Placement)
+		.OnGetMenuContent(BIND_UOBJECT_DELEGATE(FOnGetContent, HandleGetMenuContent));
+
+	if ( GetChildrenCount() > 0 )
+	{
+		MyMenuAnchor->SetContent(GetContentSlot()->Content ? GetContentSlot()->Content->GetWidget() : SNullWidget::NullWidget);
+	}
+	
+	return BuildDesignTimeWidget( MyMenuAnchor.ToSharedRef() );
+}
+
+void UMenuAnchor::SyncronizeProperties()
+{
+	Super::SyncronizeProperties();
+}
+
+void UMenuAnchor::OnSlotAdded(UPanelSlot* Slot)
+{
+	// Add the child to the live slot if it already exists
+	if ( MyMenuAnchor.IsValid() )
+	{
+		MyMenuAnchor->SetContent(Slot->Content ? Slot->Content->GetWidget() : SNullWidget::NullWidget);
+	}
+}
+
+void UMenuAnchor::OnSlotRemoved(UPanelSlot* Slot)
+{
+	// Remove the widget from the live slot if it exists.
+	if ( MyMenuAnchor.IsValid() )
+	{
+		MyMenuAnchor->SetContent(SNullWidget::NullWidget);
+	}
+}
+
+TSharedRef<SWidget> UMenuAnchor::HandleGetMenuContent()
+{
+	TSharedPtr<SWidget> SlateMenuWidget;
+
+	if ( OnGetMenuContentEvent.IsBound() )
+	{
+		UWidget* MenuWidget = OnGetMenuContentEvent.Execute();
+		if ( MenuWidget )
+		{
+			SlateMenuWidget = MenuWidget->GetWidget();
+		}
+	}
+
+	if ( MenuClass != NULL && !MenuClass->HasAnyClassFlags(CLASS_Abstract) )
+	{
+		UWidget* MenuWidget = (UWidget*)ConstructObject<UWidget>(MenuClass, GetOuter());
+		if ( MenuWidget )
+		{
+			SlateMenuWidget = MenuWidget->GetWidget();
+		}
+	}
+
+	return SlateMenuWidget.IsValid() ? SlateMenuWidget.ToSharedRef() : SNullWidget::NullWidget;
+}
+
+void UMenuAnchor::SetIsOpen(bool InIsOpen, bool bFocusMenu)
+{
+	if ( MyMenuAnchor.IsValid() )
+	{
+		MyMenuAnchor->SetIsOpen(InIsOpen, bFocusMenu);
+	}
+}
+
+bool UMenuAnchor::IsOpen() const
+{
+	if ( MyMenuAnchor.IsValid() )
+	{
+		return MyMenuAnchor->IsOpen();
+	}
+
+	return false;
+}
+
+#if WITH_EDITOR
+
+const FSlateBrush* UMenuAnchor::GetEditorIcon()
+{
+	return FUMGStyle::Get().GetBrush("Widget.MenuAnchor");
+}
+
+#endif
+
+/////////////////////////////////////////////////////
+
+#undef LOCTEXT_NAMESPACE
diff --git a/Engine/Source/Runtime/UMG/Private/Components/Widget.cpp b/Engine/Source/Runtime/UMG/Private/Components/Widget.cpp
--- a/Engine/Source/Runtime/UMG/Private/Components/Widget.cpp
+++ b/Engine/Source/Runtime/UMG/Private/Components/Widget.cpp
@@ -68,6 +68,14 @@ bool UWidget::IsHovered() const
 	return false;
 }
 
+void UWidget::ForceLayoutPrepass()
+{
+	if ( MyWidget.IsValid() )
+	{
+		MyWidget.Pin()->SlatePrepass();
+	}
+}
+
 UPanelWidget* UWidget::GetParent() const
 {
 	if ( Slot )
diff --git a/Engine/Source/Runtime/UMG/Public/Components/MenuAnchor.h b/Engine/Source/Runtime/UMG/Public/Components/MenuAnchor.h
new file mode 100644
--- /dev/null
+++ b/Engine/Source/Runtime/UMG/Public/Components/MenuAnchor.h
@@ -0,0 +1,77 @@
+// Copyright 1998-2014 Epic Games, Inc. All Rights Reserved.
+
+#pragma once
+
+#include "MenuAnchor.generated.h"
+
+/**
+ * The Menu Anchor allows you to specify an location that a popup menu should be anchored to, and should be
+ * summoned from.
+ */
+UCLASS(meta=( Category="Popup" ), ClassGroup=UserInterface)
+class UMG_API UMenuAnchor : public UContentWidget
+{
+	GENERATED_UCLASS_BODY()
+
+public:
+
+	/** The widget class to spawn when the menu is required.  Creates the object freshly each time. */
+	UPROPERTY(EditDefaultsOnly, Category=Menu)
+	TSubclassOf<class UUserWidget> MenuClass;
+
+	/** Called when the menu content is requested to allow dynamic decision over what to display */
+	UPROPERTY(EditDefaultsOnly, Category=Events)
+	FGetContent OnGetMenuContentEvent;
+	
+	/** The placement location of the summoned widget. */
+	UPROPERTY(EditDefaultsOnly, Category=Menu)
+	TEnumAsByte<EMenuPlacement> Placement;
+
+	/**
+	 * Open or close the popup
+	 *
+	 * @param InIsOpen    If true, open the popup. Otherwise close it.
+	 * @param bFocusMenu  Should we focus the popup as soon as it opens?
+	 */
+	UFUNCTION(BlueprintCallable, Category="Widget")
+	void SetIsOpen(bool InIsOpen, bool bFocusMenu = true);
+
+	/** @return true if the popup is open; false otherwise. */
+	UFUNCTION(BlueprintCallable, Category="Widget")
+	bool IsOpen() const;
+
+	///** @return true if we should open the menu due to a click. Sometimes we should not, if
+	//the same MouseDownEvent that just closed the menu is about to re-open it because it happens to land on the button.*/
+	//bool ShouldOpenDueToClick() const;
+
+	///** @return The current menu position */
+	//FVector2D GetMenuPosition() const;
+
+	///** @return Whether this menu has open submenus */
+	//bool HasOpenSubMenus() const;
+
+	// UWidget interface
+	virtual void SyncronizeProperties() override;
+	// End of UWidget interface
+
+#if WITH_EDITOR
+	virtual const FSlateBrush* GetEditorIcon() override;
+#endif
+
+protected:
+
+	// UPanelWidget
+	virtual void OnSlotAdded(UPanelSlot* Slot) override;
+	virtual void OnSlotRemoved(UPanelSlot* Slot) override;
+	// End UPanelWidget
+
+	// UWidget interface
+	virtual TSharedRef<SWidget> RebuildWidget() override;
+	// End of UWidget interface
+
+protected:
+	TSharedRef<SWidget> HandleGetMenuContent();
+
+protected:
+	TSharedPtr<SMenuAnchor> MyMenuAnchor;
+};
diff --git a/Engine/Source/Runtime/UMG/Public/Components/Widget.h b/Engine/Source/Runtime/UMG/Public/Components/Widget.h
--- a/Engine/Source/Runtime/UMG/Public/Components/Widget.h
+++ b/Engine/Source/Runtime/UMG/Public/Components/Widget.h
@@ -36,6 +36,7 @@ public:
 	DECLARE_DYNAMIC_DELEGATE_RetVal(EMouseCursor::Type, FGetMouseCursor);
 	DECLARE_DYNAMIC_DELEGATE_RetVal(USlateBrushAsset*, FGetSlateBrushAsset);
 	DECLARE_DYNAMIC_DELEGATE_RetVal(ESlateCheckBoxState::Type, FGetCheckBoxState);
+	DECLARE_DYNAMIC_DELEGATE_RetVal(UWidget*, FGetContent);
 
 	DECLARE_DYNAMIC_DELEGATE_RetVal_OneParam(UWidget*, FGenerateWidgetUObject, UObject*, Item);
 
@@ -120,6 +121,10 @@ public:
 	UFUNCTION(BlueprintCallable, Category="Widget")
 	bool IsHovered() const;
 
+	/** Forces the underlying slate system to perform a pre-pass on the layout of the widget.  This is for advanced users. */
+	UFUNCTION(BlueprintCallable, Category="Widget")
+	void ForceLayoutPrepass();
+
 	/** Gets the parent widget */
 	UFUNCTION(BlueprintCallable, Category="Widget")
 	class UPanelWidget* GetParent() const;
diff --git a/Engine/Source/Runtime/UMG/Public/UMG.h b/Engine/Source/Runtime/UMG/Public/UMG.h
--- a/Engine/Source/Runtime/UMG/Public/UMG.h
+++ b/Engine/Source/Runtime/UMG/Public/UMG.h
@@ -66,6 +66,7 @@
 #include "ComboBox.h"
 
 #include "Spacer.h"
+#include "MenuAnchor.h"
 
 #include "Throbber.h"
 #include "CircularThrobber.h"
