Fix cooking issues with UMG.  UMG now works with cooked builds

diff --git a/Engine/Source/Runtime/Engine/Engine.Build.cs b/Engine/Source/Runtime/Engine/Engine.Build.cs
--- a/Engine/Source/Runtime/Engine/Engine.Build.cs
+++ b/Engine/Source/Runtime/Engine/Engine.Build.cs
@@ -73,6 +73,7 @@ public class Engine : ModuleRules
 			{
 				"MovieSceneCore",
 				"MovieSceneCoreTypes",
+				"UMG",
 				"HeadMountedDisplay",
 				"StreamingPauseRendering",
 			}
diff --git a/Engine/Source/Runtime/Launch/Private/LaunchEngineLoop.cpp b/Engine/Source/Runtime/Launch/Private/LaunchEngineLoop.cpp
--- a/Engine/Source/Runtime/Launch/Private/LaunchEngineLoop.cpp
+++ b/Engine/Source/Runtime/Launch/Private/LaunchEngineLoop.cpp
@@ -1686,6 +1686,8 @@ bool FEngineLoop::LoadStartupCoreModules()
 	if (!IsRunningDedicatedServer())
 	{
 		FModuleManager::Get().LoadModule("Slate");
+		// UMG must be loaded for runtime and cooking.
+		FModuleManager::Get().LoadModule("UMG");
 
 #if WITH_UNREAL_DEVELOPER_TOOLS
 		FModuleManager::Get().LoadModule("MessageLog");
diff --git a/Engine/Source/Runtime/UMG/Private/UMGModule.cpp b/Engine/Source/Runtime/UMG/Private/UMGModule.cpp
--- a/Engine/Source/Runtime/UMG/Private/UMGModule.cpp
+++ b/Engine/Source/Runtime/UMG/Private/UMGModule.cpp
@@ -16,9 +16,9 @@ public:
 	/** Called right after the module DLL has been loaded and the module object has been created */
 	virtual void StartupModule() override
 	{
+#if WITH_EDITOR	
 		FUMGStyle::Initialize();
 
-#if WITH_EDITOR
 		// This is done so that the compiler is available in non-cooked builds when the widget blueprint is 
 		// compiled again in the running game.
 		if ( FParse::Param(FCommandLine::Get(), TEXT("umg")) )
@@ -31,7 +31,9 @@ public:
 	/** Called before the module is unloaded, right before the module object is destroyed. */
 	virtual void ShutdownModule() override
 	{
+#if WITH_EDITOR
 		FUMGStyle::Shutdown();
+#endif
 	}
 };
 
diff --git a/Engine/Source/Runtime/UMG/Private/WidgetBlueprint.cpp b/Engine/Source/Runtime/UMG/Private/WidgetBlueprint.cpp
--- a/Engine/Source/Runtime/UMG/Private/WidgetBlueprint.cpp
+++ b/Engine/Source/Runtime/UMG/Private/WidgetBlueprint.cpp
@@ -5,6 +5,7 @@
 
 #define LOCTEXT_NAMESPACE "UMG"
 
+#if WITH_EDITOR
 FDelegateRuntimeBinding FDelegateEditorBinding::ToRuntimeBinding(UWidgetBlueprint* Blueprint) const
 {
 	FDelegateRuntimeBinding Binding;
@@ -14,6 +15,7 @@ FDelegateRuntimeBinding FDelegateEditorBinding::ToRuntimeBinding(UWidgetBlueprin
 
 	return Binding;
 }
+#endif
 
 /////////////////////////////////////////////////////
 // UWidgetBlueprint
@@ -49,7 +51,7 @@ UClass* UWidgetBlueprint::GetBlueprintClass() const
 {
 	return UWidgetBlueprintGeneratedClass::StaticClass();
 }
-#endif
+
 
 bool UWidgetBlueprint::ValidateGeneratedClass(const UClass* InClass)
 {
@@ -78,6 +80,7 @@ bool UWidgetBlueprint::ValidateGeneratedClass(const UClass* InClass)
 
 	return Result;
 }
+#endif
 
 void UWidgetBlueprint::PostLoadSubobjects(FObjectInstancingGraph* OuterInstanceGraph)
 {
diff --git a/Engine/Source/Runtime/UMG/Public/Blueprint/WidgetBlueprint.h b/Engine/Source/Runtime/UMG/Public/Blueprint/WidgetBlueprint.h
--- a/Engine/Source/Runtime/UMG/Public/Blueprint/WidgetBlueprint.h
+++ b/Engine/Source/Runtime/UMG/Public/Blueprint/WidgetBlueprint.h
@@ -32,7 +32,9 @@ struct UMG_API FDelegateEditorBinding
 		return ObjectName == Other.ObjectName && PropertyName == Other.PropertyName;
 	}
 
+#if WITH_EDITOR
 	FDelegateRuntimeBinding ToRuntimeBinding(class UWidgetBlueprint* Blueprint) const;
+#endif
 };
 
 class UMovieScene;
