UMG - Better control over mouse absorbing.  The upper level objects don't automatically absorb clicks, the user can just set a border that absorbs them on their UI if that's important to them.  Improved styling support in widgets, several more common widgets now support a style.

diff --git a/Engine/Source/Editor/UMGEditor/Private/WidgetBlueprintEditor.cpp b/Engine/Source/Editor/UMGEditor/Private/WidgetBlueprintEditor.cpp
--- a/Engine/Source/Editor/UMGEditor/Private/WidgetBlueprintEditor.cpp
+++ b/Engine/Source/Editor/UMGEditor/Private/WidgetBlueprintEditor.cpp
@@ -176,6 +176,13 @@ void FWidgetBlueprintEditor::Tick(float DeltaTime)
 			PreviewScene.GetWorld()->Tick(bIsRealTime ? LEVELTICK_ViewportsOnly : LEVELTICK_TimeOnly, DeltaTime);
 		}
 	}
+
+	// Note: The weak ptr can become stale if the actor is reinstanced due to a Blueprint change, etc. In that case we 
+	//       look to see if we can find the new instance in the preview world and then update the weak ptr.
+	if ( PreviewWidgetActorPtr.IsStale(true) )
+	{
+		UpdatePreview(GetWidgetBlueprintObj(), true);
+	}
 }
 
 static bool MigratePropertyValue(UObject* SourceObject, UObject* DestinationObject, UProperty* MemberProperty)
@@ -306,22 +313,6 @@ UWidgetBlueprint* FWidgetBlueprintEditor::GetWidgetBlueprintObj() const
 
 UUserWidget* FWidgetBlueprintEditor::GetPreview() const
 {
-	//// Note: The weak ptr can become stale if the actor is reinstanced due to a Blueprint change, etc. In that case we look to see if we can find the new instance in the preview world and then update the weak ptr.
-	//if ( PreviewWidgetActorPtr.IsStale(true) )
-	//{
-	//	UWorld* PreviewWorld = PreviewWidgetActorPtr->GetWorld();
-	//	for ( TActorIterator<AActor> It(PreviewWorld); It; ++It )
-	//	{
-	//		AActor* Actor = *It;
-	//		if ( !Actor->IsPendingKillPending()
-	//			&& Actor->GetClass()->ClassGeneratedBy == PreviewBlueprint )
-	//		{
-	//			PreviewActorPtr = Actor;
-	//			break;
-	//		}
-	//	}
-	//}
-
 	return PreviewWidgetActorPtr.Get();
 }
 
@@ -401,7 +392,6 @@ void FWidgetBlueprintEditor::UpdatePreview(UBlueprint* InBlueprint, bool bInForc
 	}
 
 	OnWidgetPreviewUpdated.Broadcast();
-
 }
 
 UMovieScene* FWidgetBlueprintEditor::GetDefaultMovieScene()
diff --git a/Engine/Source/Runtime/UMG/Private/Components/CanvasPanel.cpp b/Engine/Source/Runtime/UMG/Private/Components/CanvasPanel.cpp
--- a/Engine/Source/Runtime/UMG/Private/Components/CanvasPanel.cpp
+++ b/Engine/Source/Runtime/UMG/Private/Components/CanvasPanel.cpp
@@ -44,6 +44,9 @@ UCanvasPanel::UCanvasPanel(const FPostConstructInitializeProperties& PCIP)
 {
 	bIsVariable = false;
 	DesiredCanvasSize = FVector2D(128.0f, 128.0f);
+
+	SConstraintCanvas::FArguments Defaults;
+	Visiblity = UWidget::ConvertRuntimeToSerializedVisiblity(Defaults._Visibility.Get());
 }
 
 UClass* UCanvasPanel::GetSlotClass() const
diff --git a/Engine/Source/Runtime/UMG/Private/Components/CanvasPanelSlot.cpp b/Engine/Source/Runtime/UMG/Private/Components/CanvasPanelSlot.cpp
--- a/Engine/Source/Runtime/UMG/Private/Components/CanvasPanelSlot.cpp
+++ b/Engine/Source/Runtime/UMG/Private/Components/CanvasPanelSlot.cpp
@@ -12,17 +12,17 @@ UCanvasPanelSlot::UCanvasPanelSlot(const FPostConstructInitializeProperties& PCI
 	LayoutData.Offsets = FMargin(0, 0, 100, 30);
 	LayoutData.Anchors = FAnchors(0.0f, 0.0f);
 	LayoutData.Alignment = FVector2D(0.0f, 0.0f);
+	ZOrder = 0;
 }
 
 void UCanvasPanelSlot::BuildSlot(TSharedRef<SConstraintCanvas> Canvas)
 {
 	Slot = &Canvas->AddSlot()
-		.Offset(LayoutData.Offsets)
-		.Anchors(LayoutData.Anchors)
-		.Alignment(LayoutData.Alignment)
 		[
 			Content == NULL ? SNullWidget::NullWidget : Content->GetWidget()
 		];
+
+	SyncronizeProperties();
 }
 
 void UCanvasPanelSlot::SetDesiredPosition(FVector2D InPosition)
@@ -123,11 +123,21 @@ void UCanvasPanelSlot::SetAlignment(FVector2D InAlignment)
 	}
 }
 
+void UCanvasPanelSlot::SetZOrder(int32 InZOrder)
+{
+	ZOrder = InZOrder;
+	if ( Slot )
+	{
+		Slot->ZOrder(InZOrder);
+	}
+}
+
 void UCanvasPanelSlot::SyncronizeProperties()
 {
 	SetOffset(LayoutData.Offsets);
 	SetAnchors(LayoutData.Anchors);
 	SetAlignment(LayoutData.Alignment);
+	SetZOrder(ZOrder);
 }
 
 #if WITH_EDITOR
diff --git a/Engine/Source/Runtime/UMG/Private/Components/EditableTextBox.cpp b/Engine/Source/Runtime/UMG/Private/Components/EditableTextBox.cpp
--- a/Engine/Source/Runtime/UMG/Private/Components/EditableTextBox.cpp
+++ b/Engine/Source/Runtime/UMG/Private/Components/EditableTextBox.cpp
@@ -56,6 +56,12 @@ void UEditableTextBox::SyncronizeProperties()
 {
 	Super::SyncronizeProperties();
 
+	const FEditableTextBoxStyle* StylePtr = ( Style != NULL ) ? Style->GetStyle<FEditableTextBoxStyle>() : NULL;
+	if ( StylePtr )
+	{
+		MyEditableTextBlock->SetStyle(StylePtr);
+	}
+
 	MyEditableTextBlock->SetText(Text);
 	MyEditableTextBlock->SetHintText(HintText);
 	MyEditableTextBlock->SetIsReadOnly(IsReadOnly);
@@ -84,6 +90,14 @@ void UEditableTextBox::SetText(FText InText)
 	}
 }
 
+void UEditableTextBox::SetError(FText InError)
+{
+	if ( MyEditableTextBlock.IsValid() )
+	{
+		MyEditableTextBlock->SetError(InError);
+	}
+}
+
 void UEditableTextBox::SlateOnTextChanged(const FText& Text)
 {
 	OnTextChanged.Broadcast(Text);
@@ -93,7 +107,7 @@ void UEditableTextBox::SlateOnTextChanged(const FText& Text)
 
 const FSlateBrush* UEditableTextBox::GetEditorIcon()
 {
-	return FUMGStyle::Get().GetBrush("Widget.EditableTextBlock");
+	return FUMGStyle::Get().GetBrush("Widget.EditableTextBox");
 }
 
 #endif
diff --git a/Engine/Source/Runtime/UMG/Private/Components/HorizontalBox.cpp b/Engine/Source/Runtime/UMG/Private/Components/HorizontalBox.cpp
--- a/Engine/Source/Runtime/UMG/Private/Components/HorizontalBox.cpp
+++ b/Engine/Source/Runtime/UMG/Private/Components/HorizontalBox.cpp
@@ -9,6 +9,9 @@ UHorizontalBox::UHorizontalBox(const FPostConstructInitializeProperties& PCIP)
 	: Super(PCIP)
 {
 	bIsVariable = false;
+
+	SHorizontalBox::FArguments Defaults;
+	Visiblity = UWidget::ConvertRuntimeToSerializedVisiblity(Defaults._Visibility.Get());
 }
 
 UClass* UHorizontalBox::GetSlotClass() const
diff --git a/Engine/Source/Runtime/UMG/Private/Components/Overlay.cpp b/Engine/Source/Runtime/UMG/Private/Components/Overlay.cpp
--- a/Engine/Source/Runtime/UMG/Private/Components/Overlay.cpp
+++ b/Engine/Source/Runtime/UMG/Private/Components/Overlay.cpp
@@ -9,6 +9,9 @@ UOverlay::UOverlay(const FPostConstructInitializeProperties& PCIP)
 	: Super(PCIP)
 {
 	bIsVariable = false;
+
+	SOverlay::FArguments Defaults;
+	Visiblity = UWidget::ConvertRuntimeToSerializedVisiblity(Defaults._Visibility.Get());
 }
 
 UClass* UOverlay::GetSlotClass() const
diff --git a/Engine/Source/Runtime/UMG/Private/Components/ProcessBar.cpp b/Engine/Source/Runtime/UMG/Private/Components/ProcessBar.cpp
--- a/Engine/Source/Runtime/UMG/Private/Components/ProcessBar.cpp
+++ b/Engine/Source/Runtime/UMG/Private/Components/ProcessBar.cpp
@@ -11,6 +11,7 @@ UProgressBar::UProgressBar(const FPostConstructInitializeProperties& PCIP)
 	: Super(PCIP)
 {
 	BarFillType = EProgressBarFillType::LeftToRight;
+	bIsMarquee = false;
 	Percent = 0;
 	FillColorAndOpacity = FLinearColor::White;
 	BorderPadding = FVector2D(1, 0);
@@ -30,8 +31,10 @@ void UProgressBar::SyncronizeProperties()
 	{
 		MyProgressBar->SetStyle(StylePtr);
 	}
-	
-	MyProgressBar->SetPercent(bOverride_Percent ? Percent : TOptional<float>());
+
+	TAttribute< TOptional<float> > PercentBinding = OPTIONAL_BINDING_CONVERT(float, Percent, TOptional<float>, ConvertFloatToOptionalFloat);
+
+	MyProgressBar->SetPercent(bIsMarquee ? TOptional<float>() : PercentBinding);
 	
 	MyProgressBar->SetBarFillType(BarFillType);
 	MyProgressBar->SetFillColorAndOpacity(FillColorAndOpacity);
diff --git a/Engine/Source/Runtime/UMG/Private/Components/ScrollBox.cpp b/Engine/Source/Runtime/UMG/Private/Components/ScrollBox.cpp
--- a/Engine/Source/Runtime/UMG/Private/Components/ScrollBox.cpp
+++ b/Engine/Source/Runtime/UMG/Private/Components/ScrollBox.cpp
@@ -11,6 +11,9 @@ UScrollBox::UScrollBox(const FPostConstructInitializeProperties& PCIP)
 	: Super(PCIP)
 {
 	bIsVariable = false;
+
+	SScrollBox::FArguments Defaults;
+	Visiblity = UWidget::ConvertRuntimeToSerializedVisiblity(Defaults._Visibility.Get());
 }
 
 UClass* UScrollBox::GetSlotClass() const
@@ -38,7 +41,23 @@ void UScrollBox::OnSlotRemoved(UPanelSlot* Slot)
 
 TSharedRef<SWidget> UScrollBox::RebuildWidget()
 {
-	MyScrollBox = SNew(SScrollBox);
+	const FScrollBoxStyle* StylePtr = ( Style != NULL ) ? Style->GetStyle<FScrollBoxStyle>() : NULL;
+	if ( StylePtr == NULL )
+	{
+		SScrollBox::FArguments Defaults;
+		StylePtr = Defaults._Style;
+	}
+
+	const FScrollBarStyle* BarStylePtr = ( BarStyle != NULL ) ? BarStyle->GetStyle<FScrollBarStyle>() : NULL;
+	if ( BarStylePtr == NULL )
+	{
+		SScrollBox::FArguments Defaults;
+		BarStylePtr = Defaults._BarStyle;
+	}
+
+	MyScrollBox = SNew(SScrollBox)
+		.Style(StylePtr)
+		.BarStyle(BarStylePtr);
 
 	for ( UPanelSlot* Slot : Slots )
 	{
diff --git a/Engine/Source/Runtime/UMG/Private/Components/UniformGridPanel.cpp b/Engine/Source/Runtime/UMG/Private/Components/UniformGridPanel.cpp
--- a/Engine/Source/Runtime/UMG/Private/Components/UniformGridPanel.cpp
+++ b/Engine/Source/Runtime/UMG/Private/Components/UniformGridPanel.cpp
@@ -9,6 +9,9 @@ UUniformGridPanel::UUniformGridPanel(const FPostConstructInitializeProperties& P
 	: Super(PCIP)
 {
 	bIsVariable = false;
+
+	SUniformGridPanel::FArguments Defaults;
+	Visiblity = UWidget::ConvertRuntimeToSerializedVisiblity(Defaults._Visibility.Get());
 }
 
 UClass* UUniformGridPanel::GetSlotClass() const
@@ -38,7 +41,6 @@ TSharedRef<SWidget> UUniformGridPanel::RebuildWidget()
 {
 	MyUniformGridPanel =
 		SNew(SUniformGridPanel)
-		.SlotPadding(SlotPadding)
 		.MinDesiredSlotWidth(MinDesiredSlotWidth)
 		.MinDesiredSlotHeight(MinDesiredSlotHeight);
 
@@ -53,3 +55,10 @@ TSharedRef<SWidget> UUniformGridPanel::RebuildWidget()
 
 	return BuildDesignTimeWidget( MyUniformGridPanel.ToSharedRef() );
 }
+
+void UUniformGridPanel::SyncronizeProperties()
+{
+	Super::SyncronizeProperties();
+
+	MyUniformGridPanel->SetSlotPadding(SlotPadding);
+}
diff --git a/Engine/Source/Runtime/UMG/Private/Components/UniformGridSlot.cpp b/Engine/Source/Runtime/UMG/Private/Components/UniformGridSlot.cpp
--- a/Engine/Source/Runtime/UMG/Private/Components/UniformGridSlot.cpp
+++ b/Engine/Source/Runtime/UMG/Private/Components/UniformGridSlot.cpp
@@ -16,8 +16,6 @@ UUniformGridSlot::UUniformGridSlot(const FPostConstructInitializeProperties& PCI
 void UUniformGridSlot::BuildSlot(TSharedRef<SUniformGridPanel> GridPanel)
 {
 	Slot = &GridPanel->AddSlot(Column, Row)
-		.HAlign(HorizontalAlignment)
-		.VAlign(VerticalAlignment)
 		[
 			Content == NULL ? SNullWidget::NullWidget : Content->GetWidget()
 		];
@@ -41,8 +39,28 @@ void UUniformGridSlot::SetColumn(int32 InColumn)
 	}
 }
 
+void UUniformGridSlot::SetHorizontalAlignment(EHorizontalAlignment InHorizontalAlignment)
+{
+	HorizontalAlignment = InHorizontalAlignment;
+	if ( Slot )
+	{
+		Slot->HAlignment = InHorizontalAlignment;
+	}
+}
+
+void UUniformGridSlot::SetVerticalAlignment(EVerticalAlignment InVerticalAlignment)
+{
+	VerticalAlignment = InVerticalAlignment;
+	if ( Slot )
+	{
+		Slot->VAlignment = InVerticalAlignment;
+	}
+}
+
 void UUniformGridSlot::SyncronizeProperties()
 {
 	SetRow(Row);
 	SetColumn(Column);
+	SetHorizontalAlignment(HorizontalAlignment);
+	SetVerticalAlignment(VerticalAlignment);
 }
diff --git a/Engine/Source/Runtime/UMG/Private/Components/VerticalBox.cpp b/Engine/Source/Runtime/UMG/Private/Components/VerticalBox.cpp
--- a/Engine/Source/Runtime/UMG/Private/Components/VerticalBox.cpp
+++ b/Engine/Source/Runtime/UMG/Private/Components/VerticalBox.cpp
@@ -9,6 +9,9 @@ UVerticalBox::UVerticalBox(const FPostConstructInitializeProperties& PCIP)
 	: Super(PCIP)
 {
 	bIsVariable = false;
+
+	SVerticalBox::FArguments Defaults;
+	Visiblity = UWidget::ConvertRuntimeToSerializedVisiblity( Defaults._Visibility.Get() );
 }
 
 UClass* UVerticalBox::GetSlotClass() const
diff --git a/Engine/Source/Runtime/UMG/Private/Components/Widget.cpp b/Engine/Source/Runtime/UMG/Private/Components/Widget.cpp
--- a/Engine/Source/Runtime/UMG/Private/Components/Widget.cpp
+++ b/Engine/Source/Runtime/UMG/Private/Components/Widget.cpp
@@ -90,14 +90,14 @@ TSharedRef<SWidget> UWidget::GetWidget() const
 {
 	TSharedPtr<SWidget> SafeWidget;
 
+	UWidget* MutableThis = const_cast<UWidget*>( this );
+
 	if ( !MyWidget.IsValid() )
 	{
 		// We lie a bit about this being a const function.  If this is the first call it's not really const
 		// and we need to construct and cache the widget for the first run.  But instead of forcing everyone
 		// downstream to make RebuildWidget const and force every implementation to make things mutable, we
 		// just blow away the const here.
-		UWidget* MutableThis = const_cast<UWidget*>(this);
-		
 		SafeWidget = MutableThis->RebuildWidget();
 		MyWidget = SafeWidget;
 		MutableThis->SyncronizeProperties();
@@ -107,7 +107,18 @@ TSharedRef<SWidget> UWidget::GetWidget() const
 		SafeWidget = MyWidget.Pin();
 	}
 
-	return SafeWidget.ToSharedRef();
+	// If it is a user widget wrap it in a SObjectWidget to keep the instance from being GC'ed
+	if ( MutableThis->IsA(UUserWidget::StaticClass()) )
+	{
+		return SNew(SObjectWidget, Cast<UUserWidget>(MutableThis))
+		[
+			SafeWidget.ToSharedRef()
+		];
+	}
+	else
+	{
+		return SafeWidget.ToSharedRef();
+	}
 }
 
 TSharedRef<SWidget> UWidget::BuildDesignTimeWidget(TSharedRef<SWidget> WrapWidget)
diff --git a/Engine/Source/Runtime/UMG/Private/Components/WidgetSwitcher.cpp b/Engine/Source/Runtime/UMG/Private/Components/WidgetSwitcher.cpp
--- a/Engine/Source/Runtime/UMG/Private/Components/WidgetSwitcher.cpp
+++ b/Engine/Source/Runtime/UMG/Private/Components/WidgetSwitcher.cpp
@@ -9,6 +9,9 @@ UWidgetSwitcher::UWidgetSwitcher(const FPostConstructInitializeProperties& PCIP)
 	: Super(PCIP)
 {
 	bIsVariable = true;
+
+	SWidgetSwitcher::FArguments Defaults;
+	Visiblity = UWidget::ConvertRuntimeToSerializedVisiblity(Defaults._Visibility.Get());
 }
 
 int32 UWidgetSwitcher::GetNumWidgets() const
diff --git a/Engine/Source/Runtime/UMG/Private/Slate/SConstraintCanvas.cpp b/Engine/Source/Runtime/UMG/Private/Slate/SConstraintCanvas.cpp
--- a/Engine/Source/Runtime/UMG/Private/Slate/SConstraintCanvas.cpp
+++ b/Engine/Source/Runtime/UMG/Private/Slate/SConstraintCanvas.cpp
@@ -37,6 +37,19 @@ int32 SConstraintCanvas::RemoveSlot( const TSharedRef<SWidget>& SlotWidget )
 	return -1;
 }
 
+struct FChildZOrder
+{
+	int32 ChildIndex;
+	int32 ZOrder;
+};
+
+struct FSortSlotsByZOrder
+{
+	FORCEINLINE bool operator()(const FChildZOrder& A, const FChildZOrder& B) const
+	{
+		return A.ZOrder < B.ZOrder;
+	}
+};
 
 /* SWidget overrides
  *****************************************************************************/
@@ -45,10 +58,26 @@ void SConstraintCanvas::OnArrangeChildren( const FGeometry& AllottedGeometry, FA
 {
 	if (Children.Num() > 0)
 	{
-		for (int32 ChildIndex = 0; ChildIndex < Children.Num(); ++ChildIndex)
+		// Sort the children based on zorder.
+		TArray< FChildZOrder > SlotOrder;
+		for ( int32 ChildIndex = 0; ChildIndex < Children.Num(); ++ChildIndex )
 		{
 			const SConstraintCanvas::FSlot& CurChild = Children[ChildIndex];
 
+			FChildZOrder Order;
+			Order.ChildIndex = ChildIndex;
+			Order.ZOrder = CurChild.ZOrderAttr.Get();
+			SlotOrder.Add( Order );
+		}
+
+		SlotOrder.Sort(FSortSlotsByZOrder());
+
+		// Arrange the children now in their proper z-order.
+		for (int32 ChildIndex = 0; ChildIndex < Children.Num(); ++ChildIndex)
+		{
+			const FChildZOrder& CurSlot = SlotOrder[ChildIndex];
+			const SConstraintCanvas::FSlot& CurChild = Children[CurSlot.ChildIndex];
+
 			const FMargin Offset = CurChild.OffsetAttr.Get();
 			const FVector2D Alignment = CurChild.AlignmentAttr.Get();
 			const FAnchors Anchors = CurChild.AnchorsAttr.Get();
@@ -93,7 +122,7 @@ void SConstraintCanvas::OnArrangeChildren( const FGeometry& AllottedGeometry, FA
 				LocalPosition.Y = AnchorPixels.Top + Offset.Top - AlignmentOffset.Y;
 				LocalSize.Y = Size.Y;
 			}
-						
+			
 			// Add the information about this child to the output list (ArrangedChildren)
 			ArrangedChildren.AddWidget( AllottedGeometry.MakeChild(
 				// The child widget being arranged
diff --git a/Engine/Source/Runtime/UMG/Private/UMGStyle.cpp b/Engine/Source/Runtime/UMG/Private/UMGStyle.cpp
--- a/Engine/Source/Runtime/UMG/Private/UMGStyle.cpp
+++ b/Engine/Source/Runtime/UMG/Private/UMGStyle.cpp
@@ -46,7 +46,7 @@ TSharedRef< FSlateStyleSet > FUMGStyle::Create()
 	Style->Set("Widget", new IMAGE_BRUSH(TEXT("Widget"), Icon16x16));
 	Style->Set("Widget.CheckBox", new IMAGE_BRUSH(TEXT("CheckBox"), Icon16x16));
 	Style->Set("Widget.Button", new IMAGE_BRUSH(TEXT("Button"), Icon16x16));
-	Style->Set("Widget.EditableTextBlock", new IMAGE_BRUSH(TEXT("EditableTextBlock"), Icon16x16));
+	Style->Set("Widget.EditableTextBox", new IMAGE_BRUSH(TEXT("EditableTextBox"), Icon16x16));
 	Style->Set("Widget.EditableText", new IMAGE_BRUSH(TEXT("EditableText"), Icon16x16));
 	Style->Set("Widget.HorizontalBox", new IMAGE_BRUSH(TEXT("HorizontalBox"), Icon16x16));
 	Style->Set("Widget.VerticalBox", new IMAGE_BRUSH(TEXT("VerticalBox"), Icon16x16));
@@ -58,6 +58,7 @@ TSharedRef< FSlateStyleSet > FUMGStyle::Create()
 	Style->Set("Widget.Spacer", new IMAGE_BRUSH(TEXT("Spacer"), Icon16x16));
 	Style->Set("Widget.ScrollBox", new IMAGE_BRUSH(TEXT("ScrollBox"), Icon16x16));
 	Style->Set("Widget.ProgressBar", new IMAGE_BRUSH(TEXT("ProgressBar"), Icon16x16));
+	Style->Set("Widget.MenuAnchor", new IMAGE_BRUSH(TEXT("MenuAnchor"), Icon16x16));
 
 	Style->Set("Widget.UserWidget", new IMAGE_BRUSH(TEXT("UserWidget"), Icon16x16));
 
diff --git a/Engine/Source/Runtime/UMG/Private/UserWidget.cpp b/Engine/Source/Runtime/UMG/Private/UserWidget.cpp
--- a/Engine/Source/Runtime/UMG/Private/UserWidget.cpp
+++ b/Engine/Source/Runtime/UMG/Private/UserWidget.cpp
@@ -30,8 +30,12 @@ TSharedPtr<SWidget> FUMGDragDropOp::GetDefaultDecorator() const
 class SViewportWidgetHost : public SCompoundWidget
 {
 	SLATE_BEGIN_ARGS(SViewportWidgetHost)
-	{ }
+	{
+		_Visibility = EVisibility::SelfHitTestInvisible;
+	}
+
 		SLATE_DEFAULT_SLOT(FArguments, Content)
+
 	SLATE_END_ARGS()
 
 	void Construct(const FArguments& InArgs, bool bInModal)
@@ -44,10 +48,10 @@ class SViewportWidgetHost : public SCompoundWidget
 			];
 	}
 
-	virtual bool OnHitTest(const FGeometry& MyGeometry, FVector2D InAbsoluteCursorPosition) override
-	{
-		return true;
-	}
+	//virtual bool OnHitTest(const FGeometry& MyGeometry, FVector2D InAbsoluteCursorPosition) override
+	//{
+	//	return false;
+	//}
 
 	virtual bool SupportsKeyboardFocus() const override
 	{
@@ -90,35 +94,35 @@ class SViewportWidgetHost : public SCompoundWidget
 	//	return FReply::Handled();
 	//}
 
-	FReply OnMouseButtonDown(const FGeometry& MyGeometry, const FPointerEvent& MouseEvent)
-	{
-		if ( bModal )
-		{
-			return FReply::Handled();
-		}
+	//FReply OnMouseButtonDown(const FGeometry& MyGeometry, const FPointerEvent& MouseEvent)
+	//{
+	//	if ( bModal )
+	//	{
+	//		return FReply::Handled();
+	//	}
 
-		return FReply::Unhandled();
-	}
+	//	return FReply::Unhandled();
+	//}
 
-	FReply OnMouseButtonUp(const FGeometry& MyGeometry, const FPointerEvent& MouseEvent)
-	{
-		if ( bModal )
-		{
-			return FReply::Handled();
-		}
+	//FReply OnMouseButtonUp(const FGeometry& MyGeometry, const FPointerEvent& MouseEvent)
+	//{
+	//	if ( bModal )
+	//	{
+	//		return FReply::Handled();
+	//	}
 
-		return FReply::Unhandled();
-	}
+	//	return FReply::Unhandled();
+	//}
 
-	FReply OnMouseMove(const FGeometry& MyGeometry, const FPointerEvent& MouseEvent)
-	{
-		if ( bModal )
-		{
-			return FReply::Handled();
-		}
+	//FReply OnMouseMove(const FGeometry& MyGeometry, const FPointerEvent& MouseEvent)
+	//{
+	//	if ( bModal )
+	//	{
+	//		return FReply::Handled();
+	//	}
 
-		return FReply::Unhandled();
-	}
+	//	return FReply::Unhandled();
+	//}
 
 protected:
 	bool bModal;
@@ -139,6 +143,8 @@ UUserWidget::UUserWidget(const FPostConstructInitializeProperties& PCIP)
 	AbsoluteSize = FVector2D(100, 100);
 	HorizontalAlignment = HAlign_Fill;
 	VerticalAlignment = VAlign_Fill;
+
+	Visiblity = ESlateVisibility::SelfHitTestInvisible;
 }
 
 void UUserWidget::PostInitProperties()
@@ -154,7 +160,7 @@ void UUserWidget::PostInitProperties()
 		BGClass->InitializeWidget(this);
 	}
 
-	//RebuildWidget();
+	//TODO UMG For non-BP versions how do we generate the Components list?
 }
 
 UWorld* UUserWidget::GetWorld() const
@@ -176,8 +182,6 @@ UWorld* UUserWidget::GetWorld() const
 
 UWidget* UUserWidget::GetWidgetHandle(TSharedRef<SWidget> InWidget)
 {
-	GetWidget();
-	
 	TWeakObjectPtr<UWidget> VisualWidget = WidgetToComponent.FindRef(InWidget);
 	return VisualWidget.Get();
 }
@@ -207,16 +211,14 @@ TSharedRef<SWidget> UUserWidget::RebuildWidget()
 		WidgetToComponent.Add(Widget, Handle);
 	}
 
-	return SNew(SObjectWidget, this)
-		[
-			UserRootWidget.ToSharedRef()
-		];
+	// Notify the widget that it has been constructed.
+	Construct();
+
+	return UserRootWidget.ToSharedRef();
 }
 
 TSharedPtr<SWidget> UUserWidget::GetWidgetFromName(const FString& Name) const
 {
-	GetWidget();
-	
 	for ( auto& Entry : WidgetToComponent )
 	{
 		if ( Entry.Value->GetName().Equals(Name, ESearchCase::IgnoreCase) )
@@ -230,8 +232,6 @@ TSharedPtr<SWidget> UUserWidget::GetWidgetFromName(const FString& Name) const
 
 UWidget* UUserWidget::GetHandleFromName(const FString& Name) const
 {
-	GetWidget();
-	
 	for ( auto& Entry : WidgetToComponent )
 	{
 		if ( Entry.Value->GetName().Equals(Name, ESearchCase::IgnoreCase) )
@@ -279,8 +279,6 @@ TSharedRef<SWidget> UUserWidget::MakeFullScreenWidget()
 
 UWidget* UUserWidget::GetRootWidgetComponent()
 {
-	GetWidget();
-	
 	if ( Components.Num() > 0 )
 	{
 		return Components[0];
@@ -291,8 +289,6 @@ UWidget* UUserWidget::GetRootWidgetComponent()
 
 void UUserWidget::Show()
 {
-	GetWidget();
-	
 	if ( !FullScreenWidget.IsValid() )
 	{
 		TSharedRef<SWidget> RootWidget = MakeFullScreenWidget();
@@ -326,8 +322,6 @@ void UUserWidget::Show()
 
 void UUserWidget::Hide()
 {
-	GetWidget();
-	
 	if ( FullScreenWidget.IsValid() )
 	{
 		TSharedPtr<SWidget> RootWidget = FullScreenWidget.Pin();
@@ -356,15 +350,11 @@ void UUserWidget::Hide()
 
 bool UUserWidget::GetIsVisible()
 {
-	GetWidget();
-	
 	return FullScreenWidget.IsValid();
 }
 
 TEnumAsByte<ESlateVisibility::Type> UUserWidget::GetVisiblity()
 {
-	GetWidget();
-	
 	if ( FullScreenWidget.IsValid() )
 	{
 		TSharedPtr<SWidget> RootWidget = FullScreenWidget.Pin();
diff --git a/Engine/Source/Runtime/UMG/Private/WidgetBlueprintLibrary.cpp b/Engine/Source/Runtime/UMG/Private/WidgetBlueprintLibrary.cpp
--- a/Engine/Source/Runtime/UMG/Private/WidgetBlueprintLibrary.cpp
+++ b/Engine/Source/Runtime/UMG/Private/WidgetBlueprintLibrary.cpp
@@ -18,21 +18,65 @@ UUserWidget* UWidgetBlueprintLibrary::Create(UObject* WorldContextObject, TSubcl
 {
 	if ( WidgetType == NULL || WidgetType->HasAnyClassFlags(CLASS_Abstract) )
 	{
-		// TODO script Error?
+		// TODO UMG script Error?
 		return NULL;
 	}
 
-	//TODO ConstructObject after UUserWidget is no longer an actor.
-
 	UWorld* World = GEngine->GetWorldFromContextObject(WorldContextObject);
 	UUserWidget* NewWidget = ConstructObject<UUserWidget>(WidgetType, World->GetCurrentLevel());
 	NewWidget->SetFlags(RF_Transactional);
-	//UUserWidget* NewWidget = World->SpawnActor<UUserWidget>(WidgetType);
-		
+	
 	return NewWidget;
 }
 
-void UWidgetBlueprintLibrary::DrawBox(FPaintContext& Context, FVector2D Position, FVector2D Size, USlateBrushAsset* Brush, FLinearColor Tint)
+//void UWidgetBlueprintLibrary::Popup(UObject* WorldContextObject, UWidget* PopupWidget, int32 ZIndex)
+//{
+//	if ( !(WorldContextObject || PopupWidget) )
+//	{
+//		return;
+//	}
+//
+//	UWorld* World = GEngine->GetWorldFromContextObject(WorldContextObject);
+//	if ( World && World->IsGameWorld() )
+//	{
+//		TSharedRef<SWidget> SlatePopupWidget = PopupWidget->GetWidget();
+//
+//		UGameViewportClient* Viewport = World->GetGameViewport();
+//		Viewport->AddViewportWidgetContent(SlatePopupWidget, ZIndex);
+//
+//		TWeakPtr<SViewport> GameViewportWidget = Viewport->GetGameViewport()->GetViewportWidget();
+//		if ( GameViewportWidget.IsValid() )
+//		{
+//			GameViewportWidget.Pin()->SetWidgetToFocusOnActivate(SlatePopupWidget);
+//			FSlateApplication::Get().SetKeyboardFocus(SlatePopupWidget);
+//		}
+//
+//		
+//
+//		//TSharedPtr<SWeakWidget> PopupWrapperPtr;
+//		//PopupLayerSlot = &PopupWindow->AddPopupLayerSlot()
+//		//	.DesktopPosition(NewPosition)
+//		//	.WidthOverride(NewWindowSize.X)
+//		//	[
+//		//		SAssignNew(PopupWrapperPtr, SWeakWidget)
+//		//		.PossiblyNullContent
+//		//		(
+//		//		MenuContentRef
+//		//		)
+//		//	];
+//
+//		// We want to support dismissing the popup widget when the user clicks outside it.
+//		//FSlateApplication::Get().GetPopupSupport().RegisterClickNotification(PopupWrapperPtr.ToSharedRef(), FOnClickedOutside::CreateSP(this, &SMenuAnchor::OnClickedOutsidePopup));
+//		//}
+//
+//		//if ( bFocusMenu )
+//		//{
+//		//	FSlateApplication::Get().SetKeyboardFocus(MenuContentRef, EKeyboardFocusCause::SetDirectly);
+//		//}
+//	}
+//}
+
+void UWidgetBlueprintLibrary::DrawBox(UPARAM(ref) FPaintContext& Context, FVector2D Position, FVector2D Size, USlateBrushAsset* Brush, FLinearColor Tint)
 {
 	Context.MaxLayer++;
 
diff --git a/Engine/Source/Runtime/UMG/Public/Blueprint/UserWidget.h b/Engine/Source/Runtime/UMG/Public/Blueprint/UserWidget.h
--- a/Engine/Source/Runtime/UMG/Public/Blueprint/UserWidget.h
+++ b/Engine/Source/Runtime/UMG/Public/Blueprint/UserWidget.h
@@ -86,6 +86,7 @@ public:
 
 //TODO UMG If you want to host a widget that's full screen there may need to be a SWindow equivalent that you spawn it into.
 
+DECLARE_DYNAMIC_MULTICAST_DELEGATE(FOnConstructEvent);
 DECLARE_DYNAMIC_MULTICAST_DELEGATE_OneParam(FOnVisibilityChangedEvent, ESlateVisibility::Type, Visibility);
 
 /**
@@ -115,6 +116,10 @@ public:
 	UFUNCTION(BlueprintPure, Category="Appearance")
 	TEnumAsByte<ESlateVisibility::Type> GetVisiblity();
 
+	/** Called when the widget is constructed */
+	UFUNCTION(BlueprintImplementableEvent, Category="User Interface")
+	void Construct();
+
 	UFUNCTION(BlueprintImplementableEvent, Category="User Interface")
 	void Tick(FGeometry MyGeometry, float InDeltaTime);
 
diff --git a/Engine/Source/Runtime/UMG/Public/Blueprint/WidgetBlueprintLibrary.h b/Engine/Source/Runtime/UMG/Public/Blueprint/WidgetBlueprintLibrary.h
--- a/Engine/Source/Runtime/UMG/Public/Blueprint/WidgetBlueprintLibrary.h
+++ b/Engine/Source/Runtime/UMG/Public/Blueprint/WidgetBlueprintLibrary.h
@@ -12,6 +12,10 @@ class UWidgetBlueprintLibrary : public UBlueprintFunctionLibrary
 	UFUNCTION(BlueprintCallable, meta=( HidePin="WorldContextObject", DefaultToSelf="WorldContextObject", UnsafeDuringActorConstruction = "true", FriendlyName = "Create Widget" ), Category="User Interface|Widget")
 	static class UUserWidget* Create(UObject* WorldContextObject, TSubclassOf<class UUserWidget> WidgetType);
 
+	///** Show a widget as a popup */
+	//UFUNCTION(BlueprintCallable, meta=( HidePin="WorldContextObject", DefaultToSelf="WorldContextObject", UnsafeDuringActorConstruction = "true" ), Category="User Interface" )
+	//static void Popup(UObject* WorldContextObject, UWidget* PopupWidget, int32 ZIndex = 0);
+
 	/** Draws a box */
 	UFUNCTION(BlueprintCallable, Category="Painting")
 	static void DrawBox(UPARAM(ref) FPaintContext& Context, FVector2D Position, FVector2D Size, USlateBrushAsset* Brush, FLinearColor Tint = FLinearColor::White);
diff --git a/Engine/Source/Runtime/UMG/Public/Components/CanvasPanelSlot.h b/Engine/Source/Runtime/UMG/Public/Components/CanvasPanelSlot.h
--- a/Engine/Source/Runtime/UMG/Public/Components/CanvasPanelSlot.h
+++ b/Engine/Source/Runtime/UMG/Public/Components/CanvasPanelSlot.h
@@ -34,6 +34,10 @@ class UMG_API UCanvasPanelSlot : public UPanelSlot
 	UPROPERTY(EditDefaultsOnly, Category=Appearance)
 	FAnchorData LayoutData;
 
+	/** The order priority this widget is rendered in.  Higher values are rendered last (and so they will appear to be on top). */
+	UPROPERTY(EditDefaultsOnly, Category=AnchorData)
+	int32 ZOrder;
+
 	void BuildSlot(TSharedRef<SConstraintCanvas> Canvas);
 
 	virtual void SetDesiredPosition(FVector2D InPosition) override;
@@ -56,6 +60,10 @@ class UMG_API UCanvasPanelSlot : public UPanelSlot
 	UFUNCTION(BlueprintCallable, Category="Appearance")
 	void SetAlignment(FVector2D InAlignment);
 
+	/** Sets the z-order on the slot */
+	UFUNCTION(BlueprintCallable, Category="Appearance")
+	void SetZOrder(int32 InZOrder);
+
 	// UPanelSlot interface
 	virtual void SyncronizeProperties() override;
 	// End of UPanelSlot interface
diff --git a/Engine/Source/Runtime/UMG/Public/Components/EditableTextBox.h b/Engine/Source/Runtime/UMG/Public/Components/EditableTextBox.h
--- a/Engine/Source/Runtime/UMG/Public/Components/EditableTextBox.h
+++ b/Engine/Source/Runtime/UMG/Public/Components/EditableTextBox.h
@@ -15,6 +15,10 @@ class UMG_API UEditableTextBox : public UWidget
 
 protected:
 
+	/** Style used for the text box */
+	UPROPERTY(EditDefaultsOnly, Category="Style", meta=( DisplayThumbnail = "true" ))
+	USlateWidgetStyleAsset* Style;
+
 	/** The text content for this editable text box widget */
 	UPROPERTY(EditDefaultsOnly, Category=Content)
 	FText Text;
@@ -93,6 +97,9 @@ protected:
 	/**  */
 	UFUNCTION(BlueprintCallable, Category="Widget")
 	void SetText(FText InText);
+
+	UFUNCTION(BlueprintCallable, Category="Widget")
+	void SetError(FText InError);
 	
 	// UWidget interface
 	virtual void SyncronizeProperties() override;
diff --git a/Engine/Source/Runtime/UMG/Public/Components/ProgressBar.h b/Engine/Source/Runtime/UMG/Public/Components/ProgressBar.h
--- a/Engine/Source/Runtime/UMG/Public/Components/ProgressBar.h
+++ b/Engine/Source/Runtime/UMG/Public/Components/ProgressBar.h
@@ -13,20 +13,36 @@ class UMG_API UProgressBar : public UWidget
 public:
 
 	/** Style used for the progress bar */
-	UPROPERTY(EditDefaultsOnly, Category=Style, meta=( DisplayThumbnail = "true" ))
+	UPROPERTY(EditDefaultsOnly, Category="Style", meta=( DisplayThumbnail = "true" ))
 	USlateWidgetStyleAsset* Style;
 
+	/** The brush to use as the background of the progress bar */
+	UPROPERTY(EditDefaultsOnly, Category="Style", meta=( DisplayThumbnail = "true" ), AdvancedDisplay)
+	USlateBrushAsset* BackgroundImage;
+	
+	/** The brush to use as the fill image */
+	UPROPERTY(EditDefaultsOnly, Category="Style", meta=( DisplayThumbnail = "true" ), AdvancedDisplay)
+	USlateBrushAsset* FillImage;
+	
+	/** The brush to use as the marquee image */
+	UPROPERTY(EditDefaultsOnly, Category="Style", meta=( DisplayThumbnail = "true" ), AdvancedDisplay)
+	USlateBrushAsset* MarqueeImage;
+
 	/** Defines if this progress bar fills Left to right or right to left */
 	UPROPERTY(EditDefaultsOnly, Category=Appearance)
 	TEnumAsByte<EProgressBarFillType::Type> BarFillType;
 	
-	UPROPERTY()
-	bool bOverride_Percent;
+	UPROPERTY(EditDefaultsOnly, Category=Appearance)
+	bool bIsMarquee;
 
 	/** Used to determine the fill position of the progress bar ranging 0..1 */
 	UPROPERTY(EditDefaultsOnly, Category=Appearance, meta=( UIMin = "0", UIMax = "1", editcondition = "bOverride_Percent" ))
 	float Percent;
 
+	/** A bindable delegate to allow logic to drive the text of the widget */
+	UPROPERTY()
+	FGetFloat PercentDelegate;
+
 	/** Fill Color and Opacity */
 	UPROPERTY(EditDefaultsOnly, Category=Appearance)
 	FLinearColor FillColorAndOpacity;
@@ -35,18 +51,6 @@ public:
 	UPROPERTY(EditDefaultsOnly, Category=Appearance)
 	FVector2D BorderPadding;
 	
-	/** The brush to use as the background of the progress bar */
-	UPROPERTY(EditDefaultsOnly, Category=Appearance, meta=( DisplayThumbnail = "true" ))
-	USlateBrushAsset* BackgroundImage;
-	
-	/** The brush to use as the fill image */
-	UPROPERTY(EditDefaultsOnly, Category=Appearance, meta=( DisplayThumbnail = "true" ))
-	USlateBrushAsset* FillImage;
-	
-	/** The brush to use as the marquee image */
-	UPROPERTY(EditDefaultsOnly, Category=Appearance, meta=( DisplayThumbnail = "true" ))
-	USlateBrushAsset* MarqueeImage;
-	
 	/** Sets the current value of the ProgressBar. */
 	UFUNCTION(BlueprintCallable, Category="Behavior")
 	void SetPercent(float InPercent);
diff --git a/Engine/Source/Runtime/UMG/Public/Components/ScrollBox.h b/Engine/Source/Runtime/UMG/Public/Components/ScrollBox.h
--- a/Engine/Source/Runtime/UMG/Public/Components/ScrollBox.h
+++ b/Engine/Source/Runtime/UMG/Public/Components/ScrollBox.h
@@ -10,6 +10,14 @@ class UMG_API UScrollBox : public UPanelWidget
 {
 	GENERATED_UCLASS_BODY()
 
+	/** Style of the scrollbox */
+	UPROPERTY(EditDefaultsOnly, Category=Style, meta=( DisplayThumbnail = "true" ))
+	USlateWidgetStyleAsset* Style;
+
+	/** Style of the scrollbox's bar */
+	UPROPERTY(EditDefaultsOnly, Category=Style, meta=( DisplayThumbnail = "true" ))
+	USlateWidgetStyleAsset* BarStyle;
+
 	/** Removes all child widget from the scrollbox */
 	UFUNCTION(BlueprintCallable, Category="Widget")
 	void ClearChildren();
diff --git a/Engine/Source/Runtime/UMG/Public/Components/UniformGridPanel.h b/Engine/Source/Runtime/UMG/Public/Components/UniformGridPanel.h
--- a/Engine/Source/Runtime/UMG/Public/Components/UniformGridPanel.h
+++ b/Engine/Source/Runtime/UMG/Public/Components/UniformGridPanel.h
@@ -22,6 +22,10 @@ class UMG_API UUniformGridPanel : public UPanelWidget
 	UPROPERTY(EditDefaultsOnly, Category=Layout)
 	float MinDesiredSlotHeight;
 
+	// UWidget interface
+	virtual void SyncronizeProperties() override;
+	// End of UWidget interface
+
 protected:
 
 	// UPanelWidget
diff --git a/Engine/Source/Runtime/UMG/Public/Components/UniformGridSlot.h b/Engine/Source/Runtime/UMG/Public/Components/UniformGridSlot.h
--- a/Engine/Source/Runtime/UMG/Public/Components/UniformGridSlot.h
+++ b/Engine/Source/Runtime/UMG/Public/Components/UniformGridSlot.h
@@ -26,11 +26,11 @@ public:
 	TEnumAsByte<EVerticalAlignment> VerticalAlignment;
 	
 	/** The row index of the cell this slot is in */
-	UPROPERTY(EditDefaultsOnly, Category=Layout)
+	UPROPERTY(EditDefaultsOnly, meta=( UIMin = "0" ), Category=Layout)
 	int32 Row;
 	
 	/** The column index of the cell this slot is in */
-	UPROPERTY(EditDefaultsOnly, Category=Layout)
+	UPROPERTY(EditDefaultsOnly, meta=( UIMin = "0" ), Category=Layout)
 	int32 Column;
 
 	/** Sets the row index of the slot, this determines what cell the slot is in the panel */
@@ -41,6 +41,14 @@ public:
 	UFUNCTION(BlueprintCallable, Category=Layout)
 	void SetColumn(int32 InColumn);
 
+	/**  */
+	UFUNCTION(BlueprintCallable, Category=Layout)
+	void SetHorizontalAlignment(EHorizontalAlignment InHorizontalAlignment);
+
+	/**  */
+	UFUNCTION(BlueprintCallable, Category=Layout)
+	void SetVerticalAlignment(EVerticalAlignment InVerticalAlignment);
+
 	// UPanelSlot interface
 	virtual void SyncronizeProperties() override;
 	// End of UPanelSlot interface
diff --git a/Engine/Source/Runtime/UMG/Public/Components/Widget.h b/Engine/Source/Runtime/UMG/Public/Components/Widget.h
--- a/Engine/Source/Runtime/UMG/Public/Components/Widget.h
+++ b/Engine/Source/Runtime/UMG/Public/Components/Widget.h
@@ -29,7 +29,13 @@ public:
 
 	// Common Bindings
 	DECLARE_DYNAMIC_DELEGATE_RetVal(bool, FGetBool);
+	DECLARE_DYNAMIC_DELEGATE_RetVal(float, FGetFloat);
+	DECLARE_DYNAMIC_DELEGATE_RetVal(int32, FGetInt32);
 	DECLARE_DYNAMIC_DELEGATE_RetVal(FText, FGetText);
+	DECLARE_DYNAMIC_DELEGATE_RetVal(FVector2D, FGetVector2D);
+	DECLARE_DYNAMIC_DELEGATE_RetVal(FVector, FGetVector);
+	DECLARE_DYNAMIC_DELEGATE_RetVal(FVector4, FGetVector4);
+	DECLARE_DYNAMIC_DELEGATE_RetVal(FMargin, FGetMargin);
 	DECLARE_DYNAMIC_DELEGATE_RetVal(FLinearColor, FGetSlateColor);
 	DECLARE_DYNAMIC_DELEGATE_RetVal(FLinearColor, FGetLinearColor);
 	DECLARE_DYNAMIC_DELEGATE_RetVal(ESlateVisibility::Type, FGetSlateVisibility);
@@ -198,11 +204,19 @@ protected:
 	
 	TSharedRef<SWidget> BuildDesignTimeWidget(TSharedRef<SWidget> WrapWidget);
 
+protected:
+	//TODO UMG Consider moving conversion functions into another class.
+	// Conversion functions
 	EVisibility ConvertVisibility(TAttribute<ESlateVisibility::Type> SerializedType) const
 	{
 		return ConvertSerializedVisibilityToRuntime(SerializedType.Get());
 	}
 
+	TOptional<float> ConvertFloatToOptionalFloat(TAttribute<float> InFloat) const
+	{
+		return InFloat.Get();
+	}
+
 protected:
 	/** The underlying SWidget. */
 	mutable TWeakPtr<SWidget> MyWidget;
diff --git a/Engine/Source/Runtime/UMG/Public/Slate/SConstraintCanvas.h b/Engine/Source/Runtime/UMG/Public/Slate/SConstraintCanvas.h
--- a/Engine/Source/Runtime/UMG/Public/Slate/SConstraintCanvas.h
+++ b/Engine/Source/Runtime/UMG/Public/Slate/SConstraintCanvas.h
@@ -35,6 +35,12 @@ public:
 			return *this;
 		}
 
+		FSlot& ZOrder(const TAttribute<int32>& InZOrder)
+		{
+			ZOrderAttr = InZOrder;
+			return *this;
+		}
+
 		FSlot& operator[]( TSharedRef<SWidget> InWidget )
 		{
 			Widget = InWidget;
@@ -59,12 +65,16 @@ public:
 		/** Size */
 		TAttribute<FVector2D> AlignmentAttr;
 
+		/** Size */
+		TAttribute<int32> ZOrderAttr;
+
 		/** Default values for a slot. */
 		FSlot()
 			: Widget( SNullWidget::NullWidget )
 			, OffsetAttr( FMargin( 0, 0, 1, 1 ) )
 			, AnchorsAttr( FAnchors( 0.5f, 0.5f ) )
 			, AlignmentAttr( FVector2D( 0.5f, 0.5f ) )
+			, ZOrderAttr( 0 )
 		{ }
 	};
 
diff --git a/Engine/Source/Runtime/UMG/Public/Slate/SObjectWidget.h b/Engine/Source/Runtime/UMG/Public/Slate/SObjectWidget.h
--- a/Engine/Source/Runtime/UMG/Public/Slate/SObjectWidget.h
+++ b/Engine/Source/Runtime/UMG/Public/Slate/SObjectWidget.h
@@ -11,7 +11,10 @@
 class SObjectWidget : public SCompoundWidget, public FGCObject
 {
 	SLATE_BEGIN_ARGS(SObjectWidget)
-	{ }
+	{
+		_Visibility = EVisibility::SelfHitTestInvisible;
+	}
+
 		SLATE_DEFAULT_SLOT(FArguments, Content)
 	SLATE_END_ARGS()
 
