UMG - Fixing a potential build error.

diff --git a/Engine/Source/Runtime/UMG/Public/Components/Widget.h b/Engine/Source/Runtime/UMG/Public/Components/Widget.h
--- a/Engine/Source/Runtime/UMG/Public/Components/Widget.h
+++ b/Engine/Source/Runtime/UMG/Public/Components/Widget.h
@@ -34,7 +34,7 @@ public:
 	DECLARE_DYNAMIC_DELEGATE_RetVal(FText, FGetText);
 	DECLARE_DYNAMIC_DELEGATE_RetVal(FVector2D, FGetVector2D);
 	DECLARE_DYNAMIC_DELEGATE_RetVal(FVector, FGetVector);
-	DECLARE_DYNAMIC_DELEGATE_RetVal(FVector4, FGetVector4);
+	//DECLARE_DYNAMIC_DELEGATE_RetVal(FVector4, FGetVector4);
 	DECLARE_DYNAMIC_DELEGATE_RetVal(FMargin, FGetMargin);
 	DECLARE_DYNAMIC_DELEGATE_RetVal(FLinearColor, FGetSlateColor);
 	DECLARE_DYNAMIC_DELEGATE_RetVal(FLinearColor, FGetLinearColor);
