UMG - Cleaning up some multicast delegates in the textboxes and slider.  Need to start exposing them in the designer again now that widgets are not actors.

diff --git a/Engine/Source/Runtime/UMG/Private/Components/EditableTextBox.cpp b/Engine/Source/Runtime/UMG/Private/Components/EditableTextBox.cpp
--- a/Engine/Source/Runtime/UMG/Private/Components/EditableTextBox.cpp
+++ b/Engine/Source/Runtime/UMG/Private/Components/EditableTextBox.cpp
@@ -46,7 +46,8 @@ TSharedRef<SWidget> UEditableTextBox::RebuildWidget()
 		.RevertTextOnEscape(RevertTextOnEscape)
 		.ClearKeyboardFocusOnCommit(ClearKeyboardFocusOnCommit)
 		.SelectAllTextOnCommit(SelectAllTextOnCommit)
-		.OnTextChanged(BIND_UOBJECT_DELEGATE(FOnTextChanged, SlateOnTextChanged))
+		.OnTextChanged(BIND_UOBJECT_DELEGATE(FOnTextChanged, HandleOnTextChanged))
+		.OnTextCommitted(BIND_UOBJECT_DELEGATE(FOnTextCommitted, HandleOnTextCommitted))
 		;
 
 	return MyEditableTextBlock.ToSharedRef();
@@ -98,11 +99,16 @@ void UEditableTextBox::SetError(FText InError)
 	}
 }
 
-void UEditableTextBox::SlateOnTextChanged(const FText& Text)
+void UEditableTextBox::HandleOnTextChanged(const FText& Text)
 {
 	OnTextChanged.Broadcast(Text);
 }
 
+void UEditableTextBox::HandleOnTextCommitted(const FText& Text, ETextCommit::Type CommitMethod)
+{
+	OnTextCommitted.Broadcast(Text, CommitMethod);
+}
+
 #if WITH_EDITOR
 
 const FSlateBrush* UEditableTextBox::GetEditorIcon()
diff --git a/Engine/Source/Runtime/UMG/Public/Components/EditableText.h b/Engine/Source/Runtime/UMG/Public/Components/EditableText.h
--- a/Engine/Source/Runtime/UMG/Public/Components/EditableText.h
+++ b/Engine/Source/Runtime/UMG/Public/Components/EditableText.h
@@ -4,16 +4,18 @@
 
 #include "EditableText.generated.h"
 
-DECLARE_DYNAMIC_MULTICAST_DELEGATE_OneParam(FOnEditableTextChangedEvent, const FText&, Text);
-DECLARE_DYNAMIC_MULTICAST_DELEGATE_TwoParams(FOnEditableTextCommittedEvent, const FText&, Text, ETextCommit::Type, CommitMethod);
-
 /** Editable text box widget */
 UCLASS(meta=( Category="Common" ), ClassGroup=UserInterface)
 class UMG_API UEditableText : public UWidget
 {
 	GENERATED_UCLASS_BODY()
 
-protected:
+public:
+
+	DECLARE_DYNAMIC_MULTICAST_DELEGATE_OneParam(FOnEditableTextChangedEvent, const FText&, Text);
+	DECLARE_DYNAMIC_MULTICAST_DELEGATE_TwoParams(FOnEditableTextCommittedEvent, const FText&, Text, ETextCommit::Type, CommitMethod);
+
+public:
 
 	/** The text content for this editable text box widget */
 	UPROPERTY(EditDefaultsOnly, Category=Content)
@@ -84,11 +86,11 @@ protected:
 	bool SelectAllTextOnCommit;
 
 	/** Called whenever the text is changed interactively by the user */
-	UPROPERTY(BlueprintAssignable)
+	UPROPERTY(BlueprintAssignable, Category="Widget Event")
 	FOnEditableTextChangedEvent OnTextChanged;
 
 	/** Called whenever the text is committed.  This happens when the user presses enter or the text box loses focus. */
-	UPROPERTY(BlueprintAssignable)
+	UPROPERTY(BlueprintAssignable, Category="Widget Event")
 	FOnEditableTextCommittedEvent OnTextCommitted;
 
 	/**  */
diff --git a/Engine/Source/Runtime/UMG/Public/Components/EditableTextBox.h b/Engine/Source/Runtime/UMG/Public/Components/EditableTextBox.h
--- a/Engine/Source/Runtime/UMG/Public/Components/EditableTextBox.h
+++ b/Engine/Source/Runtime/UMG/Public/Components/EditableTextBox.h
@@ -4,16 +4,18 @@
 
 #include "EditableTextBox.generated.h"
 
-DECLARE_DYNAMIC_MULTICAST_DELEGATE_OneParam(FOnEditableTextBlock_TextChangedEvent, const FText&, Text);
-//DECLARE_DYNAMIC_MULTICAST_DELEGATE_TwoParams(FOnEditableTextBlock_TextCommittedEvent, const FText&, Text, ETextCommit::Type, Type);
-
 /** Editable text box widget */
 UCLASS(meta=( Category="Common" ), ClassGroup=UserInterface)
 class UMG_API UEditableTextBox : public UWidget
 {
 	GENERATED_UCLASS_BODY()
 
-protected:
+public:
+
+	DECLARE_DYNAMIC_MULTICAST_DELEGATE_OneParam(FOnEditableTextBoxChangedEvent, const FText&, Text);
+	DECLARE_DYNAMIC_MULTICAST_DELEGATE_TwoParams(FOnEditableTextBoxCommittedEvent, const FText&, Text, ETextCommit::Type, CommitMethod);
+
+public:
 
 	/** Style used for the text box */
 	UPROPERTY(EditDefaultsOnly, Category="Style", meta=( DisplayThumbnail = "true" ))
@@ -80,12 +82,12 @@ protected:
 	bool SelectAllTextOnCommit;
 
 	/** Called whenever the text is changed interactively by the user */
-	UPROPERTY(BlueprintAssignable)
-	FOnEditableTextBlock_TextChangedEvent OnTextChanged;
+	UPROPERTY(BlueprintAssignable, Category="Widget Event")
+	FOnEditableTextBoxChangedEvent OnTextChanged;
 
-	/** Called whenever the text is committed.  This happens when the user presses enter or the text box loses focus. *///
-	//UPROPERTY(BlueprintAssignable)
-	//FOnTextCommittedEvent OnTextCommitted;
+	/** Called whenever the text is committed.  This happens when the user presses enter or the text box loses focus. */
+	UPROPERTY(BlueprintAssignable, Category="Widget Event")
+	FOnEditableTextBoxCommittedEvent OnTextCommitted;
 
 	/** Provide a alternative mechanism for error reporting. */
 	//SLATE_ARGUMENT(TSharedPtr<class IErrorReportingWidget>, ErrorReporting)
@@ -100,7 +102,7 @@ protected:
 
 	UFUNCTION(BlueprintCallable, Category="Widget")
 	void SetError(FText InError);
-	
+
 	// UWidget interface
 	virtual void SyncronizeProperties() override;
 	// End of UWidget interface
@@ -114,7 +116,8 @@ protected:
 	virtual TSharedRef<SWidget> RebuildWidget() override;
 	// End of UWidget
 
-	void SlateOnTextChanged(const FText& Text);
+	void HandleOnTextChanged(const FText& Text);
+	void HandleOnTextCommitted(const FText& Text, ETextCommit::Type CommitMethod);
 
 protected:
 	TSharedPtr<SEditableTextBox> MyEditableTextBlock;
diff --git a/Engine/Source/Runtime/UMG/Public/Components/Slider.h b/Engine/Source/Runtime/UMG/Public/Components/Slider.h
--- a/Engine/Source/Runtime/UMG/Public/Components/Slider.h
+++ b/Engine/Source/Runtime/UMG/Public/Components/Slider.h
@@ -42,15 +42,15 @@ class UMG_API USlider : public UWidget
 	bool Locked;
 
 	/** Invoked when the mouse is pressed and a capture begins. */
-	UPROPERTY(BlueprintAssignable)
+	UPROPERTY(BlueprintAssignable, Category="Widget Event")
 	FOnMouseCaptureBeginEvent OnMouseCaptureBegin;
 
 	/** Invoked when the mouse is released and a capture ends. */
-	UPROPERTY(BlueprintAssignable)
+	UPROPERTY(BlueprintAssignable, Category="Widget Event")
 	FOnMouseCaptureEndEvent OnMouseCaptureEnd;
 
 	/** Called when the value is changed by slider or typing. */
-	UPROPERTY(BlueprintAssignable)
+	UPROPERTY(BlueprintAssignable, Category="Widget Event")
 	FOnFloatValueChangedEvent OnValueChanged;
 
 	/** Gets the current value of the slider. */
