UMG - Adding a way to set just the opacity of UImage.

diff --git a/Engine/Source/Runtime/UMG/Private/Components/Image.cpp b/Engine/Source/Runtime/UMG/Private/Components/Image.cpp
--- a/Engine/Source/Runtime/UMG/Private/Components/Image.cpp
+++ b/Engine/Source/Runtime/UMG/Private/Components/Image.cpp
@@ -37,6 +37,15 @@ void UImage::SetColorAndOpacity(FLinearColor InColorAndOpacity)
 	}
 }
 
+void UImage::SetOpacity(float InOpacity)
+{
+	ColorAndOpacity.A = InOpacity;
+	if ( MyImage.IsValid() )
+	{
+		MyImage->SetColorAndOpacity(ColorAndOpacity);
+	}
+}
+
 void UImage::SetImage(USlateBrushAsset* InImage)
 {
 	Image = InImage;
diff --git a/Engine/Source/Runtime/UMG/Public/Components/Image.h b/Engine/Source/Runtime/UMG/Public/Components/Image.h
--- a/Engine/Source/Runtime/UMG/Public/Components/Image.h
+++ b/Engine/Source/Runtime/UMG/Public/Components/Image.h
@@ -29,6 +29,10 @@ public:
 	UFUNCTION(BlueprintCallable, Category="Appearance")
 	void SetColorAndOpacity(FLinearColor InColorAndOpacity);
 
+	/**  */
+	UFUNCTION(BlueprintCallable, Category="Appearance")
+	void SetOpacity(float InOpacity);
+
 	/**  */
 	UFUNCTION(BlueprintCallable, Category="Appearance")
 	void SetImage(USlateBrushAsset* InImage);
