UMG - Adding autowrap text support to the textblock.

diff --git a/Engine/Source/Runtime/UMG/Private/Components/TextBlock.cpp b/Engine/Source/Runtime/UMG/Private/Components/TextBlock.cpp
--- a/Engine/Source/Runtime/UMG/Private/Components/TextBlock.cpp
+++ b/Engine/Source/Runtime/UMG/Private/Components/TextBlock.cpp
@@ -92,6 +92,8 @@ void UTextBlock::SyncronizeProperties()
 	MyTextBlock->SetTextStyle(StylePtr);
 	MyTextBlock->SetShadowOffset(ShadowOffset);
 	MyTextBlock->SetShadowColorAndOpacity(ShadowColorAndOpacityBinding);
+	MyTextBlock->SetAutoWrapText(AutoWrapText);
+	MyTextBlock->SetWrapTextAt(WrapTextAt != 0 ? WrapTextAt : TAttribute<float>());
 
 #if WITH_EDITOR
 	MyEditorTextBlock->SetText(TextBinding);
