UMG - Adding a stub ScrollBar widget minamally exposed.

diff --git a/Engine/Source/Runtime/UMG/Private/Components/ScrollBar.cpp b/Engine/Source/Runtime/UMG/Private/Components/ScrollBar.cpp
new file mode 100644
--- /dev/null
+++ b/Engine/Source/Runtime/UMG/Private/Components/ScrollBar.cpp
@@ -0,0 +1,64 @@
+// Copyright 1998-2014 Epic Games, Inc. All Rights Reserved.
+
+#include "UMGPrivatePCH.h"
+
+#define LOCTEXT_NAMESPACE "UMG"
+
+/////////////////////////////////////////////////////
+// UScrollBar
+
+UScrollBar::UScrollBar(const FPostConstructInitializeProperties& PCIP)
+	: Super(PCIP)
+{
+	bIsVariable = false;
+
+	bAlwaysShowScrollbar = true;
+	Orientation = Orient_Vertical;
+	Thickness = FVector2D(12.0f, 12.0f);
+}
+
+TSharedRef<SWidget> UScrollBar::RebuildWidget()
+{
+	const FScrollBarStyle* StylePtr = ( Style != NULL ) ? Style->GetStyle<FScrollBarStyle>() : NULL;
+	if ( StylePtr == NULL )
+	{
+		SScrollBar::FArguments Defaults;
+		StylePtr = Defaults._Style;
+	}
+
+	MyScrollBar = SNew(SScrollBar)
+		.Style(StylePtr)
+		.AlwaysShowScrollbar(bAlwaysShowScrollbar)
+		.Orientation(Orientation)
+		.Thickness(Thickness);
+
+	//SLATE_EVENT(FOnUserScrolled, OnUserScrolled)
+
+	return MyScrollBar.ToSharedRef();
+}
+
+void UScrollBar::SyncronizeProperties()
+{
+	//MyScrollBar->SetScrollOffset(DesiredScrollOffset);
+}
+
+void UScrollBar::SetState(float InOffsetFraction, float InThumbSizeFraction)
+{
+	if ( MyScrollBar.IsValid() )
+	{
+		MyScrollBar->SetState(InOffsetFraction, InThumbSizeFraction);
+	}
+}
+
+#if WITH_EDITOR
+
+const FSlateBrush* UScrollBar::GetEditorIcon()
+{
+	return FUMGStyle::Get().GetBrush("Widget.ScrollBar");
+}
+
+#endif
+
+/////////////////////////////////////////////////////
+
+#undef LOCTEXT_NAMESPACE
diff --git a/Engine/Source/Runtime/UMG/Private/UMGStyle.cpp b/Engine/Source/Runtime/UMG/Private/UMGStyle.cpp
--- a/Engine/Source/Runtime/UMG/Private/UMGStyle.cpp
+++ b/Engine/Source/Runtime/UMG/Private/UMGStyle.cpp
@@ -59,6 +59,7 @@ TSharedRef< FSlateStyleSet > FUMGStyle::Create()
 	Style->Set("Widget.ScrollBox", new IMAGE_BRUSH(TEXT("ScrollBox"), Icon16x16));
 	Style->Set("Widget.ProgressBar", new IMAGE_BRUSH(TEXT("ProgressBar"), Icon16x16));
 	Style->Set("Widget.MenuAnchor", new IMAGE_BRUSH(TEXT("MenuAnchor"), Icon16x16));
+	Style->Set("Widget.ScrollBar", new IMAGE_BRUSH(TEXT("ScrollBar"), Icon16x16));
 
 	Style->Set("Widget.UserWidget", new IMAGE_BRUSH(TEXT("UserWidget"), Icon16x16));
 
diff --git a/Engine/Source/Runtime/UMG/Public/Components/MenuAnchor.h b/Engine/Source/Runtime/UMG/Public/Components/MenuAnchor.h
--- a/Engine/Source/Runtime/UMG/Public/Components/MenuAnchor.h
+++ b/Engine/Source/Runtime/UMG/Public/Components/MenuAnchor.h
@@ -8,7 +8,7 @@
  * The Menu Anchor allows you to specify an location that a popup menu should be anchored to, and should be
  * summoned from.
  */
-UCLASS(meta=( Category="Popup" ), ClassGroup=UserInterface)
+UCLASS(meta=( Category="Popup" ), ClassGroup=UserInterface, meta=( IsAdvanced = "True" ))
 class UMG_API UMenuAnchor : public UContentWidget
 {
 	GENERATED_UCLASS_BODY()
diff --git a/Engine/Source/Runtime/UMG/Public/Components/ScrollBar.h b/Engine/Source/Runtime/UMG/Public/Components/ScrollBar.h
new file mode 100644
--- /dev/null
+++ b/Engine/Source/Runtime/UMG/Public/Components/ScrollBar.h
@@ -0,0 +1,74 @@
+// Copyright 1998-2014 Epic Games, Inc. All Rights Reserved.
+
+#pragma once
+
+#include "ScrollBar.generated.h"
+
+/** */
+UCLASS(meta=( Category="Primitive" ), ClassGroup=UserInterface, meta=( IsAdvanced = "True" ))
+class UMG_API UScrollBar : public UPanelWidget
+{
+	GENERATED_UCLASS_BODY()
+
+	/** Style of the scrollbar */
+	UPROPERTY(EditDefaultsOnly, Category="Style", meta=( DisplayThumbnail = "true" ))
+	USlateWidgetStyleAsset* Style;
+
+	/**  */
+	UPROPERTY(EditDefaultsOnly, Category="Behavior")
+	bool bAlwaysShowScrollbar;
+
+	/**  */
+	UPROPERTY(EditDefaultsOnly, Category="Behavior")
+	TEnumAsByte<EOrientation> Orientation;
+
+	/** The thickness of the scrollbar thumb */
+	UPROPERTY(EditDefaultsOnly, Category="Behavior")
+	FVector2D Thickness;
+
+	/**
+	* Set the offset and size of the track's thumb.
+	* Note that the maximum offset is 1.0-ThumbSizeFraction.
+	* If the user can view 1/3 of the items in a single page, the maximum offset will be ~0.667f
+	*
+	* @param InOffsetFraction     Offset of the thumbnail from the top as a fraction of the total available scroll space.
+	* @param InThumbSizeFraction  Size of thumbnail as a fraction of the total available scroll space.
+	*/
+	UFUNCTION(BlueprintCallable, Category="Scrolling")
+	void SetState(float InOffsetFraction, float InThumbSizeFraction);
+
+	///** @return true if scrolling is possible; false if the view is big enough to fit all the content */
+	//bool IsNeeded() const;
+
+	///** @return normalized distance from top */
+	//float DistanceFromTop() const;
+
+	///** @return normalized distance from bottom */
+	//float DistanceFromBottom() const;
+
+	///** @return the scrollbar's visibility as a product of internal rules and user-specified visibility */
+	//EVisibility ShouldBeVisible() const;
+
+	///** @return True if the user is scrolling by dragging the scroll bar thumb. */
+	//bool IsScrolling() const;
+
+
+	// UWidget interface
+	virtual void SyncronizeProperties() override;
+	// End of UWidget interface
+
+#if WITH_EDITOR
+	// UWidget interface
+	virtual const FSlateBrush* GetEditorIcon() override;
+	// End UWidget interface
+#endif
+
+protected:
+
+	TSharedPtr<class SScrollBar> MyScrollBar;
+
+protected:
+	// UWidget interface
+	virtual TSharedRef<SWidget> RebuildWidget() override;
+	// End of UWidget interface
+};
diff --git a/Engine/Source/Runtime/UMG/Public/UMG.h b/Engine/Source/Runtime/UMG/Public/UMG.h
--- a/Engine/Source/Runtime/UMG/Public/UMG.h
+++ b/Engine/Source/Runtime/UMG/Public/UMG.h
@@ -67,6 +67,7 @@
 
 #include "Spacer.h"
 #include "MenuAnchor.h"
+#include "ScrollBar.h"
 
 #include "Throbber.h"
 #include "CircularThrobber.h"
