UMG - Making the style header public.

diff --git a/Engine/Source/Runtime/UMG/Private/UMGStyle.cpp b/Engine/Source/Runtime/UMG/Private/UMGStyle.cpp
--- a/Engine/Source/Runtime/UMG/Private/UMGStyle.cpp
+++ b/Engine/Source/Runtime/UMG/Private/UMGStyle.cpp
@@ -45,7 +45,7 @@ TSharedRef< FSlateStyleSet > FUMGStyle::Create()
 	
 	Style->Set("Widget", new IMAGE_BRUSH(TEXT("Widget"), Icon16x16));
 	Style->Set("Widget.CheckBox", new IMAGE_BRUSH(TEXT("CheckBox"), Icon16x16));
-	Style->Set("Widget.Button", new IMAGE_BRUSH(TEXT("Button"), Icon16x16));
+	Style->Set("ClassIcon.Button", new IMAGE_BRUSH(TEXT("Button"), Icon16x16));
 	Style->Set("Widget.EditableTextBox", new IMAGE_BRUSH(TEXT("EditableTextBox"), Icon16x16));
 	Style->Set("Widget.EditableText", new IMAGE_BRUSH(TEXT("EditableText"), Icon16x16));
 	Style->Set("Widget.HorizontalBox", new IMAGE_BRUSH(TEXT("HorizontalBox"), Icon16x16));
diff --git a/Engine/Source/Runtime/UMG/Private/UMGStyle.h b/Engine/Source/Runtime/UMG/Private/UMGStyle.h
deleted file mode 100644
diff --git a/Engine/Source/Runtime/UMG/Public/UMGStyle.h b/Engine/Source/Runtime/UMG/Public/UMGStyle.h
new file mode 100644
--- /dev/null
+++ b/Engine/Source/Runtime/UMG/Public/UMGStyle.h
@@ -0,0 +1,31 @@
+// Copyright 1998-2014 Epic Games, Inc. All Rights Reserved.
+
+#pragma once
+
+#include "Slate.h"
+
+/**  */
+class UMG_API FUMGStyle
+{
+public:
+
+	static void Initialize();
+
+	static void Shutdown();
+
+	/** reloads textures used by slate renderer */
+	static void ReloadTextures();
+
+	/** @return The Slate style set for the Shooter game */
+	static const ISlateStyle& Get();
+
+	static FName GetStyleSetName();
+
+private:
+
+	static TSharedRef< class FSlateStyleSet > Create();
+
+private:
+
+	static TSharedPtr< class FSlateStyleSet > UMGStyleInstance;
+};
\ No newline at end of file
