Added dedicated asset for dashed outlines in UMG.

diff --git a/Engine/Source/Runtime/UMG/Private/Components/Widget.cpp b/Engine/Source/Runtime/UMG/Private/Components/Widget.cpp
--- a/Engine/Source/Runtime/UMG/Private/Components/Widget.cpp
+++ b/Engine/Source/Runtime/UMG/Private/Components/Widget.cpp
@@ -140,7 +140,7 @@ TSharedRef<SWidget> UWidget::BuildDesignTimeWidget(TSharedRef<SWidget> WrapWidge
 		[
 			SNew(SBorder)
 			.Visibility( EVisibility::HitTestInvisible )
-			.BorderImage( FCoreStyle::Get().GetBrush("FocusRectangle") )
+			.BorderImage(FUMGStyle::Get().GetBrush("MarchingAnts"))
 		];
 	}
 	else
diff --git a/Engine/Source/Runtime/UMG/Private/UMGStyle.cpp b/Engine/Source/Runtime/UMG/Private/UMGStyle.cpp
--- a/Engine/Source/Runtime/UMG/Private/UMGStyle.cpp
+++ b/Engine/Source/Runtime/UMG/Private/UMGStyle.cpp
@@ -43,6 +43,8 @@ TSharedRef< FSlateStyleSet > FUMGStyle::Create()
 	TSharedRef< FSlateStyleSet > Style = MakeShareable(new FSlateStyleSet("UMGStyle"));
 	Style->SetContentRoot(FPaths::EngineContentDir() / TEXT("Editor/Slate/UMG"));
 	
+	Style->Set("MarchingAnts", new BORDER_BRUSH( TEXT("NonMarchingAnts"), FMargin(0.25f), FLinearColor(1,1,1,0.5) ));
+
 	Style->Set("Widget", new IMAGE_BRUSH(TEXT("Widget"), Icon16x16));
 	Style->Set("Widget.CheckBox", new IMAGE_BRUSH(TEXT("CheckBox"), Icon16x16));
 	Style->Set("Widget.Button", new IMAGE_BRUSH(TEXT("Button"), Icon16x16));
