UMG Vertical and Horizontal boxes now default to fill alignments and auto-sizing.

diff --git a/Engine/Source/Runtime/UMG/Private/Components/VerticalBoxSlot.cpp b/Engine/Source/Runtime/UMG/Private/Components/VerticalBoxSlot.cpp
--- a/Engine/Source/Runtime/UMG/Private/Components/VerticalBoxSlot.cpp
+++ b/Engine/Source/Runtime/UMG/Private/Components/VerticalBoxSlot.cpp
@@ -9,8 +9,9 @@ UVerticalBoxSlot::UVerticalBoxSlot(const FPostConstructInitializeProperties& PCI
 	: Super(PCIP)
 	, Slot(NULL)
 {
-	HorizontalAlignment = HAlign_Left;
-	VerticalAlignment = VAlign_Top;
+	HorizontalAlignment = HAlign_Fill;
+	VerticalAlignment = VAlign_Fill;
+	Size = FSlateChildSize(ESlateSizeRule::Automatic);
 }
 
 void UVerticalBoxSlot::BuildSlot(TSharedRef<SVerticalBox> VerticalBox)
