UMG Vertical and Horizontal boxes now default to fill alignments and auto-sizing.

diff --git a/Engine/Source/Runtime/UMG/Private/Components/HorizontalBoxSlot.cpp b/Engine/Source/Runtime/UMG/Private/Components/HorizontalBoxSlot.cpp
--- a/Engine/Source/Runtime/UMG/Private/Components/HorizontalBoxSlot.cpp
+++ b/Engine/Source/Runtime/UMG/Private/Components/HorizontalBoxSlot.cpp
@@ -8,8 +8,9 @@
 UHorizontalBoxSlot::UHorizontalBoxSlot(const FPostConstructInitializeProperties& PCIP)
 	: Super(PCIP)
 {
-	HorizontalAlignment = HAlign_Left;
-	VerticalAlignment = VAlign_Top;
+	HorizontalAlignment = HAlign_Fill;
+	VerticalAlignment = VAlign_Fill;
+	Size = FSlateChildSize(ESlateSizeRule::Automatic);
 }
 
 void UHorizontalBoxSlot::BuildSlot(TSharedRef<SHorizontalBox> HorizontalBox)
