UMG - Adding a way to make a UImage material easily changable at runtime.  Added a radial bar example material.

diff --git a/Engine/Source/Runtime/SlateCore/Public/Styling/SlateBrush.h b/Engine/Source/Runtime/SlateCore/Public/Styling/SlateBrush.h
--- a/Engine/Source/Runtime/SlateCore/Public/Styling/SlateBrush.h
+++ b/Engine/Source/Runtime/SlateCore/Public/Styling/SlateBrush.h
@@ -159,6 +159,14 @@ public:
 		return ResourceObject;
 	}
 
+	/**
+	* Sets the UObject that represents the brush resource.
+	*/
+	void SetResourceObject(class UObject* InResourceObject)
+	{
+		ResourceObject = InResourceObject;
+	}
+
 	/**
 	 * Gets the brush's tint color.
 	 *
diff --git a/Engine/Source/Runtime/UMG/Private/Components/Image.cpp b/Engine/Source/Runtime/UMG/Private/Components/Image.cpp
--- a/Engine/Source/Runtime/UMG/Private/Components/Image.cpp
+++ b/Engine/Source/Runtime/UMG/Private/Components/Image.cpp
@@ -55,8 +55,48 @@ void UImage::SetImage(USlateBrushAsset* InImage)
 	}
 }
 
+UMaterialInstanceDynamic* UImage::GetDynamicMaterial()
+{
+	if ( !DynamicBrush.IsSet() )
+	{
+		const FSlateBrush* Brush = UImage::GetImageBrush();
+		FSlateBrush ClonedBrush = *Brush;
+		UObject* Resource = ClonedBrush.GetResourceObject();
+		UMaterialInterface* Material = Cast<UMaterialInterface>(Resource);
+
+		if ( Material )
+		{
+			UMaterialInstanceDynamic* DynamicMaterial = UMaterialInstanceDynamic::Create(Material, this);
+			ClonedBrush.SetResourceObject(DynamicMaterial);
+			
+			DynamicBrush = ClonedBrush;
+
+			if ( MyImage.IsValid() )
+			{
+				MyImage->SetImage(&DynamicBrush.GetValue());
+			}
+		}
+		else
+		{
+			//TODO UMG Log error for debugging.
+		}
+	}
+	else
+	{
+		UObject* Resource = DynamicBrush.GetValue().GetResourceObject();
+		return Cast<UMaterialInstanceDynamic>(Resource);
+	}
+
+	return NULL;
+}
+
 const FSlateBrush* UImage::GetImageBrush() const
 {
+	if ( DynamicBrush.IsSet() )
+	{
+		return &DynamicBrush.GetValue();
+	}
+
 	if ( Image == NULL )
 	{
 		SImage::FArguments ImageDefaults;
diff --git a/Engine/Source/Runtime/UMG/Public/Components/Image.h b/Engine/Source/Runtime/UMG/Public/Components/Image.h
--- a/Engine/Source/Runtime/UMG/Public/Components/Image.h
+++ b/Engine/Source/Runtime/UMG/Public/Components/Image.h
@@ -37,6 +37,10 @@ public:
 	UFUNCTION(BlueprintCallable, Category="Appearance")
 	void SetImage(USlateBrushAsset* InImage);
 
+	/**  */
+	UFUNCTION(BlueprintCallable, Category="Appearance")
+	UMaterialInstanceDynamic* GetDynamicMaterial();
+
 	// UWidget interface
 	virtual void SyncronizeProperties() override;
 	// End of UWidget interface
@@ -58,4 +62,6 @@ protected:
 
 protected:
 	TSharedPtr<SImage> MyImage;
+
+	TOptional<FSlateBrush> DynamicBrush;
 };
