UMG - Value property of the slider is now bindable.

diff --git a/Engine/Source/Runtime/UMG/Private/Components/Slider.cpp b/Engine/Source/Runtime/UMG/Private/Components/Slider.cpp
--- a/Engine/Source/Runtime/UMG/Private/Components/Slider.cpp
+++ b/Engine/Source/Runtime/UMG/Private/Components/Slider.cpp
@@ -28,11 +28,13 @@ TSharedRef<SWidget> USlider::RebuildWidget()
 void USlider::SyncronizeProperties()
 {
 	Super::SyncronizeProperties();
+
+	TAttribute<float> ValueBinding = OPTIONAL_BINDING(float, Value);
 	
 	MySlider->SetOrientation(Orientation);
 	MySlider->SetSliderBarColor(SliderBarColor);
 	MySlider->SetSliderHandleColor(SliderHandleColor);
-	MySlider->SetValue(Value);
+	MySlider->SetValue(ValueBinding);
 }
 
 void USlider::HandleOnValueChanged(float InValue)
diff --git a/Engine/Source/Runtime/UMG/Public/Components/Slider.h b/Engine/Source/Runtime/UMG/Public/Components/Slider.h
--- a/Engine/Source/Runtime/UMG/Public/Components/Slider.h
+++ b/Engine/Source/Runtime/UMG/Public/Components/Slider.h
@@ -21,6 +21,10 @@ class UMG_API USlider : public UWidget
 	UPROPERTY(EditDefaultsOnly, Category=Appearance)
 	float Value;
 
+	/** A bindable delegate to allow logic to drive the value of the widget */
+	UPROPERTY()
+	FGetFloat ValueDelegate;
+
 	/** The slider's orientation. */
 	UPROPERTY(EditDefaultsOnly, Category=Appearance)
 	TEnumAsByte<EOrientation> Orientation;
