UMG - Reducing the number of events that fire at design time for widgets, tick, on paint...etc no longer fire in the designer.  Will need to eventually create a 'simulate' toggle to allow users to see what it looks like.  For now users can build a simple sandbox level to test UIs.

diff --git a/Engine/Source/Runtime/UMG/Private/Slate/SObjectWidget.cpp b/Engine/Source/Runtime/UMG/Private/Slate/SObjectWidget.cpp
--- a/Engine/Source/Runtime/UMG/Private/Slate/SObjectWidget.cpp
+++ b/Engine/Source/Runtime/UMG/Private/Slate/SObjectWidget.cpp
@@ -23,7 +23,7 @@ void SObjectWidget::AddReferencedObjects(FReferenceCollector& Collector)
 
 void SObjectWidget::Tick(const FGeometry& AllottedGeometry, const double InCurrentTime, const float InDeltaTime)
 {
-	if ( WidgetObject )
+	if ( WidgetObject && !WidgetObject->IsDesignTime() )
 	{
 		return WidgetObject->Tick(AllottedGeometry, InDeltaTime);
 	}
@@ -38,7 +38,7 @@ int32 SObjectWidget::OnPaint(const FGeometry& AllottedGeometry, const FSlateRect
 {
 	int32 MaxLayer = SCompoundWidget::OnPaint(AllottedGeometry, MyClippingRect, OutDrawElements, LayerId, InWidgetStyle, bParentEnabled);
 
-	if (WidgetObject)
+	if ( WidgetObject && !WidgetObject->IsDesignTime() )
 	{
 		//const FGeometry& AllottedGeometry, const FSlateRect& MyClippingRect, FSlateWindowElementList& OutDrawElements, int32 LayerId, const FWidgetStyle& InWidgetStyle, bool bParentEnabled
 		FPaintContext Context(AllottedGeometry, MyClippingRect, OutDrawElements, LayerId, InWidgetStyle, bParentEnabled);
@@ -52,12 +52,20 @@ int32 SObjectWidget::OnPaint(const FGeometry& AllottedGeometry, const FSlateRect
 
 FReply SObjectWidget::OnKeyboardFocusReceived(const FGeometry& MyGeometry, const FKeyboardFocusEvent& InKeyboardFocusEvent)
 {
-	return WidgetObject->OnKeyboardFocusReceived(MyGeometry, InKeyboardFocusEvent).ToReply(WidgetObject->GetWidget());
+	if ( WidgetObject && !WidgetObject->IsDesignTime() )
+	{
+		return WidgetObject->OnKeyboardFocusReceived(MyGeometry, InKeyboardFocusEvent).ToReply(WidgetObject->GetWidget());
+	}
+
+	return FReply::Unhandled();
 }
 
 void SObjectWidget::OnKeyboardFocusLost(const FKeyboardFocusEvent& InKeyboardFocusEvent)
 {
-	WidgetObject->OnKeyboardFocusLost(InKeyboardFocusEvent);
+	if ( WidgetObject && !WidgetObject->IsDesignTime() )
+	{
+		WidgetObject->OnKeyboardFocusLost(InKeyboardFocusEvent);
+	}
 }
 
 void SObjectWidget::OnKeyboardFocusChanging(const FWeakWidgetPath& PreviousFocusPath, const FWidgetPath& NewWidgetPath)
diff --git a/Engine/Source/Runtime/UMG/Private/UserWidget.cpp b/Engine/Source/Runtime/UMG/Private/UserWidget.cpp
--- a/Engine/Source/Runtime/UMG/Private/UserWidget.cpp
+++ b/Engine/Source/Runtime/UMG/Private/UserWidget.cpp
@@ -211,8 +211,11 @@ TSharedRef<SWidget> UUserWidget::RebuildWidget()
 		WidgetToComponent.Add(Widget, Handle);
 	}
 
-	// Notify the widget that it has been constructed.
-	Construct();
+	if ( !IsDesignTime() )
+	{
+		// Notify the widget that it has been constructed.
+		Construct();
+	}
 
 	return UserRootWidget.ToSharedRef();
 }
