Minor fix (ignoring function arg, caught by CIS).

diff --git a/Engine/Source/Runtime/UMG/Private/Components/CircularThrobber.cpp b/Engine/Source/Runtime/UMG/Private/Components/CircularThrobber.cpp
--- a/Engine/Source/Runtime/UMG/Private/Components/CircularThrobber.cpp
+++ b/Engine/Source/Runtime/UMG/Private/Components/CircularThrobber.cpp
@@ -53,7 +53,7 @@ const FSlateBrush* UCircularThrobber::GetPieceBrush() const
 
 void UCircularThrobber::SetNumberOfPieces(int32 InNumberOfPieces)
 {
-	NumberOfPieces = NumberOfPieces;
+	NumberOfPieces = InNumberOfPieces;
 	if (MyCircularThrobber.IsValid())
 	{
 		MyCircularThrobber->SetNumPieces(InNumberOfPieces);
diff --git a/Engine/Source/Runtime/UMG/Private/Components/Throbber.cpp b/Engine/Source/Runtime/UMG/Private/Components/Throbber.cpp
--- a/Engine/Source/Runtime/UMG/Private/Components/Throbber.cpp
+++ b/Engine/Source/Runtime/UMG/Private/Components/Throbber.cpp
@@ -62,7 +62,7 @@ const FSlateBrush* UThrobber::GetPieceBrush() const
 
 void UThrobber::SetNumberOfPieces(int32 InNumberOfPieces)
 {
-	NumberOfPieces = NumberOfPieces;
+	NumberOfPieces = InNumberOfPieces;
 	if (MyThrobber.IsValid())
 	{
 		MyThrobber->SetNumPieces(InNumberOfPieces);
