A bit more CIS fixing.

diff --git a/Engine/Source/Runtime/UMG/Private/Components/Throbber.cpp b/Engine/Source/Runtime/UMG/Private/Components/Throbber.cpp
--- a/Engine/Source/Runtime/UMG/Private/Components/Throbber.cpp
+++ b/Engine/Source/Runtime/UMG/Private/Components/Throbber.cpp
@@ -71,7 +71,7 @@ void UThrobber::SetNumberOfPieces(int32 InNumberOfPieces)
 
 void UThrobber:: SetAnimateHorizontally(bool bInAnimateHorizontally)
 {
-	bInAnimateHorizontally = bInAnimateHorizontally;
+	bAnimateHorizontally = bInAnimateHorizontally;
 	if (MyThrobber.IsValid())
 	{
 		MyThrobber->SetAnimate(GetAnimation());
@@ -80,7 +80,7 @@ void UThrobber:: SetAnimateHorizontally(bool bInAnimateHorizontally)
 
 void UThrobber:: SetAnimateVertically(bool bInAnimateVertically)
 {
-	bInAnimateVertically = bInAnimateVertically;
+	bAnimateVertically = bInAnimateVertically;
 	if (MyThrobber.IsValid())
 	{
 		MyThrobber->SetAnimate(GetAnimation());
@@ -89,7 +89,7 @@ void UThrobber:: SetAnimateVertically(bool bInAnimateVertically)
 
 void UThrobber:: SetAnimateOpacity(bool bInAnimateOpacity)
 {
-	bInAnimateOpacity = bInAnimateOpacity;
+	bAnimateOpacity = bInAnimateOpacity;
 	if (MyThrobber.IsValid())
 	{
 		MyThrobber->SetAnimate(GetAnimation());
