UMG - Adding the new multiline editable text box to the available widgets.  Still hasn't been full wrapped, just the basics.

diff --git a/Engine/Source/Runtime/UMG/Private/Components/MultiLineEditableText.cpp b/Engine/Source/Runtime/UMG/Private/Components/MultiLineEditableText.cpp
new file mode 100644
--- /dev/null
+++ b/Engine/Source/Runtime/UMG/Private/Components/MultiLineEditableText.cpp
@@ -0,0 +1,115 @@
+// Copyright 1998-2014 Epic Games, Inc. All Rights Reserved.
+
+#include "UMGPrivatePCH.h"
+
+#define LOCTEXT_NAMESPACE "UMG"
+
+/////////////////////////////////////////////////////
+// UMultiLineEditableText
+
+UMultiLineEditableText::UMultiLineEditableText(const FPostConstructInitializeProperties& PCIP)
+	: Super(PCIP)
+{
+	ColorAndOpacity = FLinearColor::Black;
+
+	// HACK Special font initialization hack since there are no font assets yet for slate.
+	Font = FSlateFontInfo(TEXT("Slate/Fonts/Roboto-Bold.ttf"), 12);
+
+	// Grab other defaults from slate arguments.
+	SMultiLineEditableTextBox::FArguments Defaults;
+	IsReadOnly = Defaults._IsReadOnly.Get();
+	IsPassword = Defaults._IsReadOnly.Get();
+//	MinimumDesiredWidth = Defaults._MinDesiredWidth.Get();
+	IsCaretMovedWhenGainFocus = Defaults._IsCaretMovedWhenGainFocus.Get();
+	SelectAllTextWhenFocused = Defaults._SelectAllTextWhenFocused.Get();
+	RevertTextOnEscape = Defaults._RevertTextOnEscape.Get();
+	ClearKeyboardFocusOnCommit = Defaults._ClearKeyboardFocusOnCommit.Get();
+	SelectAllTextOnCommit = Defaults._SelectAllTextOnCommit.Get();
+}
+
+TSharedRef<SWidget> UMultiLineEditableText::RebuildWidget()
+{
+	FString FontPath = FPaths::EngineContentDir() / Font.FontName.ToString();
+	
+	SMultiLineEditableText::FArguments Defaults;
+	
+	//const FMultiLineEditableTextStyle* StylePtr = ( Style != NULL ) ? Style->GetStyle<FMultiLineEditableTextStyle>() : NULL;
+	//if ( StylePtr == NULL )
+	//{
+	//	StylePtr = Defaults._Style;
+	//}
+	
+	MyMultiLineEditableText = SNew(SMultiLineEditableText)
+	//.Style(StylePtr)
+	.Font(FSlateFontInfo(FontPath, Font.Size))
+//	.MinDesiredWidth(MinimumDesiredWidth)
+//	.IsCaretMovedWhenGainFocus(IsCaretMovedWhenGainFocus)
+//	.SelectAllTextWhenFocused(SelectAllTextWhenFocused)
+//	.RevertTextOnEscape(RevertTextOnEscape)
+//	.ClearKeyboardFocusOnCommit(ClearKeyboardFocusOnCommit)
+//	.SelectAllTextOnCommit(SelectAllTextOnCommit)
+//	.BackgroundImageSelected(BackgroundImageSelected ? TAttribute<const FSlateBrush*>(&BackgroundImageSelected->Brush) : TAttribute<const FSlateBrush*>())
+//	.BackgroundImageSelectionTarget(BackgroundImageSelectionTarget ? TAttribute<const FSlateBrush*>(&BackgroundImageSelectionTarget->Brush) : TAttribute<const FSlateBrush*>())
+//	.BackgroundImageComposing(BackgroundImageComposing ? TAttribute<const FSlateBrush*>(&BackgroundImageComposing->Brush) : TAttribute<const FSlateBrush*>())
+//	.CaretImage(CaretImage ? TAttribute<const FSlateBrush*>(&CaretImage->Brush) : TAttribute<const FSlateBrush*>())
+//	.OnTextChanged(BIND_UOBJECT_DELEGATE(FOnTextChanged, HandleOnTextChanged))
+//	.OnTextCommitted(BIND_UOBJECT_DELEGATE(FOnTextCommitted, HandleOnTextCommitted))
+	;
+	
+	return BuildDesignTimeWidget( MyMultiLineEditableText.ToSharedRef() );
+}
+
+void UMultiLineEditableText::SyncronizeProperties()
+{
+	Super::SyncronizeProperties();
+
+	MyMultiLineEditableText->SetText(Text);
+//	MyMultiLineEditableText->SetHintText(HintText);
+//	MyMultiLineEditableText->SetIsReadOnly(IsReadOnly);
+//	MyMultiLineEditableText->SetIsPassword(IsPassword);
+//	MyMultiLineEditableText->SetColorAndOpacity(ColorAndOpacity);
+
+	// TODO UMG Complete making all properties settable on SMultiLineEditableText
+}
+
+FText UMultiLineEditableText::GetText() const
+{
+	if ( MyMultiLineEditableText.IsValid() )
+	{
+		return MyMultiLineEditableText->GetText();
+	}
+
+	return Text;
+}
+
+void UMultiLineEditableText::SetText(FText InText)
+{
+	Text = InText;
+	if ( MyMultiLineEditableText.IsValid() )
+	{
+		MyMultiLineEditableText->SetText(Text);
+	}
+}
+
+void UMultiLineEditableText::HandleOnTextChanged(const FText& Text)
+{
+	OnTextChanged.Broadcast(Text);
+}
+
+void UMultiLineEditableText::HandleOnTextCommitted(const FText& Text, ETextCommit::Type CommitMethod)
+{
+	OnTextCommitted.Broadcast(Text, CommitMethod);
+}
+
+#if WITH_EDITOR
+
+const FSlateBrush* UMultiLineEditableText::GetEditorIcon()
+{
+	return FUMGStyle::Get().GetBrush("Widget.MultiLineEditableText");
+}
+
+#endif
+
+/////////////////////////////////////////////////////
+
+#undef LOCTEXT_NAMESPACE
diff --git a/Engine/Source/Runtime/UMG/Private/Components/MultiLineEditableTextBox.cpp b/Engine/Source/Runtime/UMG/Private/Components/MultiLineEditableTextBox.cpp
new file mode 100644
--- /dev/null
+++ b/Engine/Source/Runtime/UMG/Private/Components/MultiLineEditableTextBox.cpp
@@ -0,0 +1,123 @@
+// Copyright 1998-2014 Epic Games, Inc. All Rights Reserved.
+
+#include "UMGPrivatePCH.h"
+
+#define LOCTEXT_NAMESPACE "UMG"
+
+/////////////////////////////////////////////////////
+// UMultiLineEditableTextBox
+
+UMultiLineEditableTextBox::UMultiLineEditableTextBox(const FPostConstructInitializeProperties& PCIP)
+	: Super(PCIP)
+{
+	ForegroundColor = FLinearColor::Black;
+	BackgroundColor = FLinearColor::White;
+	ReadOnlyForegroundColor = FLinearColor::Black;
+
+	// HACK Special font initialization hack since there are no font assets yet for slate.
+	Font = FSlateFontInfo(TEXT("Slate/Fonts/Roboto-Bold.ttf"), 12);
+
+	// Grab other defaults from slate arguments.
+	SMultiLineEditableTextBox::FArguments Defaults;
+	IsReadOnly = Defaults._IsReadOnly.Get();
+	IsPassword = Defaults._IsReadOnly.Get();
+//	MinimumDesiredWidth = Defaults._MinDesiredWidth.Get();
+	Padding = Defaults._Padding.Get();
+	IsCaretMovedWhenGainFocus = Defaults._IsCaretMovedWhenGainFocus.Get();
+	SelectAllTextWhenFocused = Defaults._SelectAllTextWhenFocused.Get();
+	RevertTextOnEscape = Defaults._RevertTextOnEscape.Get();
+	ClearKeyboardFocusOnCommit = Defaults._ClearKeyboardFocusOnCommit.Get();
+	SelectAllTextOnCommit = Defaults._SelectAllTextOnCommit.Get();
+}
+
+TSharedRef<SWidget> UMultiLineEditableTextBox::RebuildWidget()
+{
+	FString FontPath = FPaths::EngineContentDir() / Font.FontName.ToString();
+
+	MyEditableTextBlock = SNew(SMultiLineEditableTextBox)
+		.Font(FSlateFontInfo(FontPath, Font.Size))
+		.ForegroundColor(ForegroundColor)
+		.BackgroundColor(BackgroundColor)
+		.ReadOnlyForegroundColor(ReadOnlyForegroundColor)
+//		.MinDesiredWidth(MinimumDesiredWidth)
+//		.Padding(Padding)
+//		.IsCaretMovedWhenGainFocus(IsCaretMovedWhenGainFocus)
+//		.SelectAllTextWhenFocused(SelectAllTextWhenFocused)
+//		.RevertTextOnEscape(RevertTextOnEscape)
+//		.ClearKeyboardFocusOnCommit(ClearKeyboardFocusOnCommit)
+//		.SelectAllTextOnCommit(SelectAllTextOnCommit)
+//		.OnTextChanged(BIND_UOBJECT_DELEGATE(FOnTextChanged, HandleOnTextChanged))
+//		.OnTextCommitted(BIND_UOBJECT_DELEGATE(FOnTextCommitted, HandleOnTextCommitted))
+		;
+
+	return MyEditableTextBlock.ToSharedRef();
+}
+
+void UMultiLineEditableTextBox::SyncronizeProperties()
+{
+	Super::SyncronizeProperties();
+
+	//const FMultiLineEditableTextBoxStyle* StylePtr = ( Style != NULL ) ? Style->GetStyle<FMultiLineEditableTextBoxStyle>() : NULL;
+	//if ( StylePtr )
+	//{
+	//	MyEditableTextBlock->SetStyle(StylePtr);
+	//}
+
+	MyEditableTextBlock->SetText(Text);
+//	MyEditableTextBlock->SetHintText(HintText);
+//	MyEditableTextBlock->SetIsReadOnly(IsReadOnly);
+//	MyEditableTextBlock->SetIsPassword(IsPassword);
+//	MyEditableTextBlock->SetColorAndOpacity(ColorAndOpacity);
+
+	// TODO UMG Complete making all properties settable on SMultiLineEditableTextBox
+}
+
+FText UMultiLineEditableTextBox::GetText() const
+{
+	if ( MyEditableTextBlock.IsValid() )
+	{
+		return MyEditableTextBlock->GetText();
+	}
+
+	return Text;
+}
+
+void UMultiLineEditableTextBox::SetText(FText InText)
+{
+	Text = InText;
+	if ( MyEditableTextBlock.IsValid() )
+	{
+		MyEditableTextBlock->SetText(Text);
+	}
+}
+
+void UMultiLineEditableTextBox::SetError(FText InError)
+{
+	if ( MyEditableTextBlock.IsValid() )
+	{
+		MyEditableTextBlock->SetError(InError);
+	}
+}
+
+void UMultiLineEditableTextBox::HandleOnTextChanged(const FText& Text)
+{
+	OnTextChanged.Broadcast(Text);
+}
+
+void UMultiLineEditableTextBox::HandleOnTextCommitted(const FText& Text, ETextCommit::Type CommitMethod)
+{
+	OnTextCommitted.Broadcast(Text, CommitMethod);
+}
+
+#if WITH_EDITOR
+
+const FSlateBrush* UMultiLineEditableTextBox::GetEditorIcon()
+{
+	return FUMGStyle::Get().GetBrush("Widget.MultiLineEditableTextBox");
+}
+
+#endif
+
+/////////////////////////////////////////////////////
+
+#undef LOCTEXT_NAMESPACE
diff --git a/Engine/Source/Runtime/UMG/Private/UMGStyle.cpp b/Engine/Source/Runtime/UMG/Private/UMGStyle.cpp
--- a/Engine/Source/Runtime/UMG/Private/UMGStyle.cpp
+++ b/Engine/Source/Runtime/UMG/Private/UMGStyle.cpp
@@ -64,6 +64,8 @@ TSharedRef< FSlateStyleSet > FUMGStyle::Create()
 	Style->Set("Widget.ScrollBar", new IMAGE_BRUSH(TEXT("ScrollBar"), Icon16x16));
 	Style->Set("Widget.UniformGrid", new IMAGE_BRUSH(TEXT("UniformGrid"), Icon16x16));
 	Style->Set("Widget.WidgetSwitcher", new IMAGE_BRUSH(TEXT("WidgetSwitcher"), Icon16x16));
+	Style->Set("Widget.MultiLineEditableText", new IMAGE_BRUSH(TEXT("MultiLineEditableText"), Icon16x16));
+	Style->Set("Widget.MultiLineEditableTextBox", new IMAGE_BRUSH(TEXT("MultiLineEditableTextBox"), Icon16x16));
 
 	Style->Set("Widget.UserWidget", new IMAGE_BRUSH(TEXT("UserWidget"), Icon16x16));
 
diff --git a/Engine/Source/Runtime/UMG/Public/Components/MultiLineEditableText.h b/Engine/Source/Runtime/UMG/Public/Components/MultiLineEditableText.h
new file mode 100644
--- /dev/null
+++ b/Engine/Source/Runtime/UMG/Public/Components/MultiLineEditableText.h
@@ -0,0 +1,122 @@
+// Copyright 1998-2014 Epic Games, Inc. All Rights Reserved.
+
+#pragma once
+
+#include "MultiLineEditableText.generated.h"
+
+/** Editable text box widget */
+UCLASS(meta=( Category="Common" ), ClassGroup=UserInterface)
+class UMG_API UMultiLineEditableText : public UWidget
+{
+	GENERATED_UCLASS_BODY()
+
+public:
+
+	DECLARE_DYNAMIC_MULTICAST_DELEGATE_OneParam(FOnMultiLineEditableTextChangedEvent, const FText&, Text);
+	DECLARE_DYNAMIC_MULTICAST_DELEGATE_TwoParams(FOnMultiLineEditableTextCommittedEvent, const FText&, Text, ETextCommit::Type, CommitMethod);
+
+public:
+
+	/** The text content for this editable text box widget */
+	UPROPERTY(EditDefaultsOnly, Category=Content)
+	FText Text;
+
+	/** Hint text that appears when there is no text in the text box */
+	UPROPERTY(EditDefaultsOnly, Category=Content)
+	FText HintText;
+
+	/** Text style */
+	UPROPERTY(EditDefaultsOnly, Category=Style, meta=( DisplayThumbnail = "true" ))
+	USlateWidgetStyleAsset* Style;
+
+	/** Background image for the selected text (overrides Style) */
+	UPROPERTY(EditDefaultsOnly, Category="Style", meta=( DisplayThumbnail = "true" ), AdvancedDisplay)
+	USlateBrushAsset* BackgroundImageSelected;
+
+	/** Background image for the selection targeting effect (overrides Style) */
+	UPROPERTY(EditDefaultsOnly, Category="Style", meta=( DisplayThumbnail = "true" ), AdvancedDisplay)
+	USlateBrushAsset* BackgroundImageSelectionTarget;
+
+	/** Background image for the composing text (overrides Style) */
+	UPROPERTY(EditDefaultsOnly, Category="Style", meta=( DisplayThumbnail = "true" ), AdvancedDisplay)
+	USlateBrushAsset* BackgroundImageComposing;
+
+	/** Image brush used for the caret (overrides Style) */
+	UPROPERTY(EditDefaultsOnly, Category="Style", meta=( DisplayThumbnail = "true" ), AdvancedDisplay)
+	USlateBrushAsset* CaretImage;
+
+	/** Font color and opacity (overrides Style) */
+	UPROPERTY(EditDefaultsOnly, Category=Appearance)
+	FSlateFontInfo Font;
+
+	/** Text color and opacity (overrides Style) */
+	UPROPERTY(EditDefaultsOnly, Category=Appearance)
+	FSlateColor ColorAndOpacity;
+
+	/** Sets whether this text box can actually be modified interactively by the user */
+	UPROPERTY(EditDefaultsOnly, Category=Appearance)
+	bool IsReadOnly;
+
+	/** Sets whether this text box is for storing a password */
+	UPROPERTY(EditDefaultsOnly, Category=Appearance)
+	bool IsPassword;
+
+	/** Minimum width that a text block should be */
+	UPROPERTY(EditDefaultsOnly, Category=Appearance)
+	float MinimumDesiredWidth;
+
+	/** Workaround as we loose focus when the auto completion closes. */
+	UPROPERTY(EditDefaultsOnly, Category=Behavior, AdvancedDisplay)
+	bool IsCaretMovedWhenGainFocus;
+
+	/** Whether to select all text when the user clicks to give focus on the widget */
+	UPROPERTY(EditDefaultsOnly, Category=Behavior, AdvancedDisplay)
+	bool SelectAllTextWhenFocused;
+
+	/** Whether to allow the user to back out of changes when they press the escape key */
+	UPROPERTY(EditDefaultsOnly, Category=Behavior, AdvancedDisplay)
+	bool RevertTextOnEscape;
+
+	/** Whether to clear keyboard focus when pressing enter to commit changes */
+	UPROPERTY(EditDefaultsOnly, Category=Behavior, AdvancedDisplay)
+	bool ClearKeyboardFocusOnCommit;
+
+	/** Whether to select all text when pressing enter to commit changes */
+	UPROPERTY(EditDefaultsOnly, Category=Behavior, AdvancedDisplay)
+	bool SelectAllTextOnCommit;
+
+	/** Called whenever the text is changed interactively by the user */
+	UPROPERTY(BlueprintAssignable, Category="Widget Event")
+	FOnMultiLineEditableTextChangedEvent OnTextChanged;
+
+	/** Called whenever the text is committed.  This happens when the user presses enter or the text box loses focus. */
+	UPROPERTY(BlueprintAssignable, Category="Widget Event")
+	FOnMultiLineEditableTextCommittedEvent OnTextCommitted;
+
+	/**  */
+	UFUNCTION(BlueprintCallable, Category="Widget")
+	FText GetText() const;
+
+	/**  */
+	UFUNCTION(BlueprintCallable, Category="Widget")
+	void SetText(FText InText);
+	
+	// UWidget interface
+	virtual void SyncronizeProperties() override;
+	// End of UWidget interface
+
+#if WITH_EDITOR
+	virtual const FSlateBrush* GetEditorIcon() override;
+#endif
+
+protected:
+	// UWidget interface
+	virtual TSharedRef<SWidget> RebuildWidget() override;
+	// End of UWidget
+
+	void HandleOnTextChanged(const FText& Text);
+	void HandleOnTextCommitted(const FText& Text, ETextCommit::Type CommitMethod);
+
+protected:
+	TSharedPtr<SMultiLineEditableText> MyMultiLineEditableText;
+};
diff --git a/Engine/Source/Runtime/UMG/Public/Components/MultiLineEditableTextBox.h b/Engine/Source/Runtime/UMG/Public/Components/MultiLineEditableTextBox.h
new file mode 100644
--- /dev/null
+++ b/Engine/Source/Runtime/UMG/Public/Components/MultiLineEditableTextBox.h
@@ -0,0 +1,124 @@
+// Copyright 1998-2014 Epic Games, Inc. All Rights Reserved.
+
+#pragma once
+
+#include "MultiLineEditableTextBox.generated.h"
+
+/** Editable text box widget */
+UCLASS(meta=( Category="Common" ), ClassGroup=UserInterface)
+class UMG_API UMultiLineEditableTextBox : public UWidget
+{
+	GENERATED_UCLASS_BODY()
+
+public:
+
+	DECLARE_DYNAMIC_MULTICAST_DELEGATE_OneParam(FOnMultiLineEditableTextBoxChangedEvent, const FText&, Text);
+	DECLARE_DYNAMIC_MULTICAST_DELEGATE_TwoParams(FOnMultiLineEditableTextBoxCommittedEvent, const FText&, Text, ETextCommit::Type, CommitMethod);
+
+public:
+
+	/** Style used for the text box */
+	UPROPERTY(EditDefaultsOnly, Category="Style", meta=( DisplayThumbnail = "true" ))
+	USlateWidgetStyleAsset* Style;
+
+	/** The text content for this editable text box widget */
+	UPROPERTY(EditDefaultsOnly, Category=Content)
+	FText Text;
+
+	/** Hint text that appears when there is no text in the text box */
+	UPROPERTY(EditDefaultsOnly, Category=Content)
+	FText HintText;
+
+	/** Font color and opacity (overrides Style) */
+	UPROPERTY(EditDefaultsOnly, Category=Appearance)
+	FSlateFontInfo Font;
+
+	/** Text color and opacity (overrides Style) */
+	UPROPERTY(EditDefaultsOnly, Category=Appearance)
+	FSlateColor ForegroundColor;
+
+	/** The color of the background/border around the editable text (overrides Style) */
+	UPROPERTY(EditDefaultsOnly, Category=Appearance)
+	FSlateColor BackgroundColor;
+
+	/** Text color and opacity when read-only (overrides Style) */
+	UPROPERTY(EditDefaultsOnly, Category=Appearance)
+	FSlateColor ReadOnlyForegroundColor;
+
+	/** Sets whether this text box can actually be modified interactively by the user */
+	UPROPERTY(EditDefaultsOnly, Category=Appearance)
+	bool IsReadOnly;
+
+	/** Sets whether this text box is for storing a password */
+	UPROPERTY(EditDefaultsOnly, Category=Appearance)
+	bool IsPassword;
+
+	/** Minimum width that a text block should be */
+	UPROPERTY(EditDefaultsOnly, Category=Appearance)
+	float MinimumDesiredWidth;
+
+	/** Padding between the box/border and the text widget inside (overrides Style) */
+	UPROPERTY(EditDefaultsOnly, Category=Appearance)
+	FMargin Padding;
+
+	/** Workaround as we loose focus when the auto completion closes. */
+	UPROPERTY(EditDefaultsOnly, Category=Behavior, AdvancedDisplay)
+	bool IsCaretMovedWhenGainFocus;
+
+	/** Whether to select all text when the user clicks to give focus on the widget */
+	UPROPERTY(EditDefaultsOnly, Category=Behavior, AdvancedDisplay)
+	bool SelectAllTextWhenFocused;
+
+	/** Whether to allow the user to back out of changes when they press the escape key */
+	UPROPERTY(EditDefaultsOnly, Category=Behavior, AdvancedDisplay)
+	bool RevertTextOnEscape;
+
+	/** Whether to clear keyboard focus when pressing enter to commit changes */
+	UPROPERTY(EditDefaultsOnly, Category=Behavior, AdvancedDisplay)
+	bool ClearKeyboardFocusOnCommit;
+
+	/** Whether to select all text when pressing enter to commit changes */
+	UPROPERTY(EditDefaultsOnly, Category=Behavior, AdvancedDisplay)
+	bool SelectAllTextOnCommit;
+
+	/** Called whenever the text is changed interactively by the user */
+	UPROPERTY(BlueprintAssignable, Category="Widget Event")
+	FOnMultiLineEditableTextBoxChangedEvent OnTextChanged;
+
+	/** Called whenever the text is committed.  This happens when the user presses enter or the text box loses focus. */
+	UPROPERTY(BlueprintAssignable, Category="Widget Event")
+	FOnMultiLineEditableTextBoxCommittedEvent OnTextCommitted;
+
+	/** Provide a alternative mechanism for error reporting. */
+	//SLATE_ARGUMENT(TSharedPtr<class IErrorReportingWidget>, ErrorReporting)
+
+	/**  */
+	UFUNCTION(BlueprintCallable, Category="Widget")
+	FText GetText() const;
+
+	/**  */
+	UFUNCTION(BlueprintCallable, Category="Widget")
+	void SetText(FText InText);
+
+	UFUNCTION(BlueprintCallable, Category="Widget")
+	void SetError(FText InError);
+
+	// UWidget interface
+	virtual void SyncronizeProperties() override;
+	// End of UWidget interface
+
+#if WITH_EDITOR
+	virtual const FSlateBrush* GetEditorIcon() override;
+#endif
+
+protected:
+	// UWidget interface
+	virtual TSharedRef<SWidget> RebuildWidget() override;
+	// End of UWidget
+
+	void HandleOnTextChanged(const FText& Text);
+	void HandleOnTextCommitted(const FText& Text, ETextCommit::Type CommitMethod);
+
+protected:
+	TSharedPtr<SMultiLineEditableTextBox> MyEditableTextBlock;
+};
diff --git a/Engine/Source/Runtime/UMG/Public/UMG.h b/Engine/Source/Runtime/UMG/Public/UMG.h
--- a/Engine/Source/Runtime/UMG/Public/UMG.h
+++ b/Engine/Source/Runtime/UMG/Public/UMG.h
@@ -60,9 +60,14 @@
 #include "TextBlock.h"
 #include "Border.h"
 #include "CheckBox.h"
+#include "Image.h"
+
 #include "EditableText.h"
 #include "EditableTextBox.h"
-#include "Image.h"
+
+#include "MultiLineEditableText.h"
+#include "MultiLineEditableTextBox.h"
+
 #include "Slider.h"
 #include "ProgressBar.h"
 #include "ComboBox.h"
