UMG - Exposing several panel functions as UFunctions.  Will allow you to dynamically add and mange children of a panel.

diff --git a/Engine/Source/Runtime/UMG/Public/Components/PanelWidget.h b/Engine/Source/Runtime/UMG/Public/Components/PanelWidget.h
--- a/Engine/Source/Runtime/UMG/Public/Components/PanelWidget.h
+++ b/Engine/Source/Runtime/UMG/Public/Components/PanelWidget.h
@@ -21,31 +21,42 @@ public:
 	//TODO UMG Add ways to make adding slots callable by blueprints.
 
 	/**  */
+	UFUNCTION(BlueprintCallable, Category="Appearance")
 	int32 GetChildrenCount() const;
 
 	/**  */
+	UFUNCTION(BlueprintCallable, Category="Appearance")
 	UWidget* GetChildAt(int32 Index) const;
 
 	/**  */
+	UFUNCTION(BlueprintCallable, Category="Appearance")
 	int32 GetChildIndex(UWidget* Content) const;
 
 	/**  */
+	UFUNCTION(BlueprintCallable, Category="Appearance")
 	bool RemoveChildAt(int32 Index);
 
 	/**  */
+	UFUNCTION(BlueprintCallable, Category="Appearance")
 	UPanelSlot* AddChild(UWidget* Content);
 
 	/**  */
+	UFUNCTION(BlueprintCallable, Category="Appearance")
 	void ReplaceChildAt(int32 Index, UWidget* Content);
 
 	/**  */
+	UFUNCTION(BlueprintCallable, Category="Appearance")
 	void InsertChildAt(int32 Index, UWidget* Content);
 
 	/**  */
+	UFUNCTION(BlueprintCallable, Category="Appearance")
 	bool RemoveChild(UWidget* Content);
 
 	/**  */
-	bool CanHaveMultipleChildren() const { return bCanHaveMultipleChildren; }
+	bool CanHaveMultipleChildren() const
+	{
+		return bCanHaveMultipleChildren;
+	}
 
 #if WITH_EDITOR
 	virtual void ConnectEditorData() override
