UMG - Fixing leaky PIE, by having a single point where GCRef container widgets when destroyed command their owned widget to destroy any native slate shared pointers they may be tracking that could be keeping other objects in memory.  All widgets and slots get an opportunity to handle this when their hierarchy is being destroyed.

diff --git a/Engine/Source/Runtime/UMG/Private/Components/Border.cpp b/Engine/Source/Runtime/UMG/Private/Components/Border.cpp
--- a/Engine/Source/Runtime/UMG/Private/Components/Border.cpp
+++ b/Engine/Source/Runtime/UMG/Private/Components/Border.cpp
@@ -29,6 +29,13 @@ UBorder::UBorder(const FPostConstructInitializeProperties& PCIP)
 	bShowEffectWhenDisabled = BorderDefaults._ShowEffectWhenDisabled.Get();
 }
 
+void UBorder::ReleaseNativeWidget()
+{
+	Super::ReleaseNativeWidget();
+
+	MyBorder.Reset();
+}
+
 TSharedRef<SWidget> UBorder::RebuildWidget()
 {
 	MyBorder = SNew(SBorder);
diff --git a/Engine/Source/Runtime/UMG/Private/Components/Button.cpp b/Engine/Source/Runtime/UMG/Private/Components/Button.cpp
--- a/Engine/Source/Runtime/UMG/Private/Components/Button.cpp
+++ b/Engine/Source/Runtime/UMG/Private/Components/Button.cpp
@@ -22,6 +22,13 @@ UButton::UButton(const FPostConstructInitializeProperties& PCIP)
 	ForegroundColor = FLinearColor::Black;
 }
 
+void UButton::ReleaseNativeWidget()
+{
+	Super::ReleaseNativeWidget();
+
+	MyButton.Reset();
+}
+
 TSharedRef<SWidget> UButton::RebuildWidget()
 {
 	MyButton = SNew(SButton);
diff --git a/Engine/Source/Runtime/UMG/Private/Components/ButtonSlot.cpp b/Engine/Source/Runtime/UMG/Private/Components/ButtonSlot.cpp
--- a/Engine/Source/Runtime/UMG/Private/Components/ButtonSlot.cpp
+++ b/Engine/Source/Runtime/UMG/Private/Components/ButtonSlot.cpp
@@ -14,6 +14,13 @@ UButtonSlot::UButtonSlot(const FPostConstructInitializeProperties& PCIP)
 	VerticalAlignment = VAlign_Center;
 }
 
+void UButtonSlot::ReleaseNativeWidget()
+{
+	Super::ReleaseNativeWidget();
+
+	Button.Reset();
+}
+
 void UButtonSlot::BuildSlot(TSharedRef<SButton> InButton)
 {
 	Button = InButton;
diff --git a/Engine/Source/Runtime/UMG/Private/Components/CanvasPanel.cpp b/Engine/Source/Runtime/UMG/Private/Components/CanvasPanel.cpp
--- a/Engine/Source/Runtime/UMG/Private/Components/CanvasPanel.cpp
+++ b/Engine/Source/Runtime/UMG/Private/Components/CanvasPanel.cpp
@@ -49,6 +49,13 @@ UCanvasPanel::UCanvasPanel(const FPostConstructInitializeProperties& PCIP)
 	Visiblity = UWidget::ConvertRuntimeToSerializedVisiblity(Defaults._Visibility.Get());
 }
 
+void UCanvasPanel::ReleaseNativeWidget()
+{
+	Super::ReleaseNativeWidget();
+
+	MyCanvas.Reset();
+}
+
 UClass* UCanvasPanel::GetSlotClass() const
 {
 	return UCanvasPanelSlot::StaticClass();
diff --git a/Engine/Source/Runtime/UMG/Private/Components/CanvasPanelSlot.cpp b/Engine/Source/Runtime/UMG/Private/Components/CanvasPanelSlot.cpp
--- a/Engine/Source/Runtime/UMG/Private/Components/CanvasPanelSlot.cpp
+++ b/Engine/Source/Runtime/UMG/Private/Components/CanvasPanelSlot.cpp
@@ -15,6 +15,13 @@ UCanvasPanelSlot::UCanvasPanelSlot(const FPostConstructInitializeProperties& PCI
 	ZOrder = 0;
 }
 
+void UCanvasPanelSlot::ReleaseNativeWidget()
+{
+	Super::ReleaseNativeWidget();
+
+	Slot = NULL;
+}
+
 void UCanvasPanelSlot::BuildSlot(TSharedRef<SConstraintCanvas> Canvas)
 {
 	Slot = &Canvas->AddSlot()
diff --git a/Engine/Source/Runtime/UMG/Private/Components/CheckBox.cpp b/Engine/Source/Runtime/UMG/Private/Components/CheckBox.cpp
--- a/Engine/Source/Runtime/UMG/Private/Components/CheckBox.cpp
+++ b/Engine/Source/Runtime/UMG/Private/Components/CheckBox.cpp
@@ -21,6 +21,13 @@ UCheckBox::UCheckBox(const FPostConstructInitializeProperties& PCIP)
 	BorderBackgroundColor = FLinearColor::White;
 }
 
+void UCheckBox::ReleaseNativeWidget()
+{
+	Super::ReleaseNativeWidget();
+
+	MyCheckbox.Reset();
+}
+
 TSharedRef<SWidget> UCheckBox::RebuildWidget()
 {
 	TOptional<FSlateSound> OptionalCheckedSound;
diff --git a/Engine/Source/Runtime/UMG/Private/Components/CircularThrobber.cpp b/Engine/Source/Runtime/UMG/Private/Components/CircularThrobber.cpp
--- a/Engine/Source/Runtime/UMG/Private/Components/CircularThrobber.cpp
+++ b/Engine/Source/Runtime/UMG/Private/Components/CircularThrobber.cpp
@@ -18,6 +18,13 @@ UCircularThrobber::UCircularThrobber(const FPostConstructInitializeProperties& P
 	Radius = DefaultArgs._Radius;
 }
 
+void UCircularThrobber::ReleaseNativeWidget()
+{
+	Super::ReleaseNativeWidget();
+
+	MyCircularThrobber.Reset();
+}
+
 TSharedRef<SWidget> UCircularThrobber::RebuildWidget()
 {
 	SCircularThrobber::FArguments DefaultArgs;
diff --git a/Engine/Source/Runtime/UMG/Private/Components/ComboBox.cpp b/Engine/Source/Runtime/UMG/Private/Components/ComboBox.cpp
--- a/Engine/Source/Runtime/UMG/Private/Components/ComboBox.cpp
+++ b/Engine/Source/Runtime/UMG/Private/Components/ComboBox.cpp
@@ -12,6 +12,11 @@ UComboBox::UComboBox(const FPostConstructInitializeProperties& PCIP)
 {
 }
 
+void UComboBox::ReleaseNativeWidget()
+{
+	Super::ReleaseNativeWidget();
+}
+
 TSharedRef<SWidget> UComboBox::RebuildWidget()
 {
 	TSharedRef< SComboBox<UObject*> > NewComboBox =
diff --git a/Engine/Source/Runtime/UMG/Private/Components/EditableText.cpp b/Engine/Source/Runtime/UMG/Private/Components/EditableText.cpp
--- a/Engine/Source/Runtime/UMG/Private/Components/EditableText.cpp
+++ b/Engine/Source/Runtime/UMG/Private/Components/EditableText.cpp
@@ -27,6 +27,13 @@ UEditableText::UEditableText(const FPostConstructInitializeProperties& PCIP)
 	SelectAllTextOnCommit = Defaults._SelectAllTextOnCommit.Get();
 }
 
+void UEditableText::ReleaseNativeWidget()
+{
+	Super::ReleaseNativeWidget();
+
+	MyEditableText.Reset();
+}
+
 TSharedRef<SWidget> UEditableText::RebuildWidget()
 {
 	FString FontPath = FPaths::EngineContentDir() / Font.FontName.ToString();
diff --git a/Engine/Source/Runtime/UMG/Private/Components/EditableTextBox.cpp b/Engine/Source/Runtime/UMG/Private/Components/EditableTextBox.cpp
--- a/Engine/Source/Runtime/UMG/Private/Components/EditableTextBox.cpp
+++ b/Engine/Source/Runtime/UMG/Private/Components/EditableTextBox.cpp
@@ -30,6 +30,13 @@ UEditableTextBox::UEditableTextBox(const FPostConstructInitializeProperties& PCI
 	SelectAllTextOnCommit = Defaults._SelectAllTextOnCommit.Get();
 }
 
+void UEditableTextBox::ReleaseNativeWidget()
+{
+	Super::ReleaseNativeWidget();
+
+	MyEditableTextBlock.Reset();
+}
+
 TSharedRef<SWidget> UEditableTextBox::RebuildWidget()
 {
 	FString FontPath = FPaths::EngineContentDir() / Font.FontName.ToString();
diff --git a/Engine/Source/Runtime/UMG/Private/Components/HorizontalBox.cpp b/Engine/Source/Runtime/UMG/Private/Components/HorizontalBox.cpp
--- a/Engine/Source/Runtime/UMG/Private/Components/HorizontalBox.cpp
+++ b/Engine/Source/Runtime/UMG/Private/Components/HorizontalBox.cpp
@@ -14,6 +14,13 @@ UHorizontalBox::UHorizontalBox(const FPostConstructInitializeProperties& PCIP)
 	Visiblity = UWidget::ConvertRuntimeToSerializedVisiblity(Defaults._Visibility.Get());
 }
 
+void UHorizontalBox::ReleaseNativeWidget()
+{
+	Super::ReleaseNativeWidget();
+
+	MyHorizontalBox.Reset();
+}
+
 UClass* UHorizontalBox::GetSlotClass() const
 {
 	return UHorizontalBoxSlot::StaticClass();
diff --git a/Engine/Source/Runtime/UMG/Private/Components/HorizontalBoxSlot.cpp b/Engine/Source/Runtime/UMG/Private/Components/HorizontalBoxSlot.cpp
--- a/Engine/Source/Runtime/UMG/Private/Components/HorizontalBoxSlot.cpp
+++ b/Engine/Source/Runtime/UMG/Private/Components/HorizontalBoxSlot.cpp
@@ -13,6 +13,13 @@ UHorizontalBoxSlot::UHorizontalBoxSlot(const FPostConstructInitializeProperties&
 	Size = FSlateChildSize(ESlateSizeRule::Automatic);
 }
 
+void UHorizontalBoxSlot::ReleaseNativeWidget()
+{
+	Super::ReleaseNativeWidget();
+
+	Slot = NULL;
+}
+
 void UHorizontalBoxSlot::BuildSlot(TSharedRef<SHorizontalBox> HorizontalBox)
 {
 	Slot = &HorizontalBox->AddSlot()
diff --git a/Engine/Source/Runtime/UMG/Private/Components/Image.cpp b/Engine/Source/Runtime/UMG/Private/Components/Image.cpp
--- a/Engine/Source/Runtime/UMG/Private/Components/Image.cpp
+++ b/Engine/Source/Runtime/UMG/Private/Components/Image.cpp
@@ -15,6 +15,13 @@ UImage::UImage(const FPostConstructInitializeProperties& PCIP)
 {
 }
 
+void UImage::ReleaseNativeWidget()
+{
+	Super::ReleaseNativeWidget();
+
+	MyImage.Reset();
+}
+
 TSharedRef<SWidget> UImage::RebuildWidget()
 {
 	MyImage = SNew(SImage);
diff --git a/Engine/Source/Runtime/UMG/Private/Components/MenuAnchor.cpp b/Engine/Source/Runtime/UMG/Private/Components/MenuAnchor.cpp
--- a/Engine/Source/Runtime/UMG/Private/Components/MenuAnchor.cpp
+++ b/Engine/Source/Runtime/UMG/Private/Components/MenuAnchor.cpp
@@ -13,6 +13,13 @@ UMenuAnchor::UMenuAnchor(const FPostConstructInitializeProperties& PCIP)
 	Placement = MenuPlacement_ComboBox;
 }
 
+void UMenuAnchor::ReleaseNativeWidget()
+{
+	Super::ReleaseNativeWidget();
+
+	MyMenuAnchor.Reset();
+}
+
 TSharedRef<SWidget> UMenuAnchor::RebuildWidget()
 {
 	MyMenuAnchor = SNew(SMenuAnchor)
diff --git a/Engine/Source/Runtime/UMG/Private/Components/Overlay.cpp b/Engine/Source/Runtime/UMG/Private/Components/Overlay.cpp
--- a/Engine/Source/Runtime/UMG/Private/Components/Overlay.cpp
+++ b/Engine/Source/Runtime/UMG/Private/Components/Overlay.cpp
@@ -14,6 +14,13 @@ UOverlay::UOverlay(const FPostConstructInitializeProperties& PCIP)
 	Visiblity = UWidget::ConvertRuntimeToSerializedVisiblity(Defaults._Visibility.Get());
 }
 
+void UOverlay::ReleaseNativeWidget()
+{
+	Super::ReleaseNativeWidget();
+
+	MyOverlay.Reset();
+}
+
 UClass* UOverlay::GetSlotClass() const
 {
 	return UOverlaySlot::StaticClass();
diff --git a/Engine/Source/Runtime/UMG/Private/Components/OverlaySlot.cpp b/Engine/Source/Runtime/UMG/Private/Components/OverlaySlot.cpp
--- a/Engine/Source/Runtime/UMG/Private/Components/OverlaySlot.cpp
+++ b/Engine/Source/Runtime/UMG/Private/Components/OverlaySlot.cpp
@@ -12,6 +12,13 @@ UOverlaySlot::UOverlaySlot(const FPostConstructInitializeProperties& PCIP)
 	VerticalAlignment = VAlign_Top;
 }
 
+void UOverlaySlot::ReleaseNativeWidget()
+{
+	Super::ReleaseNativeWidget();
+
+	Slot = NULL;
+}
+
 void UOverlaySlot::BuildSlot(TSharedRef<SOverlay> Overlay)
 {
 	Slot = &Overlay->AddSlot()
diff --git a/Engine/Source/Runtime/UMG/Private/Components/PanelSlot.cpp b/Engine/Source/Runtime/UMG/Private/Components/PanelSlot.cpp
--- a/Engine/Source/Runtime/UMG/Private/Components/PanelSlot.cpp
+++ b/Engine/Source/Runtime/UMG/Private/Components/PanelSlot.cpp
@@ -44,3 +44,13 @@ bool UPanelSlot::IsDesignTime() const
 {
 	return Parent->IsDesignTime();
 }
+
+void UPanelSlot::ReleaseNativeWidget()
+{
+	Super::ReleaseNativeWidget();
+
+	if ( Content )
+	{
+		Content->ReleaseNativeWidget();
+	}
+}
\ No newline at end of file
diff --git a/Engine/Source/Runtime/UMG/Private/Components/PanelWidget.cpp b/Engine/Source/Runtime/UMG/Private/Components/PanelWidget.cpp
--- a/Engine/Source/Runtime/UMG/Private/Components/PanelWidget.cpp
+++ b/Engine/Source/Runtime/UMG/Private/Components/PanelWidget.cpp
@@ -18,6 +18,19 @@ UPanelWidget::UPanelWidget(const FPostConstructInitializeProperties& PCIP)
 {
 }
 
+void UPanelWidget::ReleaseNativeWidget()
+{
+	Super::ReleaseNativeWidget();
+
+	for ( int32 SlotIndex = 0; SlotIndex < Slots.Num(); SlotIndex++ )
+	{
+		if ( Slots[SlotIndex]->Content != NULL )
+		{
+			Slots[SlotIndex]->ReleaseNativeWidget();
+		}
+	}
+}
+
 int32 UPanelWidget::GetChildrenCount() const
 {
 	return Slots.Num();
diff --git a/Engine/Source/Runtime/UMG/Private/Components/ScrollBox.cpp b/Engine/Source/Runtime/UMG/Private/Components/ScrollBox.cpp
--- a/Engine/Source/Runtime/UMG/Private/Components/ScrollBox.cpp
+++ b/Engine/Source/Runtime/UMG/Private/Components/ScrollBox.cpp
@@ -16,6 +16,13 @@ UScrollBox::UScrollBox(const FPostConstructInitializeProperties& PCIP)
 	Visiblity = UWidget::ConvertRuntimeToSerializedVisiblity(Defaults._Visibility.Get());
 }
 
+void UScrollBox::ReleaseNativeWidget()
+{
+	Super::ReleaseNativeWidget();
+
+	MyScrollBox.Reset();
+}
+
 UClass* UScrollBox::GetSlotClass() const
 {
 	return UScrollBoxSlot::StaticClass();
diff --git a/Engine/Source/Runtime/UMG/Private/Components/ScrollBoxSlot.cpp b/Engine/Source/Runtime/UMG/Private/Components/ScrollBoxSlot.cpp
--- a/Engine/Source/Runtime/UMG/Private/Components/ScrollBoxSlot.cpp
+++ b/Engine/Source/Runtime/UMG/Private/Components/ScrollBoxSlot.cpp
@@ -45,3 +45,9 @@ void UScrollBoxSlot::SyncronizeProperties()
 	SetPadding(Padding);
 	SetHorizontalAlignment(HorizontalAlignment);
 }
+
+void UScrollBoxSlot::ReleaseNativeWidget()
+{
+	Super::ReleaseNativeWidget();
+	Slot = NULL;
+}
\ No newline at end of file
diff --git a/Engine/Source/Runtime/UMG/Private/Components/UniformGridPanel.cpp b/Engine/Source/Runtime/UMG/Private/Components/UniformGridPanel.cpp
--- a/Engine/Source/Runtime/UMG/Private/Components/UniformGridPanel.cpp
+++ b/Engine/Source/Runtime/UMG/Private/Components/UniformGridPanel.cpp
@@ -16,6 +16,13 @@ UUniformGridPanel::UUniformGridPanel(const FPostConstructInitializeProperties& P
 	Visiblity = UWidget::ConvertRuntimeToSerializedVisiblity(Defaults._Visibility.Get());
 }
 
+void UUniformGridPanel::ReleaseNativeWidget()
+{
+	Super::ReleaseNativeWidget();
+
+	MyUniformGridPanel.Reset();
+}
+
 UClass* UUniformGridPanel::GetSlotClass() const
 {
 	return UUniformGridSlot::StaticClass();
diff --git a/Engine/Source/Runtime/UMG/Private/Components/UniformGridSlot.cpp b/Engine/Source/Runtime/UMG/Private/Components/UniformGridSlot.cpp
--- a/Engine/Source/Runtime/UMG/Private/Components/UniformGridSlot.cpp
+++ b/Engine/Source/Runtime/UMG/Private/Components/UniformGridSlot.cpp
@@ -13,6 +13,13 @@ UUniformGridSlot::UUniformGridSlot(const FPostConstructInitializeProperties& PCI
 	VerticalAlignment = VAlign_Top;
 }
 
+void UUniformGridSlot::ReleaseNativeWidget()
+{
+	Super::ReleaseNativeWidget();
+
+	Slot = NULL;
+}
+
 void UUniformGridSlot::BuildSlot(TSharedRef<SUniformGridPanel> GridPanel)
 {
 	Slot = &GridPanel->AddSlot(Column, Row)
diff --git a/Engine/Source/Runtime/UMG/Private/Components/VerticalBox.cpp b/Engine/Source/Runtime/UMG/Private/Components/VerticalBox.cpp
--- a/Engine/Source/Runtime/UMG/Private/Components/VerticalBox.cpp
+++ b/Engine/Source/Runtime/UMG/Private/Components/VerticalBox.cpp
@@ -14,6 +14,13 @@ UVerticalBox::UVerticalBox(const FPostConstructInitializeProperties& PCIP)
 	Visiblity = UWidget::ConvertRuntimeToSerializedVisiblity( Defaults._Visibility.Get() );
 }
 
+void UVerticalBox::ReleaseNativeWidget()
+{
+	Super::ReleaseNativeWidget();
+
+	MyVerticalBox.Reset();
+}
+
 UClass* UVerticalBox::GetSlotClass() const
 {
 	return UVerticalBoxSlot::StaticClass();
diff --git a/Engine/Source/Runtime/UMG/Private/Components/VerticalBoxSlot.cpp b/Engine/Source/Runtime/UMG/Private/Components/VerticalBoxSlot.cpp
--- a/Engine/Source/Runtime/UMG/Private/Components/VerticalBoxSlot.cpp
+++ b/Engine/Source/Runtime/UMG/Private/Components/VerticalBoxSlot.cpp
@@ -14,6 +14,13 @@ UVerticalBoxSlot::UVerticalBoxSlot(const FPostConstructInitializeProperties& PCI
 	Size = FSlateChildSize(ESlateSizeRule::Automatic);
 }
 
+void UVerticalBoxSlot::ReleaseNativeWidget()
+{
+	Super::ReleaseNativeWidget();
+
+	Slot = NULL;
+}
+
 void UVerticalBoxSlot::BuildSlot(TSharedRef<SVerticalBox> VerticalBox)
 {
 	Slot = &VerticalBox->AddSlot()
diff --git a/Engine/Source/Runtime/UMG/Private/Components/Visual.cpp b/Engine/Source/Runtime/UMG/Private/Components/Visual.cpp
--- a/Engine/Source/Runtime/UMG/Private/Components/Visual.cpp
+++ b/Engine/Source/Runtime/UMG/Private/Components/Visual.cpp
@@ -9,3 +9,7 @@ UVisual::UVisual(const FPostConstructInitializeProperties& PCIP)
 	: Super(PCIP)
 {
 }
+
+void UVisual::ReleaseNativeWidget()
+{
+}
diff --git a/Engine/Source/Runtime/UMG/Private/Components/Widget.cpp b/Engine/Source/Runtime/UMG/Private/Components/Widget.cpp
--- a/Engine/Source/Runtime/UMG/Private/Components/Widget.cpp
+++ b/Engine/Source/Runtime/UMG/Private/Components/Widget.cpp
@@ -121,6 +121,11 @@ TSharedRef<SWidget> UWidget::GetWidget() const
 	}
 }
 
+TSharedPtr<SWidget> UWidget::GetCachedWidget() const
+{
+	return MyWidget.Pin();
+}
+
 TSharedRef<SWidget> UWidget::BuildDesignTimeWidget(TSharedRef<SWidget> WrapWidget)
 {
 	if (IsDesignTime())
diff --git a/Engine/Source/Runtime/UMG/Private/Components/WidgetSwitcher.cpp b/Engine/Source/Runtime/UMG/Private/Components/WidgetSwitcher.cpp
--- a/Engine/Source/Runtime/UMG/Private/Components/WidgetSwitcher.cpp
+++ b/Engine/Source/Runtime/UMG/Private/Components/WidgetSwitcher.cpp
@@ -16,6 +16,13 @@ UWidgetSwitcher::UWidgetSwitcher(const FPostConstructInitializeProperties& PCIP)
 	Visiblity = UWidget::ConvertRuntimeToSerializedVisiblity(Defaults._Visibility.Get());
 }
 
+void UWidgetSwitcher::ReleaseNativeWidget()
+{
+	Super::ReleaseNativeWidget();
+
+	MyWidgetSwitcher.Reset();
+}
+
 int32 UWidgetSwitcher::GetNumWidgets() const
 {
 	if ( MyWidgetSwitcher.IsValid() )
diff --git a/Engine/Source/Runtime/UMG/Private/Components/WidgetSwitcherSlot.cpp b/Engine/Source/Runtime/UMG/Private/Components/WidgetSwitcherSlot.cpp
--- a/Engine/Source/Runtime/UMG/Private/Components/WidgetSwitcherSlot.cpp
+++ b/Engine/Source/Runtime/UMG/Private/Components/WidgetSwitcherSlot.cpp
@@ -13,6 +13,13 @@ UWidgetSwitcherSlot::UWidgetSwitcherSlot(const FPostConstructInitializePropertie
 	VerticalAlignment = VAlign_Fill;
 }
 
+void UWidgetSwitcherSlot::ReleaseNativeWidget()
+{
+	Super::ReleaseNativeWidget();
+
+	Slot = NULL;
+}
+
 void UWidgetSwitcherSlot::BuildSlot(TSharedRef<SWidgetSwitcher> WidgetSwitcher)
 {
 	Slot = &WidgetSwitcher->AddSlot()
diff --git a/Engine/Source/Runtime/UMG/Private/Slate/SObjectWidget.cpp b/Engine/Source/Runtime/UMG/Private/Slate/SObjectWidget.cpp
--- a/Engine/Source/Runtime/UMG/Private/Slate/SObjectWidget.cpp
+++ b/Engine/Source/Runtime/UMG/Private/Slate/SObjectWidget.cpp
@@ -14,6 +14,10 @@ void SObjectWidget::Construct(const FArguments& InArgs, UUserWidget* InWidgetObj
 
 SObjectWidget::~SObjectWidget(void)
 {
+	if ( WidgetObject )
+	{
+		WidgetObject->ReleaseNativeWidget();
+	}
 }
 
 void SObjectWidget::AddReferencedObjects(FReferenceCollector& Collector)
@@ -54,7 +58,11 @@ FReply SObjectWidget::OnKeyboardFocusReceived(const FGeometry& MyGeometry, const
 {
 	if ( WidgetObject && !WidgetObject->IsDesignTime() )
 	{
-		return WidgetObject->OnKeyboardFocusReceived(MyGeometry, InKeyboardFocusEvent).ToReply(WidgetObject->GetWidget());
+		TSharedPtr<SWidget> SlateWidget = WidgetObject->GetCachedWidget();
+		if ( SlateWidget.IsValid() )
+		{
+			return WidgetObject->OnKeyboardFocusReceived(MyGeometry, InKeyboardFocusEvent).ToReply(SlateWidget.ToSharedRef());
+		}
 	}
 
 	return FReply::Unhandled();
@@ -75,49 +83,98 @@ void SObjectWidget::OnKeyboardFocusChanging(const FWeakWidgetPath& PreviousFocus
 
 FReply SObjectWidget::OnKeyChar(const FGeometry& MyGeometry, const FCharacterEvent& InCharacterEvent)
 {
-	return WidgetObject->OnKeyChar(MyGeometry, InCharacterEvent).ToReply(WidgetObject->GetWidget());
+	TSharedPtr<SWidget> SlateWidget = WidgetObject->GetCachedWidget();
+	if ( SlateWidget.IsValid() )
+	{
+		return WidgetObject->OnKeyChar(MyGeometry, InCharacterEvent).ToReply(SlateWidget.ToSharedRef());
+	}
+
+	return FReply::Unhandled();
 }
 
 FReply SObjectWidget::OnPreviewKeyDown(const FGeometry& MyGeometry, const FKeyboardEvent& InKeyboardEvent)
 {
-	return WidgetObject->OnPreviewKeyDown(MyGeometry, InKeyboardEvent).ToReply(WidgetObject->GetWidget());
+	TSharedPtr<SWidget> SlateWidget = WidgetObject->GetCachedWidget();
+	if ( SlateWidget.IsValid() )
+	{
+		return WidgetObject->OnPreviewKeyDown(MyGeometry, InKeyboardEvent).ToReply(SlateWidget.ToSharedRef());
+	}
+
+	return FReply::Unhandled();
 }
 
 FReply SObjectWidget::OnKeyDown(const FGeometry& MyGeometry, const FKeyboardEvent& InKeyboardEvent)
 {
-	return WidgetObject->OnKeyDown(MyGeometry, InKeyboardEvent).ToReply(WidgetObject->GetWidget());
+	TSharedPtr<SWidget> SlateWidget = WidgetObject->GetCachedWidget();
+	if ( SlateWidget.IsValid() )
+	{
+		return WidgetObject->OnKeyDown(MyGeometry, InKeyboardEvent).ToReply(SlateWidget.ToSharedRef());
+	}
+
+	return FReply::Unhandled();
 }
 
 FReply SObjectWidget::OnKeyUp(const FGeometry& MyGeometry, const FKeyboardEvent& InKeyboardEvent)
 {
-	return WidgetObject->OnKeyUp(MyGeometry, InKeyboardEvent).ToReply(WidgetObject->GetWidget());
+	TSharedPtr<SWidget> SlateWidget = WidgetObject->GetCachedWidget();
+	if ( SlateWidget.IsValid() )
+	{
+		return WidgetObject->OnKeyUp(MyGeometry, InKeyboardEvent).ToReply(SlateWidget.ToSharedRef());
+	}
+
+	return FReply::Unhandled();
 }
 
 FReply SObjectWidget::OnMouseButtonDown(const FGeometry& MyGeometry, const FPointerEvent& MouseEvent)
 {
-	FReply Reply = WidgetObject->OnMouseButtonDown(MyGeometry, MouseEvent).ToReply(WidgetObject->GetWidget());
-	//TODO UMG Figure out how to let the user more easily manage when drags start.
-	if ( MouseEvent.GetEffectingButton() == EKeys::LeftMouseButton )
+	TSharedPtr<SWidget> SlateWidget = WidgetObject->GetCachedWidget();
+	if ( SlateWidget.IsValid() )
 	{
-		return Reply.DetectDrag(WidgetObject->GetWidget(), EKeys::LeftMouseButton);
+		FReply Reply = WidgetObject->OnMouseButtonDown(MyGeometry, MouseEvent).ToReply(SlateWidget.ToSharedRef());
+
+		//TODO UMG Figure out how to let the user more easily manage when drags start.
+		if ( MouseEvent.GetEffectingButton() == EKeys::LeftMouseButton )
+		{
+			return Reply.DetectDrag(SlateWidget.ToSharedRef(), EKeys::LeftMouseButton);
+		}
+
+		return Reply;
 	}
 
-	return Reply;
+	return FReply::Unhandled();
 }
 
 FReply SObjectWidget::OnPreviewMouseButtonDown(const FGeometry& MyGeometry, const FPointerEvent& MouseEvent)
 {
-	return WidgetObject->OnPreviewMouseButtonDown(MyGeometry, MouseEvent).ToReply(WidgetObject->GetWidget());
+	TSharedPtr<SWidget> SlateWidget = WidgetObject->GetCachedWidget();
+	if ( SlateWidget.IsValid() )
+	{
+		return WidgetObject->OnPreviewMouseButtonDown(MyGeometry, MouseEvent).ToReply(SlateWidget.ToSharedRef());
+	}
+
+	return FReply::Unhandled();
 }
 
 FReply SObjectWidget::OnMouseButtonUp(const FGeometry& MyGeometry, const FPointerEvent& MouseEvent)
 {
-	return WidgetObject->OnMouseButtonUp(MyGeometry, MouseEvent).ToReply(WidgetObject->GetWidget());
+	TSharedPtr<SWidget> SlateWidget = WidgetObject->GetCachedWidget();
+	if ( SlateWidget.IsValid() )
+	{
+		return WidgetObject->OnMouseButtonUp(MyGeometry, MouseEvent).ToReply(SlateWidget.ToSharedRef());
+	}
+
+	return FReply::Unhandled();
 }
 
 FReply SObjectWidget::OnMouseMove(const FGeometry& MyGeometry, const FPointerEvent& MouseEvent)
 {
-	return WidgetObject->OnMouseMove(MyGeometry, MouseEvent).ToReply(WidgetObject->GetWidget());
+	TSharedPtr<SWidget> SlateWidget = WidgetObject->GetCachedWidget();
+	if ( SlateWidget.IsValid() )
+	{
+		return WidgetObject->OnMouseMove(MyGeometry, MouseEvent).ToReply(SlateWidget.ToSharedRef());
+	}
+
+	return FReply::Unhandled();
 }
 
 void SObjectWidget::OnMouseEnter(const FGeometry& MyGeometry, const FPointerEvent& MouseEvent)
@@ -132,7 +189,13 @@ void SObjectWidget::OnMouseLeave(const FPointerEvent& MouseEvent)
 
 FReply SObjectWidget::OnMouseWheel(const FGeometry& MyGeometry, const FPointerEvent& MouseEvent)
 {
-	return WidgetObject->OnMouseWheel(MyGeometry, MouseEvent).ToReply(WidgetObject->GetWidget());
+	TSharedPtr<SWidget> SlateWidget = WidgetObject->GetCachedWidget();
+	if ( SlateWidget.IsValid() )
+	{
+		return WidgetObject->OnMouseWheel(MyGeometry, MouseEvent).ToReply(SlateWidget.ToSharedRef());
+	}
+
+	return FReply::Unhandled();
 }
 
 FCursorReply SObjectWidget::OnCursorQuery(const FGeometry& MyGeometry, const FPointerEvent& CursorEvent) const
@@ -143,7 +206,13 @@ FCursorReply SObjectWidget::OnCursorQuery(const FGeometry& MyGeometry, const FPo
 
 FReply SObjectWidget::OnMouseButtonDoubleClick(const FGeometry& MyGeometry, const FPointerEvent& MouseEvent)
 {
-	return WidgetObject->OnMouseButtonDoubleClick(MyGeometry, MouseEvent).ToReply(WidgetObject->GetWidget());
+	TSharedPtr<SWidget> SlateWidget = WidgetObject->GetCachedWidget();
+	if ( SlateWidget.IsValid() )
+	{
+		return WidgetObject->OnMouseButtonDoubleClick(MyGeometry, MouseEvent).ToReply(SlateWidget.ToSharedRef());
+	}
+
+	return FReply::Unhandled();
 }
 
 FReply SObjectWidget::OnDragDetected(const FGeometry& MyGeometry, const FPointerEvent& MouseEvent)
@@ -186,40 +255,88 @@ FReply SObjectWidget::OnDrop(const FGeometry& MyGeometry, const FDragDropEvent&
 
 FReply SObjectWidget::OnControllerButtonPressed(const FGeometry& MyGeometry, const FControllerEvent& ControllerEvent)
 {
-	return WidgetObject->OnControllerButtonPressed(MyGeometry, ControllerEvent).ToReply(WidgetObject->GetWidget());
+	TSharedPtr<SWidget> SlateWidget = WidgetObject->GetCachedWidget();
+	if ( SlateWidget.IsValid() )
+	{
+		return WidgetObject->OnControllerButtonPressed(MyGeometry, ControllerEvent).ToReply(SlateWidget.ToSharedRef());
+	}
+
+	return FReply::Unhandled();
 }
 
 FReply SObjectWidget::OnControllerButtonReleased(const FGeometry& MyGeometry, const FControllerEvent& ControllerEvent)
 {
-	return WidgetObject->OnControllerButtonReleased(MyGeometry, ControllerEvent).ToReply(WidgetObject->GetWidget());
+	TSharedPtr<SWidget> SlateWidget = WidgetObject->GetCachedWidget();
+	if ( SlateWidget.IsValid() )
+	{
+		return WidgetObject->OnControllerButtonReleased(MyGeometry, ControllerEvent).ToReply(SlateWidget.ToSharedRef());
+	}
+
+	return FReply::Unhandled();
 }
 
 FReply SObjectWidget::OnControllerAnalogValueChanged(const FGeometry& MyGeometry, const FControllerEvent& ControllerEvent)
 {
-	return WidgetObject->OnControllerAnalogValueChanged(MyGeometry, ControllerEvent).ToReply(WidgetObject->GetWidget());
+	TSharedPtr<SWidget> SlateWidget = WidgetObject->GetCachedWidget();
+	if ( SlateWidget.IsValid() )
+	{
+		return WidgetObject->OnControllerAnalogValueChanged(MyGeometry, ControllerEvent).ToReply(SlateWidget.ToSharedRef());
+	}
+
+	return FReply::Unhandled();
 }
 
 FReply SObjectWidget::OnTouchGesture(const FGeometry& MyGeometry, const FPointerEvent& GestureEvent)
 {
-	return WidgetObject->OnTouchGesture(MyGeometry, GestureEvent).ToReply(WidgetObject->GetWidget());
+	TSharedPtr<SWidget> SlateWidget = WidgetObject->GetCachedWidget();
+	if ( SlateWidget.IsValid() )
+	{
+		return WidgetObject->OnTouchGesture(MyGeometry, GestureEvent).ToReply(SlateWidget.ToSharedRef());
+	}
+
+	return FReply::Unhandled();
 }
 
 FReply SObjectWidget::OnTouchStarted(const FGeometry& MyGeometry, const FPointerEvent& InTouchEvent)
 {
-	return WidgetObject->OnTouchStarted(MyGeometry, InTouchEvent).ToReply(WidgetObject->GetWidget());
+	TSharedPtr<SWidget> SlateWidget = WidgetObject->GetCachedWidget();
+	if ( SlateWidget.IsValid() )
+	{
+		return WidgetObject->OnTouchStarted(MyGeometry, InTouchEvent).ToReply(SlateWidget.ToSharedRef());
+	}
+
+	return FReply::Unhandled();
 }
 
 FReply SObjectWidget::OnTouchMoved(const FGeometry& MyGeometry, const FPointerEvent& InTouchEvent)
 {
-	return WidgetObject->OnTouchMoved(MyGeometry, InTouchEvent).ToReply(WidgetObject->GetWidget());
+	TSharedPtr<SWidget> SlateWidget = WidgetObject->GetCachedWidget();
+	if ( SlateWidget.IsValid() )
+	{
+		return WidgetObject->OnTouchMoved(MyGeometry, InTouchEvent).ToReply(SlateWidget.ToSharedRef());
+	}
+
+	return FReply::Unhandled();
 }
 
 FReply SObjectWidget::OnTouchEnded(const FGeometry& MyGeometry, const FPointerEvent& InTouchEvent)
 {
-	return WidgetObject->OnTouchEnded(MyGeometry, InTouchEvent).ToReply(WidgetObject->GetWidget());
+	TSharedPtr<SWidget> SlateWidget = WidgetObject->GetCachedWidget();
+	if ( SlateWidget.IsValid() )
+	{
+		return WidgetObject->OnTouchEnded(MyGeometry, InTouchEvent).ToReply(SlateWidget.ToSharedRef());
+	}
+
+	return FReply::Unhandled();
 }
 
 FReply SObjectWidget::OnMotionDetected(const FGeometry& MyGeometry, const FMotionEvent& InMotionEvent)
 {
-	return WidgetObject->OnMotionDetected(MyGeometry, InMotionEvent).ToReply(WidgetObject->GetWidget());
+	TSharedPtr<SWidget> SlateWidget = WidgetObject->GetCachedWidget();
+	if ( SlateWidget.IsValid() )
+	{
+		return WidgetObject->OnMotionDetected(MyGeometry, InMotionEvent).ToReply(SlateWidget.ToSharedRef());
+	}
+
+	return FReply::Unhandled();
 }
diff --git a/Engine/Source/Runtime/UMG/Private/UserWidget.cpp b/Engine/Source/Runtime/UMG/Private/UserWidget.cpp
--- a/Engine/Source/Runtime/UMG/Private/UserWidget.cpp
+++ b/Engine/Source/Runtime/UMG/Private/UserWidget.cpp
@@ -148,6 +148,17 @@ UUserWidget::UUserWidget(const FPostConstructInitializeProperties& PCIP)
 	Visiblity = ESlateVisibility::SelfHitTestInvisible;
 }
 
+void UUserWidget::ReleaseNativeWidget()
+{
+	Super::ReleaseNativeWidget();
+
+	UWidget* RootWidget = GetRootWidgetComponent();
+	if ( RootWidget )
+	{
+		RootWidget->ReleaseNativeWidget();
+	}
+}
+
 void UUserWidget::PostInitProperties()
 {
 	Super::PostInitProperties();
diff --git a/Engine/Source/Runtime/UMG/Public/Blueprint/UserWidget.h b/Engine/Source/Runtime/UMG/Public/Blueprint/UserWidget.h
--- a/Engine/Source/Runtime/UMG/Public/Blueprint/UserWidget.h
+++ b/Engine/Source/Runtime/UMG/Public/Blueprint/UserWidget.h
@@ -104,6 +104,10 @@ public:
 	virtual class UWorld* GetWorld() const override;
 	// End of UObject interface
 
+	//UVisual interface
+	virtual void ReleaseNativeWidget() override;
+	// End of UVisual interface
+
 	/*  */
 	UFUNCTION(BlueprintCallable, Category="Appearance")
 	void Show();
@@ -197,7 +201,7 @@ public:
 	//virtual bool OnVisualizeTooltip(const TSharedPtr<SWidget>& TooltipContent);
 
 	/**
- 	 * Plays an animation in this widget
+	 * Plays an animation in this widget
 	 * 
 	 * @param The name of the animation to play
 	 */
@@ -205,7 +209,7 @@ public:
 	void PlayAnimation(FName AnimationName);
 
 	/**
- 	 * Stops an already running animation in this widget
+	 * Stops an already running animation in this widget
 	 * 
 	 * @param The name of the animation to stop
 	 */
diff --git a/Engine/Source/Runtime/UMG/Public/Components/Border.h b/Engine/Source/Runtime/UMG/Public/Components/Border.h
--- a/Engine/Source/Runtime/UMG/Public/Components/Border.h
+++ b/Engine/Source/Runtime/UMG/Public/Components/Border.h
@@ -91,6 +91,8 @@ class UMG_API UBorder : public UContentWidget
 	virtual void SyncronizeProperties() override;
 	// End of UWidget interface
 
+	virtual void ReleaseNativeWidget() override;
+
 #if WITH_EDITOR
 	virtual const FSlateBrush* GetEditorIcon() override;
 #endif
diff --git a/Engine/Source/Runtime/UMG/Public/Components/Button.h b/Engine/Source/Runtime/UMG/Public/Components/Button.h
--- a/Engine/Source/Runtime/UMG/Public/Components/Button.h
+++ b/Engine/Source/Runtime/UMG/Public/Components/Button.h
@@ -46,6 +46,8 @@ public:
 	/** Called when the button is clicked */
 	UPROPERTY(EditDefaultsOnly, Category=Events)
 	FOnReply OnClickedEvent;
+
+	virtual void ReleaseNativeWidget() override;
 	
 	/**  */
 	UFUNCTION(BlueprintCallable, Category="Appearance")
diff --git a/Engine/Source/Runtime/UMG/Public/Components/ButtonSlot.h b/Engine/Source/Runtime/UMG/Public/Components/ButtonSlot.h
--- a/Engine/Source/Runtime/UMG/Public/Components/ButtonSlot.h
+++ b/Engine/Source/Runtime/UMG/Public/Components/ButtonSlot.h
@@ -40,6 +40,8 @@ class UMG_API UButtonSlot : public UPanelSlot
 	/** Builds the underlying slot for the slate button. */
 	void BuildSlot(TSharedRef<SButton> InButton);
 
+	virtual void ReleaseNativeWidget() override;
+
 private:
 
 	/** A pointer to the button to allow us to adjust the size, padding...etc at runtime. */
diff --git a/Engine/Source/Runtime/UMG/Public/Components/CanvasPanel.h b/Engine/Source/Runtime/UMG/Public/Components/CanvasPanel.h
--- a/Engine/Source/Runtime/UMG/Public/Components/CanvasPanel.h
+++ b/Engine/Source/Runtime/UMG/Public/Components/CanvasPanel.h
@@ -26,6 +26,8 @@ class UMG_API UCanvasPanel : public UPanelWidget
 	virtual void SyncronizeProperties() override;
 	// End of UWidget interface
 
+	void ReleaseNativeWidget() override;
+
 #if WITH_EDITOR
 	// UWidget interface
 	virtual const FSlateBrush* GetEditorIcon() override;
diff --git a/Engine/Source/Runtime/UMG/Public/Components/CanvasPanelSlot.h b/Engine/Source/Runtime/UMG/Public/Components/CanvasPanelSlot.h
--- a/Engine/Source/Runtime/UMG/Public/Components/CanvasPanelSlot.h
+++ b/Engine/Source/Runtime/UMG/Public/Components/CanvasPanelSlot.h
@@ -68,6 +68,8 @@ class UMG_API UCanvasPanelSlot : public UPanelSlot
 	virtual void SyncronizeProperties() override;
 	// End of UPanelSlot interface
 
+	virtual void ReleaseNativeWidget() override;
+
 #if WITH_EDITOR
 	// UObject interface
 	virtual void PreEditChange(UProperty* PropertyAboutToChange) override;
diff --git a/Engine/Source/Runtime/UMG/Public/Components/CheckBox.h b/Engine/Source/Runtime/UMG/Public/Components/CheckBox.h
--- a/Engine/Source/Runtime/UMG/Public/Components/CheckBox.h
+++ b/Engine/Source/Runtime/UMG/Public/Components/CheckBox.h
@@ -119,6 +119,8 @@ protected:
 	virtual void SyncronizeProperties() override;
 	// End of UWidget interface
 
+	virtual void ReleaseNativeWidget() override;
+
 #if WITH_EDITOR
 	virtual const FSlateBrush* GetEditorIcon() override;
 #endif
diff --git a/Engine/Source/Runtime/UMG/Public/Components/CircularThrobber.h b/Engine/Source/Runtime/UMG/Public/Components/CircularThrobber.h
--- a/Engine/Source/Runtime/UMG/Public/Components/CircularThrobber.h
+++ b/Engine/Source/Runtime/UMG/Public/Components/CircularThrobber.h
@@ -34,6 +34,8 @@ public:
 	virtual void SyncronizeProperties() override;
 	// End of UWidget interface
 
+	virtual void ReleaseNativeWidget() override;
+
 protected:
 
 	/** How many pieces there are */
diff --git a/Engine/Source/Runtime/UMG/Public/Components/ComboBox.h b/Engine/Source/Runtime/UMG/Public/Components/ComboBox.h
--- a/Engine/Source/Runtime/UMG/Public/Components/ComboBox.h
+++ b/Engine/Source/Runtime/UMG/Public/Components/ComboBox.h
@@ -19,6 +19,8 @@ public:
 	/** Called when the widget is needed for the item. */
 	UPROPERTY(EditDefaultsOnly, Category=Events)
 	FGenerateWidgetUObject OnGenerateWidget;
+
+	virtual void ReleaseNativeWidget() override;
 	
 protected:
 	TSharedRef<SWidget> HandleGenerateWidget(UObject* Item) const;
diff --git a/Engine/Source/Runtime/UMG/Public/Components/EditableText.h b/Engine/Source/Runtime/UMG/Public/Components/EditableText.h
--- a/Engine/Source/Runtime/UMG/Public/Components/EditableText.h
+++ b/Engine/Source/Runtime/UMG/Public/Components/EditableText.h
@@ -105,6 +105,8 @@ public:
 	virtual void SyncronizeProperties() override;
 	// End of UWidget interface
 
+	virtual void ReleaseNativeWidget() override;
+
 #if WITH_EDITOR
 	virtual const FSlateBrush* GetEditorIcon() override;
 #endif
diff --git a/Engine/Source/Runtime/UMG/Public/Components/EditableTextBox.h b/Engine/Source/Runtime/UMG/Public/Components/EditableTextBox.h
--- a/Engine/Source/Runtime/UMG/Public/Components/EditableTextBox.h
+++ b/Engine/Source/Runtime/UMG/Public/Components/EditableTextBox.h
@@ -107,6 +107,8 @@ public:
 	virtual void SyncronizeProperties() override;
 	// End of UWidget interface
 
+	virtual void ReleaseNativeWidget() override;
+
 #if WITH_EDITOR
 	virtual const FSlateBrush* GetEditorIcon() override;
 #endif
diff --git a/Engine/Source/Runtime/UMG/Public/Components/HorizontalBox.h b/Engine/Source/Runtime/UMG/Public/Components/HorizontalBox.h
--- a/Engine/Source/Runtime/UMG/Public/Components/HorizontalBox.h
+++ b/Engine/Source/Runtime/UMG/Public/Components/HorizontalBox.h
@@ -19,6 +19,8 @@ class UMG_API UHorizontalBox : public UPanelWidget
 	// End UWidget interface
 #endif
 
+	virtual void ReleaseNativeWidget() override;
+
 protected:
 
 	// UPanelWidget
diff --git a/Engine/Source/Runtime/UMG/Public/Components/HorizontalBoxSlot.h b/Engine/Source/Runtime/UMG/Public/Components/HorizontalBoxSlot.h
--- a/Engine/Source/Runtime/UMG/Public/Components/HorizontalBoxSlot.h
+++ b/Engine/Source/Runtime/UMG/Public/Components/HorizontalBoxSlot.h
@@ -43,6 +43,8 @@ class UMG_API UHorizontalBoxSlot : public UPanelSlot
 	virtual void SyncronizeProperties() override;
 	// End of UPanelSlot interface
 
+	virtual void ReleaseNativeWidget() override;
+
 private:
 	SHorizontalBox::FSlot* Slot;
 };
\ No newline at end of file
diff --git a/Engine/Source/Runtime/UMG/Public/Components/Image.h b/Engine/Source/Runtime/UMG/Public/Components/Image.h
--- a/Engine/Source/Runtime/UMG/Public/Components/Image.h
+++ b/Engine/Source/Runtime/UMG/Public/Components/Image.h
@@ -45,6 +45,8 @@ public:
 	virtual void SyncronizeProperties() override;
 	// End of UWidget interface
 
+	virtual void ReleaseNativeWidget() override;
+
 #if WITH_EDITOR
 	// UWidget interface
 	virtual const FSlateBrush* GetEditorIcon() override;
diff --git a/Engine/Source/Runtime/UMG/Public/Components/MenuAnchor.h b/Engine/Source/Runtime/UMG/Public/Components/MenuAnchor.h
--- a/Engine/Source/Runtime/UMG/Public/Components/MenuAnchor.h
+++ b/Engine/Source/Runtime/UMG/Public/Components/MenuAnchor.h
@@ -54,6 +54,8 @@ public:
 	virtual void SyncronizeProperties() override;
 	// End of UWidget interface
 
+	virtual void ReleaseNativeWidget() override;
+
 #if WITH_EDITOR
 	virtual const FSlateBrush* GetEditorIcon() override;
 #endif
diff --git a/Engine/Source/Runtime/UMG/Public/Components/Overlay.h b/Engine/Source/Runtime/UMG/Public/Components/Overlay.h
--- a/Engine/Source/Runtime/UMG/Public/Components/Overlay.h
+++ b/Engine/Source/Runtime/UMG/Public/Components/Overlay.h
@@ -9,6 +9,8 @@ class UMG_API UOverlay : public UPanelWidget
 {
 	GENERATED_UCLASS_BODY()
 
+	virtual void ReleaseNativeWidget() override;
+
 protected:
 
 	// UPanelWidget
diff --git a/Engine/Source/Runtime/UMG/Public/Components/OverlaySlot.h b/Engine/Source/Runtime/UMG/Public/Components/OverlaySlot.h
--- a/Engine/Source/Runtime/UMG/Public/Components/OverlaySlot.h
+++ b/Engine/Source/Runtime/UMG/Public/Components/OverlaySlot.h
@@ -42,6 +42,8 @@ class UMG_API UOverlaySlot : public UPanelSlot
 	/** Builds the underlying FSlot for the Slate layout panel. */
 	void BuildSlot(TSharedRef<SOverlay> InOverlay);
 
+	virtual void ReleaseNativeWidget() override;
+
 private:
 	//TODO UMG Slots should hold weak or shared refs to slots.
 
diff --git a/Engine/Source/Runtime/UMG/Public/Components/PanelSlot.h b/Engine/Source/Runtime/UMG/Public/Components/PanelSlot.h
--- a/Engine/Source/Runtime/UMG/Public/Components/PanelSlot.h
+++ b/Engine/Source/Runtime/UMG/Public/Components/PanelSlot.h
@@ -32,6 +32,8 @@ public:
 
 	virtual bool CanMove() const;
 
+	virtual void ReleaseNativeWidget() override;
+
 	/** Applies all properties to the live slot if possible. */
 	virtual void SyncronizeProperties()
 	{
diff --git a/Engine/Source/Runtime/UMG/Public/Components/PanelWidget.h b/Engine/Source/Runtime/UMG/Public/Components/PanelWidget.h
--- a/Engine/Source/Runtime/UMG/Public/Components/PanelWidget.h
+++ b/Engine/Source/Runtime/UMG/Public/Components/PanelWidget.h
@@ -58,6 +58,8 @@ public:
 		return bCanHaveMultipleChildren;
 	}
 
+	virtual void ReleaseNativeWidget() override;
+
 #if WITH_EDITOR
 	virtual void ConnectEditorData() override
 	{
diff --git a/Engine/Source/Runtime/UMG/Public/Components/ScrollBox.h b/Engine/Source/Runtime/UMG/Public/Components/ScrollBox.h
--- a/Engine/Source/Runtime/UMG/Public/Components/ScrollBox.h
+++ b/Engine/Source/Runtime/UMG/Public/Components/ScrollBox.h
@@ -18,6 +18,8 @@ class UMG_API UScrollBox : public UPanelWidget
 	UPROPERTY(EditDefaultsOnly, Category=Style, meta=( DisplayThumbnail = "true" ))
 	USlateWidgetStyleAsset* BarStyle;
 
+	virtual void ReleaseNativeWidget() override;
+
 	/** Removes all child widget from the scrollbox */
 	UFUNCTION(BlueprintCallable, Category="Widget")
 	void ClearChildren();
diff --git a/Engine/Source/Runtime/UMG/Public/Components/ScrollBoxSlot.h b/Engine/Source/Runtime/UMG/Public/Components/ScrollBoxSlot.h
--- a/Engine/Source/Runtime/UMG/Public/Components/ScrollBoxSlot.h
+++ b/Engine/Source/Runtime/UMG/Public/Components/ScrollBoxSlot.h
@@ -33,6 +33,8 @@ class UMG_API UScrollBoxSlot : public UPanelSlot
 	/** Builds the underlying FSlot for the Slate layout panel. */
 	void BuildSlot(TSharedRef<SScrollBox> ScrollBox);
 
+	virtual void ReleaseNativeWidget() override;
+
 private:
 	//TODO UMG Slots should hold weak or shared refs to slots.
 
diff --git a/Engine/Source/Runtime/UMG/Public/Components/UniformGridPanel.h b/Engine/Source/Runtime/UMG/Public/Components/UniformGridPanel.h
--- a/Engine/Source/Runtime/UMG/Public/Components/UniformGridPanel.h
+++ b/Engine/Source/Runtime/UMG/Public/Components/UniformGridPanel.h
@@ -26,6 +26,8 @@ class UMG_API UUniformGridPanel : public UPanelWidget
 	virtual void SyncronizeProperties() override;
 	// End of UWidget interface
 
+	virtual void ReleaseNativeWidget() override;
+
 #if WITH_EDITOR
 	virtual const FSlateBrush* GetEditorIcon() override;
 #endif
diff --git a/Engine/Source/Runtime/UMG/Public/Components/UniformGridSlot.h b/Engine/Source/Runtime/UMG/Public/Components/UniformGridSlot.h
--- a/Engine/Source/Runtime/UMG/Public/Components/UniformGridSlot.h
+++ b/Engine/Source/Runtime/UMG/Public/Components/UniformGridSlot.h
@@ -56,6 +56,8 @@ public:
 	/** Builds the underlying FSlot for the Slate layout panel. */
 	void BuildSlot(TSharedRef<SUniformGridPanel> GridPanel);
 
+	virtual void ReleaseNativeWidget() override;
+
 private:
 	//TODO UMG Slots should hold weak or shared refs to slots.
 
diff --git a/Engine/Source/Runtime/UMG/Public/Components/VerticalBox.h b/Engine/Source/Runtime/UMG/Public/Components/VerticalBox.h
--- a/Engine/Source/Runtime/UMG/Public/Components/VerticalBox.h
+++ b/Engine/Source/Runtime/UMG/Public/Components/VerticalBox.h
@@ -31,6 +31,8 @@ protected:
 	virtual void OnSlotRemoved(UPanelSlot* Slot) override;
 	// End UPanelWidget
 
+	virtual void ReleaseNativeWidget() override;
+
 protected:
 
 	TSharedPtr<class SVerticalBox> MyVerticalBox;
diff --git a/Engine/Source/Runtime/UMG/Public/Components/VerticalBoxSlot.h b/Engine/Source/Runtime/UMG/Public/Components/VerticalBoxSlot.h
--- a/Engine/Source/Runtime/UMG/Public/Components/VerticalBoxSlot.h
+++ b/Engine/Source/Runtime/UMG/Public/Components/VerticalBoxSlot.h
@@ -44,6 +44,8 @@ class UMG_API UVerticalBoxSlot : public UPanelSlot
 	virtual void SyncronizeProperties() override;
 	// End of UPanelSlot interface
 
+	virtual void ReleaseNativeWidget() override;
+
 	/** Builds the underlying FSlot for the Slate layout panel. */
 	void BuildSlot(TSharedRef<SVerticalBox> InVerticalBox);
 
diff --git a/Engine/Source/Runtime/UMG/Public/Components/Visual.h b/Engine/Source/Runtime/UMG/Public/Components/Visual.h
--- a/Engine/Source/Runtime/UMG/Public/Components/Visual.h
+++ b/Engine/Source/Runtime/UMG/Public/Components/Visual.h
@@ -11,5 +11,5 @@ class UMG_API UVisual : public UObject
 	GENERATED_UCLASS_BODY()
 
 public:
-
+	virtual void ReleaseNativeWidget();
 };
diff --git a/Engine/Source/Runtime/UMG/Public/Components/Widget.h b/Engine/Source/Runtime/UMG/Public/Components/Widget.h
--- a/Engine/Source/Runtime/UMG/Public/Components/Widget.h
+++ b/Engine/Source/Runtime/UMG/Public/Components/Widget.h
@@ -142,6 +142,8 @@ public:
 	 */
 	TSharedRef<SWidget> GetWidget() const;
 
+	TSharedPtr<SWidget> GetCachedWidget() const;
+
 	/**
 	 * Applies all properties to the native widget if possible.  This is called after a widget is constructed.
 	 * It can also be called by the editor to update modified state, so ensure all initialization to a widgets
diff --git a/Engine/Source/Runtime/UMG/Public/Components/WidgetSwitcher.h b/Engine/Source/Runtime/UMG/Public/Components/WidgetSwitcher.h
--- a/Engine/Source/Runtime/UMG/Public/Components/WidgetSwitcher.h
+++ b/Engine/Source/Runtime/UMG/Public/Components/WidgetSwitcher.h
@@ -29,7 +29,9 @@ class UMG_API UWidgetSwitcher : public UPanelWidget
 	UFUNCTION(BlueprintCallable, Category="Switcher")
 	void SetActiveWidgetIndex( int32 Index );
 	
-	void SyncronizeProperties();
+	virtual void SyncronizeProperties() override;
+
+	virtual void ReleaseNativeWidget() override;
 
 #if WITH_EDITOR
 	virtual const FSlateBrush* GetEditorIcon() override;
diff --git a/Engine/Source/Runtime/UMG/Public/Components/WidgetSwitcherSlot.h b/Engine/Source/Runtime/UMG/Public/Components/WidgetSwitcherSlot.h
--- a/Engine/Source/Runtime/UMG/Public/Components/WidgetSwitcherSlot.h
+++ b/Engine/Source/Runtime/UMG/Public/Components/WidgetSwitcherSlot.h
@@ -37,6 +37,8 @@ class UMG_API UWidgetSwitcherSlot : public UPanelSlot
 	virtual void SyncronizeProperties() override;
 	// End of UPanelSlot interface
 
+	virtual void ReleaseNativeWidget() override;
+
 	/** Builds the underlying FSlot for the Slate layout panel. */
 	void BuildSlot(TSharedRef<SWidgetSwitcher> InWidgetSwitcher);
 
