Remove includes from UMGSequencerPlayer.h

diff --git a/Engine/Source/Runtime/UMG/Public/Animation/UMGSequencePlayer.h b/Engine/Source/Runtime/UMG/Public/Animation/UMGSequencePlayer.h
--- a/Engine/Source/Runtime/UMG/Public/Animation/UMGSequencePlayer.h
+++ b/Engine/Source/Runtime/UMG/Public/Animation/UMGSequencePlayer.h
@@ -3,10 +3,11 @@
 #pragma once
 
 #include "IMovieScenePlayer.h"
-#include "MovieScene.h"
-#include "MovieSceneBindings.h"
 #include "UMGSequencePlayer.generated.h"
 
+class UMovieScene;
+class UMovieSceneBindings;
+
 UCLASS(transient)
 class UMG_API UUMGSequencePlayer : public UObject, public IMovieScenePlayer
 {
