UMG - Adding the beginning of a viewport widget.  Fixing many bugs with widget lifetimes, renaming GetWidget to TakeWidget so that is better reflects the reality of what is happening, or will happen in the future.  There's now a GetCachedWidget call that is what you should use when you just want the underlying Slate widget but don't want to create it.  Improving my QA game sandbox to include a viewport and now the dynamic list of players allows you to add new players and shows off dynamic adding and removing.

diff --git a/Engine/Source/Editor/UMGEditor/Private/Designer/SDesignerView.cpp b/Engine/Source/Editor/UMGEditor/Private/Designer/SDesignerView.cpp
--- a/Engine/Source/Editor/UMGEditor/Private/Designer/SDesignerView.cpp
+++ b/Engine/Source/Editor/UMGEditor/Private/Designer/SDesignerView.cpp
@@ -1077,7 +1077,8 @@ UWidget* SDesignerView::ProcessDropAndAddWidget(const FGeometry& MyGeometry, con
 			FVector2D LocalPosition = ArrangedWidget.Geometry.AbsoluteToLocal(DragDropEvent.GetScreenSpacePosition());
 			if ( UPanelSlot* Slot = Parent->AddChild(Widget) )
 			{
-				TSharedPtr<SWidget> SlateWidget = Widget->GetWidget();
+				// HACK UMG - This seems like a bad idea to call TakeWidget
+				TSharedPtr<SWidget> SlateWidget = Widget->TakeWidget();
 				SlateWidget->SlatePrepass();
 				const FVector2D& WidgetDesiredSize = SlateWidget->GetDesiredSize();
 
diff --git a/Engine/Source/Editor/UMGEditor/Private/Hierarchy/SHierarchyViewItem.cpp b/Engine/Source/Editor/UMGEditor/Private/Hierarchy/SHierarchyViewItem.cpp
--- a/Engine/Source/Editor/UMGEditor/Private/Hierarchy/SHierarchyViewItem.cpp
+++ b/Engine/Source/Editor/UMGEditor/Private/Hierarchy/SHierarchyViewItem.cpp
@@ -294,7 +294,11 @@ FReply SHierarchyViewItem::ProcessDragDrop(const FDragDropEvent& DragDropEvent,
 					}
 					else
 					{
-						Item.GetPreview()->GetWidget()->SlatePrepass();
+						TSharedPtr<SWidget> Widget = Item.GetPreview()->GetCachedWidget();
+						if ( Widget.IsValid() )
+						{
+							Widget->SlatePrepass();
+						}
 					}
 				}
 				else
diff --git a/Engine/Source/Runtime/Engine/Public/Slate/SceneViewport.h b/Engine/Source/Runtime/Engine/Public/Slate/SceneViewport.h
--- a/Engine/Source/Runtime/Engine/Public/Slate/SceneViewport.h
+++ b/Engine/Source/Runtime/Engine/Public/Slate/SceneViewport.h
@@ -11,7 +11,6 @@ DECLARE_DELEGATE_OneParam( FOnSceneViewportResize, FVector2D );
 
 /**
  * A viewport for use with Slate SViewport widgets.
- * This viewport does not render directly to the back buffer but instead renders to a separate render target which Slate uses when it draws
  */
 class ENGINE_API FSceneViewport : public FViewportFrame, public FViewport, public ISlateViewport
 {
diff --git a/Engine/Source/Runtime/Slate/Public/Widgets/SViewport.h b/Engine/Source/Runtime/Slate/Public/Widgets/SViewport.h
--- a/Engine/Source/Runtime/Slate/Public/Widgets/SViewport.h
+++ b/Engine/Source/Runtime/Slate/Public/Widgets/SViewport.h
@@ -161,11 +161,13 @@ private:
 	{
 		return ChildSlot;
 	}
-	
-private:
+
+protected:
 
 	// Interface to the rendering and I/O implementation of the viewport.
 	TWeakPtr<ISlateViewport> ViewportInterface;
+	
+private:
 
 	// Whether or not to show the disabled effect when this viewport is disabled.
 	TAttribute<bool> ShowDisabledEffect;
diff --git a/Engine/Source/Runtime/UMG/Private/Components/Border.cpp b/Engine/Source/Runtime/UMG/Private/Components/Border.cpp
--- a/Engine/Source/Runtime/UMG/Private/Components/Border.cpp
+++ b/Engine/Source/Runtime/UMG/Private/Components/Border.cpp
@@ -42,7 +42,7 @@ TSharedRef<SWidget> UBorder::RebuildWidget()
 	
 	if ( GetChildrenCount() > 0 )
 	{
-		MyBorder->SetContent(GetContentSlot()->Content ? GetContentSlot()->Content->GetWidget() : SNullWidget::NullWidget);
+		MyBorder->SetContent(GetContentSlot()->Content ? GetContentSlot()->Content->TakeWidget() : SNullWidget::NullWidget);
 	}
 
 	return MyBorder.ToSharedRef();
@@ -78,7 +78,7 @@ void UBorder::OnSlotAdded(UPanelSlot* Slot)
 	// Add the child to the live canvas if it already exists
 	if ( MyBorder.IsValid() )
 	{
-		MyBorder->SetContent(Slot->Content ? Slot->Content->GetWidget() : SNullWidget::NullWidget);
+		MyBorder->SetContent(Slot->Content ? Slot->Content->TakeWidget() : SNullWidget::NullWidget);
 	}
 }
 
diff --git a/Engine/Source/Runtime/UMG/Private/Components/Button.cpp b/Engine/Source/Runtime/UMG/Private/Components/Button.cpp
--- a/Engine/Source/Runtime/UMG/Private/Components/Button.cpp
+++ b/Engine/Source/Runtime/UMG/Private/Components/Button.cpp
@@ -131,7 +131,7 @@ void UButton::SetForegroundColor(FLinearColor InForegroundColor)
 
 FReply UButton::HandleOnClicked()
 {
-	if ( OnClickedEvent.IsBound() )
+	if ( OnClickedEvent.IsBound() && MyButton.IsValid() )
 	{
 		return OnClickedEvent.Execute().ToReply( MyButton.ToSharedRef() );
 	}
diff --git a/Engine/Source/Runtime/UMG/Private/Components/ButtonSlot.cpp b/Engine/Source/Runtime/UMG/Private/Components/ButtonSlot.cpp
--- a/Engine/Source/Runtime/UMG/Private/Components/ButtonSlot.cpp
+++ b/Engine/Source/Runtime/UMG/Private/Components/ButtonSlot.cpp
@@ -29,7 +29,7 @@ void UButtonSlot::BuildSlot(TSharedRef<SButton> InButton)
 	Button->SetHAlign(HorizontalAlignment);
 	Button->SetVAlign(VerticalAlignment);
 
-	Button->SetContent(Content ? Content->GetWidget() : SNullWidget::NullWidget);
+	Button->SetContent(Content ? Content->TakeWidget() : SNullWidget::NullWidget);
 }
 
 void UButtonSlot::SetPadding(FMargin InPadding)
diff --git a/Engine/Source/Runtime/UMG/Private/Components/CanvasPanel.cpp b/Engine/Source/Runtime/UMG/Private/Components/CanvasPanel.cpp
--- a/Engine/Source/Runtime/UMG/Private/Components/CanvasPanel.cpp
+++ b/Engine/Source/Runtime/UMG/Private/Components/CanvasPanel.cpp
@@ -75,7 +75,11 @@ void UCanvasPanel::OnSlotRemoved(UPanelSlot* Slot)
 	// Remove the widget from the live slot if it exists.
 	if ( MyCanvas.IsValid() )
 	{
-		MyCanvas->RemoveSlot(Slot->Content->GetWidget());
+		TSharedPtr<SWidget> Widget = Slot->Content->GetCachedWidget();
+		if ( Widget.IsValid() )
+		{
+			MyCanvas->RemoveSlot(Widget.ToSharedRef());
+		}
 	}
 }
 
@@ -128,7 +132,7 @@ bool UCanvasPanel::GetGeometryForSlot(UCanvasPanelSlot* Slot, FGeometry& Arrange
 
 		for ( int32 ChildIndex = 0; ChildIndex < ArrangedChildren.Num(); ChildIndex++ )
 		{
-			if ( ArrangedChildren(ChildIndex).Widget == Slot->Content->GetWidget() )
+			if ( ArrangedChildren(ChildIndex).Widget == Slot->Content->TakeWidget() )
 			{
 				ArrangedGeometry = ArrangedChildren(ChildIndex).Geometry;
 				return true;
diff --git a/Engine/Source/Runtime/UMG/Private/Components/CanvasPanelSlot.cpp b/Engine/Source/Runtime/UMG/Private/Components/CanvasPanelSlot.cpp
--- a/Engine/Source/Runtime/UMG/Private/Components/CanvasPanelSlot.cpp
+++ b/Engine/Source/Runtime/UMG/Private/Components/CanvasPanelSlot.cpp
@@ -26,7 +26,7 @@ void UCanvasPanelSlot::BuildSlot(TSharedRef<SConstraintCanvas> Canvas)
 {
 	Slot = &Canvas->AddSlot()
 		[
-			Content == NULL ? SNullWidget::NullWidget : Content->GetWidget()
+			Content == NULL ? SNullWidget::NullWidget : Content->TakeWidget()
 		];
 
 	SyncronizeProperties();
diff --git a/Engine/Source/Runtime/UMG/Private/Components/CheckBox.cpp b/Engine/Source/Runtime/UMG/Private/Components/CheckBox.cpp
--- a/Engine/Source/Runtime/UMG/Private/Components/CheckBox.cpp
+++ b/Engine/Source/Runtime/UMG/Private/Components/CheckBox.cpp
@@ -61,7 +61,7 @@ TSharedRef<SWidget> UCheckBox::RebuildWidget()
 
 	if ( GetChildrenCount() > 0 )
 	{
-		MyCheckbox->SetContent(GetContentSlot()->Content ? GetContentSlot()->Content->GetWidget() : SNullWidget::NullWidget);
+		MyCheckbox->SetContent(GetContentSlot()->Content ? GetContentSlot()->Content->TakeWidget() : SNullWidget::NullWidget);
 	}
 	
 	return MyCheckbox.ToSharedRef();
@@ -100,7 +100,7 @@ void UCheckBox::OnSlotAdded(UPanelSlot* Slot)
 	// Add the child to the live slot if it already exists
 	if ( MyCheckbox.IsValid() )
 	{
-		MyCheckbox->SetContent(Slot->Content ? Slot->Content->GetWidget() : SNullWidget::NullWidget);
+		MyCheckbox->SetContent(Slot->Content ? Slot->Content->TakeWidget() : SNullWidget::NullWidget);
 	}
 }
 
diff --git a/Engine/Source/Runtime/UMG/Private/Components/ComboBox.cpp b/Engine/Source/Runtime/UMG/Private/Components/ComboBox.cpp
--- a/Engine/Source/Runtime/UMG/Private/Components/ComboBox.cpp
+++ b/Engine/Source/Runtime/UMG/Private/Components/ComboBox.cpp
@@ -35,7 +35,7 @@ TSharedRef<SWidget> UComboBox::HandleGenerateWidget(UObject* Item) const
 		UWidget* Widget = OnGenerateWidget.Execute(Item);
 		if ( Widget != NULL )
 		{
-			return Widget->GetWidget();
+			return Widget->TakeWidget();
 		}
 	}
 
diff --git a/Engine/Source/Runtime/UMG/Private/Components/EditableTextBox.cpp b/Engine/Source/Runtime/UMG/Private/Components/EditableTextBox.cpp
--- a/Engine/Source/Runtime/UMG/Private/Components/EditableTextBox.cpp
+++ b/Engine/Source/Runtime/UMG/Private/Components/EditableTextBox.cpp
@@ -106,6 +106,14 @@ void UEditableTextBox::SetError(FText InError)
 	}
 }
 
+void UEditableTextBox::ClearError()
+{
+	if ( MyEditableTextBlock.IsValid() )
+	{
+		MyEditableTextBlock->SetError(FText::GetEmpty());
+	}
+}
+
 void UEditableTextBox::HandleOnTextChanged(const FText& Text)
 {
 	OnTextChanged.Broadcast(Text);
diff --git a/Engine/Source/Runtime/UMG/Private/Components/HorizontalBox.cpp b/Engine/Source/Runtime/UMG/Private/Components/HorizontalBox.cpp
--- a/Engine/Source/Runtime/UMG/Private/Components/HorizontalBox.cpp
+++ b/Engine/Source/Runtime/UMG/Private/Components/HorizontalBox.cpp
@@ -40,7 +40,11 @@ void UHorizontalBox::OnSlotRemoved(UPanelSlot* Slot)
 	// Remove the widget from the live slot if it exists.
 	if ( MyHorizontalBox.IsValid() )
 	{
-		MyHorizontalBox->RemoveSlot(Slot->Content->GetWidget());
+		TSharedPtr<SWidget> Widget = Slot->Content->GetCachedWidget();
+		if ( Widget.IsValid() )
+		{
+			MyHorizontalBox->RemoveSlot(Widget.ToSharedRef());
+		}
 	}
 }
 
diff --git a/Engine/Source/Runtime/UMG/Private/Components/HorizontalBoxSlot.cpp b/Engine/Source/Runtime/UMG/Private/Components/HorizontalBoxSlot.cpp
--- a/Engine/Source/Runtime/UMG/Private/Components/HorizontalBoxSlot.cpp
+++ b/Engine/Source/Runtime/UMG/Private/Components/HorizontalBoxSlot.cpp
@@ -27,7 +27,7 @@ void UHorizontalBoxSlot::BuildSlot(TSharedRef<SHorizontalBox> HorizontalBox)
 		.VAlign(VerticalAlignment)
 		.Padding(Padding)
 		[
-			Content == NULL ? SNullWidget::NullWidget : Content->GetWidget()
+			Content == NULL ? SNullWidget::NullWidget : Content->TakeWidget()
 		];
 
 	Slot->SizeParam = UWidget::ConvertSerializedSizeParamToRuntime(Size);
diff --git a/Engine/Source/Runtime/UMG/Private/Components/ListView.cpp b/Engine/Source/Runtime/UMG/Private/Components/ListView.cpp
--- a/Engine/Source/Runtime/UMG/Private/Components/ListView.cpp
+++ b/Engine/Source/Runtime/UMG/Private/Components/ListView.cpp
@@ -47,7 +47,7 @@ TSharedRef<ITableRow> UListView::HandleOnGenerateRow(UObject* Item, const TShare
 		{
 			return SNew(STableRow< UObject* >, OwnerTable)
 			[
-				Widget->GetWidget()
+				Widget->TakeWidget()
 			];
 		}
 	}
diff --git a/Engine/Source/Runtime/UMG/Private/Components/MenuAnchor.cpp b/Engine/Source/Runtime/UMG/Private/Components/MenuAnchor.cpp
--- a/Engine/Source/Runtime/UMG/Private/Components/MenuAnchor.cpp
+++ b/Engine/Source/Runtime/UMG/Private/Components/MenuAnchor.cpp
@@ -29,7 +29,7 @@ TSharedRef<SWidget> UMenuAnchor::RebuildWidget()
 
 	if ( GetChildrenCount() > 0 )
 	{
-		MyMenuAnchor->SetContent(GetContentSlot()->Content ? GetContentSlot()->Content->GetWidget() : SNullWidget::NullWidget);
+		MyMenuAnchor->SetContent(GetContentSlot()->Content ? GetContentSlot()->Content->TakeWidget() : SNullWidget::NullWidget);
 	}
 	
 	return BuildDesignTimeWidget( MyMenuAnchor.ToSharedRef() );
@@ -45,7 +45,7 @@ void UMenuAnchor::OnSlotAdded(UPanelSlot* Slot)
 	// Add the child to the live slot if it already exists
 	if ( MyMenuAnchor.IsValid() )
 	{
-		MyMenuAnchor->SetContent(Slot->Content ? Slot->Content->GetWidget() : SNullWidget::NullWidget);
+		MyMenuAnchor->SetContent(Slot->Content ? Slot->Content->TakeWidget() : SNullWidget::NullWidget);
 	}
 }
 
@@ -67,7 +67,7 @@ TSharedRef<SWidget> UMenuAnchor::HandleGetMenuContent()
 		UWidget* MenuWidget = OnGetMenuContentEvent.Execute();
 		if ( MenuWidget )
 		{
-			SlateMenuWidget = MenuWidget->GetWidget();
+			SlateMenuWidget = MenuWidget->TakeWidget();
 		}
 	}
 
@@ -76,7 +76,7 @@ TSharedRef<SWidget> UMenuAnchor::HandleGetMenuContent()
 		UWidget* MenuWidget = (UWidget*)ConstructObject<UWidget>(MenuClass, GetOuter());
 		if ( MenuWidget )
 		{
-			SlateMenuWidget = MenuWidget->GetWidget();
+			SlateMenuWidget = MenuWidget->TakeWidget();
 		}
 	}
 
diff --git a/Engine/Source/Runtime/UMG/Private/Components/Overlay.cpp b/Engine/Source/Runtime/UMG/Private/Components/Overlay.cpp
--- a/Engine/Source/Runtime/UMG/Private/Components/Overlay.cpp
+++ b/Engine/Source/Runtime/UMG/Private/Components/Overlay.cpp
@@ -40,7 +40,11 @@ void UOverlay::OnSlotRemoved(UPanelSlot* Slot)
 	// Remove the widget from the live slot if it exists.
 	if ( MyOverlay.IsValid() )
 	{
-		MyOverlay->RemoveSlot(Slot->Content->GetWidget());
+		TSharedPtr<SWidget> Widget = Slot->Content->GetCachedWidget();
+		if ( Widget.IsValid() )
+		{
+			MyOverlay->RemoveSlot(Widget.ToSharedRef());
+		}
 	}
 }
 
diff --git a/Engine/Source/Runtime/UMG/Private/Components/OverlaySlot.cpp b/Engine/Source/Runtime/UMG/Private/Components/OverlaySlot.cpp
--- a/Engine/Source/Runtime/UMG/Private/Components/OverlaySlot.cpp
+++ b/Engine/Source/Runtime/UMG/Private/Components/OverlaySlot.cpp
@@ -26,7 +26,7 @@ void UOverlaySlot::BuildSlot(TSharedRef<SOverlay> Overlay)
 		.HAlign(HorizontalAlignment)
 		.VAlign(VerticalAlignment)
 		[
-			Content == NULL ? SNullWidget::NullWidget : Content->GetWidget()
+			Content == NULL ? SNullWidget::NullWidget : Content->TakeWidget()
 		];
 }
 
diff --git a/Engine/Source/Runtime/UMG/Private/Components/ScrollBox.cpp b/Engine/Source/Runtime/UMG/Private/Components/ScrollBox.cpp
--- a/Engine/Source/Runtime/UMG/Private/Components/ScrollBox.cpp
+++ b/Engine/Source/Runtime/UMG/Private/Components/ScrollBox.cpp
@@ -42,7 +42,11 @@ void UScrollBox::OnSlotRemoved(UPanelSlot* Slot)
 	// Remove the widget from the live slot if it exists.
 	if ( MyScrollBox.IsValid() )
 	{
-		MyScrollBox->RemoveSlot(Slot->Content->GetWidget());
+		TSharedPtr<SWidget> Widget = Slot->Content->GetCachedWidget();
+		if ( Widget.IsValid() )
+		{
+			MyScrollBox->RemoveSlot(Widget.ToSharedRef());
+		}
 	}
 }
 
diff --git a/Engine/Source/Runtime/UMG/Private/Components/ScrollBoxSlot.cpp b/Engine/Source/Runtime/UMG/Private/Components/ScrollBoxSlot.cpp
--- a/Engine/Source/Runtime/UMG/Private/Components/ScrollBoxSlot.cpp
+++ b/Engine/Source/Runtime/UMG/Private/Components/ScrollBoxSlot.cpp
@@ -18,7 +18,7 @@ void UScrollBoxSlot::BuildSlot(TSharedRef<SScrollBox> ScrollBox)
 		.Padding(Padding)
 		.HAlign(HorizontalAlignment)
 		[
-			Content == NULL ? SNullWidget::NullWidget : Content->GetWidget()
+			Content == NULL ? SNullWidget::NullWidget : Content->TakeWidget()
 		];
 }
 
diff --git a/Engine/Source/Runtime/UMG/Private/Components/TileView.cpp b/Engine/Source/Runtime/UMG/Private/Components/TileView.cpp
--- a/Engine/Source/Runtime/UMG/Private/Components/TileView.cpp
+++ b/Engine/Source/Runtime/UMG/Private/Components/TileView.cpp
@@ -49,7 +49,7 @@ TSharedRef<ITableRow> UTileView::HandleOnGenerateTile(UObject* Item, const TShar
 		{
 			return SNew(STableRow< UObject* >, OwnerTable)
 			[
-				Widget->GetWidget()
+				Widget->TakeWidget()
 			];
 		}
 	}
diff --git a/Engine/Source/Runtime/UMG/Private/Components/UniformGridPanel.cpp b/Engine/Source/Runtime/UMG/Private/Components/UniformGridPanel.cpp
--- a/Engine/Source/Runtime/UMG/Private/Components/UniformGridPanel.cpp
+++ b/Engine/Source/Runtime/UMG/Private/Components/UniformGridPanel.cpp
@@ -42,7 +42,11 @@ void UUniformGridPanel::OnSlotRemoved(UPanelSlot* Slot)
 	// Remove the widget from the live slot if it exists.
 	if ( MyUniformGridPanel.IsValid() )
 	{
-		MyUniformGridPanel->RemoveSlot(Slot->Content->GetWidget());
+		TSharedPtr<SWidget> Widget = Slot->Content->GetCachedWidget();
+		if ( Widget.IsValid() )
+		{
+			MyUniformGridPanel->RemoveSlot(Widget.ToSharedRef());
+		}
 	}
 }
 
diff --git a/Engine/Source/Runtime/UMG/Private/Components/UniformGridSlot.cpp b/Engine/Source/Runtime/UMG/Private/Components/UniformGridSlot.cpp
--- a/Engine/Source/Runtime/UMG/Private/Components/UniformGridSlot.cpp
+++ b/Engine/Source/Runtime/UMG/Private/Components/UniformGridSlot.cpp
@@ -24,7 +24,7 @@ void UUniformGridSlot::BuildSlot(TSharedRef<SUniformGridPanel> GridPanel)
 {
 	Slot = &GridPanel->AddSlot(Column, Row)
 		[
-			Content == NULL ? SNullWidget::NullWidget : Content->GetWidget()
+			Content == NULL ? SNullWidget::NullWidget : Content->TakeWidget()
 		];
 }
 
diff --git a/Engine/Source/Runtime/UMG/Private/Components/VerticalBox.cpp b/Engine/Source/Runtime/UMG/Private/Components/VerticalBox.cpp
--- a/Engine/Source/Runtime/UMG/Private/Components/VerticalBox.cpp
+++ b/Engine/Source/Runtime/UMG/Private/Components/VerticalBox.cpp
@@ -40,7 +40,11 @@ void UVerticalBox::OnSlotRemoved(UPanelSlot* Slot)
 	// Remove the widget from the live slot if it exists.
 	if ( MyVerticalBox.IsValid() )
 	{
-		MyVerticalBox->RemoveSlot(Slot->Content->GetWidget());
+		TSharedPtr<SWidget> Widget = Slot->Content->GetCachedWidget();
+		if ( Widget.IsValid() )
+		{
+			MyVerticalBox->RemoveSlot(Widget.ToSharedRef());
+		}
 	}
 }
 
diff --git a/Engine/Source/Runtime/UMG/Private/Components/VerticalBoxSlot.cpp b/Engine/Source/Runtime/UMG/Private/Components/VerticalBoxSlot.cpp
--- a/Engine/Source/Runtime/UMG/Private/Components/VerticalBoxSlot.cpp
+++ b/Engine/Source/Runtime/UMG/Private/Components/VerticalBoxSlot.cpp
@@ -28,7 +28,7 @@ void UVerticalBoxSlot::BuildSlot(TSharedRef<SVerticalBox> VerticalBox)
 		.HAlign(HorizontalAlignment)
 		.VAlign(VerticalAlignment)
 		[
-			Content == NULL ? SNullWidget::NullWidget : Content->GetWidget()
+			Content == NULL ? SNullWidget::NullWidget : Content->TakeWidget()
 		];
 
 	Slot->SizeParam = UWidget::ConvertSerializedSizeParamToRuntime(Size);
diff --git a/Engine/Source/Runtime/UMG/Private/Components/Viewport.cpp b/Engine/Source/Runtime/UMG/Private/Components/Viewport.cpp
new file mode 100644
--- /dev/null
+++ b/Engine/Source/Runtime/UMG/Private/Components/Viewport.cpp
@@ -0,0 +1,523 @@
+// Copyright 1998-2014 Epic Games, Inc. All Rights Reserved.
+
+#include "UMGPrivatePCH.h"
+#include "Slate/SlateBrushAsset.h"
+#include "PreviewScene.h"
+#include "EngineModule.h"
+
+#define LOCTEXT_NAMESPACE "UMG"
+
+namespace FocusConstants
+{
+	const float TransitionTime = 0.25f;
+}
+
+FUMGViewportCameraTransform::FUMGViewportCameraTransform()
+	: TransitionCurve(new FCurveSequence(0.0f, FocusConstants::TransitionTime, ECurveEaseFunction::CubicOut))
+	, ViewLocation(FVector::ZeroVector)
+	, ViewRotation(FRotator::ZeroRotator)
+	, DesiredLocation(FVector::ZeroVector)
+	, LookAt(FVector::ZeroVector)
+	, StartLocation(FVector::ZeroVector)
+	, OrthoZoom(DEFAULT_ORTHOZOOM)
+{}
+
+void FUMGViewportCameraTransform::SetLocation(const FVector& Position)
+{
+	ViewLocation = Position;
+	DesiredLocation = ViewLocation;
+}
+
+void FUMGViewportCameraTransform::TransitionToLocation(const FVector& InDesiredLocation, bool bInstant)
+{
+	if ( bInstant )
+	{
+		SetLocation(InDesiredLocation);
+		TransitionCurve->JumpToEnd();
+	}
+	else
+	{
+		DesiredLocation = InDesiredLocation;
+		StartLocation = ViewLocation;
+
+		TransitionCurve->Play();
+	}
+}
+
+bool FUMGViewportCameraTransform::UpdateTransition()
+{
+	bool bIsAnimating = false;
+	if ( TransitionCurve->IsPlaying() || ViewLocation != DesiredLocation )
+	{
+		float LerpWeight = TransitionCurve->GetLerp();
+
+		if ( LerpWeight == 1.0f )
+		{
+			// Failsafe for the value not being exact on lerps
+			ViewLocation = DesiredLocation;
+		}
+		else
+		{
+			ViewLocation = FMath::Lerp(StartLocation, DesiredLocation, LerpWeight);
+		}
+
+
+		bIsAnimating = true;
+	}
+
+	return bIsAnimating;
+}
+
+FMatrix FUMGViewportCameraTransform::ComputeOrbitMatrix() const
+{
+	FTransform Transform =
+		FTransform(-LookAt) *
+		FTransform(FRotator(0, ViewRotation.Yaw, 0)) *
+		FTransform(FRotator(0, 0, ViewRotation.Pitch)) *
+		FTransform(FVector(0, ( ViewLocation - LookAt ).Size(), 0));
+
+	return Transform.ToMatrixNoScale() * FInverseRotationMatrix(FRotator(0, 90.f, 0));
+}
+
+FUMGViewportClient::FUMGViewportClient(FPreviewScene* InPreviewScene)
+	: PreviewScene(InPreviewScene)
+	, Viewport(nullptr)
+	, ViewFOV(90.0f)
+	, FOVAngle(90.0f)
+	, AspectRatio(1.777777f)
+	, EngineShowFlags(ESFIM_Game)
+{
+	ViewState.Allocate();
+
+	BackgroundColor = FColor(55, 55, 55);
+}
+
+FUMGViewportClient::~FUMGViewportClient()
+{
+
+}
+
+void FUMGViewportClient::AddReferencedObjects(FReferenceCollector & Collector)
+{
+
+}
+
+void FUMGViewportClient::Draw(FViewport* InViewport, FCanvas* Canvas)
+{
+	FViewport* ViewportBackup = Viewport;
+	Viewport = InViewport ? InViewport : Viewport;
+
+	// Determine whether we should use world time or real time based on the scene.
+	float TimeSeconds;
+	float RealTimeSeconds;
+	float DeltaTimeSeconds;
+
+	const bool bIsRealTime = true;
+
+	UWorld* World = GWorld;
+	if ( ( GetScene() != World->Scene ) || ( bIsRealTime == true ) )
+	{
+		// Use time relative to start time to avoid issues with float vs double
+		TimeSeconds = FApp::GetCurrentTime() - GStartTime;
+		RealTimeSeconds = FApp::GetCurrentTime() - GStartTime;
+		DeltaTimeSeconds = FApp::GetDeltaTime();
+	}
+	else
+	{
+		TimeSeconds = World->GetTimeSeconds();
+		RealTimeSeconds = World->GetRealTimeSeconds();
+		DeltaTimeSeconds = World->GetDeltaSeconds();
+	}
+
+	// Setup a FSceneViewFamily/FSceneView for the viewport.
+	FSceneViewFamilyContext ViewFamily(FSceneViewFamily::ConstructionValues(
+		Canvas->GetRenderTarget(),
+		GetScene(),
+		EngineShowFlags)
+		.SetWorldTimes(TimeSeconds, DeltaTimeSeconds, RealTimeSeconds)
+		.SetRealtimeUpdate(bIsRealTime));
+
+	ViewFamily.EngineShowFlags = EngineShowFlags;
+
+	//UpdateLightingShowFlags(ViewFamily.EngineShowFlags);
+
+	//ViewFamily.ExposureSettings = ExposureSettings;
+
+	//ViewFamily.LandscapeLODOverride = LandscapeLODOverride;
+
+	FSceneView* View = CalcSceneView(&ViewFamily);
+
+	//SetupViewForRendering(ViewFamily, *View);
+
+	FSlateRect SafeFrame;
+	View->CameraConstrainedViewRect = View->UnscaledViewRect;
+	//if ( CalculateEditorConstrainedViewRect(SafeFrame, Viewport) )
+	//{
+	//	View->CameraConstrainedViewRect = FIntRect(SafeFrame.Left, SafeFrame.Top, SafeFrame.Right, SafeFrame.Bottom);
+	//}
+
+	if ( IsAspectRatioConstrained() )
+	{
+		// Clear the background to black if the aspect ratio is constrained, as the scene view won't write to all pixels.
+		Canvas->Clear(FLinearColor::Black);
+	}
+
+	Canvas->Clear(FLinearColor::Black);
+
+	GetRendererModule().BeginRenderingViewFamily(Canvas, &ViewFamily);
+
+	// Remove temporary debug lines.
+	// Possibly a hack. Lines may get added without the scene being rendered etc.
+	if ( World->LineBatcher != NULL && ( World->LineBatcher->BatchedLines.Num() || World->LineBatcher->BatchedPoints.Num() ) )
+	{
+		World->LineBatcher->Flush();
+	}
+
+	if ( World->ForegroundLineBatcher != NULL && ( World->ForegroundLineBatcher->BatchedLines.Num() || World->ForegroundLineBatcher->BatchedPoints.Num() ) )
+	{
+		World->ForegroundLineBatcher->Flush();
+	}
+	
+	//FCanvas* DebugCanvas = Viewport->GetDebugCanvas();
+
+	//UDebugDrawService::Draw(ViewFamily.EngineShowFlags, Viewport, View, DebugCanvas);
+
+	//
+	//FlushRenderingCommands();
+
+	Viewport = ViewportBackup;
+
+	Canvas->Clear(FLinearColor::Red);
+}
+
+FSceneInterface* FUMGViewportClient::GetScene() const
+{
+	UWorld* World = GetWorld();
+	if ( World )
+	{
+		return World->Scene;
+	}
+
+	return NULL;
+}
+
+UWorld* FUMGViewportClient::GetWorld() const
+{
+	UWorld* OutWorldPtr = NULL;
+	// If we have a valid scene get its world
+	if ( PreviewScene )
+	{
+		OutWorldPtr = PreviewScene->GetWorld();
+	}
+	if ( OutWorldPtr == NULL )
+	{
+		OutWorldPtr = GWorld;
+	}
+	return OutWorldPtr;
+}
+
+bool FUMGViewportClient::IsOrtho() const
+{
+	return ViewInfo.ProjectionMode == ECameraProjectionMode::Orthographic;
+}
+
+bool FUMGViewportClient::IsPerspective() const
+{
+	return ViewInfo.ProjectionMode == ECameraProjectionMode::Perspective;
+}
+
+bool FUMGViewportClient::IsAspectRatioConstrained() const
+{
+	return ViewInfo.bConstrainAspectRatio;
+}
+
+ECameraProjectionMode::Type FUMGViewportClient::GetProjectionType() const
+{
+	return ViewInfo.ProjectionMode;
+}
+
+FLinearColor FUMGViewportClient::GetBackgroundColor() const
+{
+	return BackgroundColor;
+}
+
+float FUMGViewportClient::GetNearClipPlane() const
+{
+	return ( NearPlane < 0.0f ) ? GNearClippingPlane : NearPlane;
+}
+
+void FUMGViewportClient::OverrideNearClipPlane(float InNearPlane)
+{
+	NearPlane = InNearPlane;
+}
+
+float FUMGViewportClient::GetFarClipPlaneOverride() const
+{
+	return FarPlane;
+}
+
+void FUMGViewportClient::OverrideFarClipPlane(const float InFarPlane)
+{
+	FarPlane = InFarPlane;
+}
+
+float FUMGViewportClient::GetOrthoUnitsPerPixel(const FViewport* InViewport) const
+{
+	const float SizeX = InViewport->GetSizeXY().X;
+
+	// 15.0f was coming from the CAMERA_ZOOM_DIV marco, seems it was chosen arbitrarily
+	return ( GetOrthoZoom() / ( SizeX * 15.f ) )/* * ComputeOrthoZoomFactor(SizeX)*/;
+}
+
+
+FSceneView* FUMGViewportClient::CalcSceneView(FSceneViewFamily* ViewFamily)
+{
+	FSceneViewInitOptions ViewInitOptions;
+
+	const FVector& ViewLocation = GetViewLocation();
+	const FRotator& ViewRotation = GetViewRotation();
+
+	const FIntPoint ViewportSizeXY = Viewport->GetSizeXY();
+
+	FIntRect ViewRect = FIntRect(0, 0, ViewportSizeXY.X, ViewportSizeXY.Y);
+	ViewInitOptions.SetViewRectangle(ViewRect);
+
+	const ECameraProjectionMode::Type ProjectionType = GetProjectionType();
+	const bool bConstrainAspectRatio = ViewInfo.bConstrainAspectRatio;
+
+	ViewInitOptions.ViewMatrix = FTranslationMatrix(-GetViewLocation());
+	if ( ProjectionType == ECameraProjectionMode::Perspective )
+	{
+		ViewInitOptions.ViewMatrix = ViewInitOptions.ViewMatrix * FInverseRotationMatrix(ViewRotation);
+
+		ViewInitOptions.ViewMatrix = ViewInitOptions.ViewMatrix * FMatrix(
+			FPlane(0, 0, 1, 0),
+			FPlane(1, 0, 0, 0),
+			FPlane(0, 1, 0, 0),
+			FPlane(0, 0, 0, 1));
+
+		float MinZ = GetNearClipPlane();
+		float MaxZ = MinZ;
+		float MatrixFOV = ViewFOV * (float)PI / 360.0f;
+
+		if ( bConstrainAspectRatio )
+		{
+			ViewInitOptions.ProjectionMatrix = FReversedZPerspectiveMatrix(
+				MatrixFOV,
+				MatrixFOV,
+				1.0f,
+				AspectRatio,
+				MinZ,
+				MaxZ
+				);
+
+			ViewInitOptions.SetConstrainedViewRectangle(Viewport->CalculateViewExtents(AspectRatio, ViewRect));
+		}
+		else
+		{
+			float XAxisMultiplier;
+			float YAxisMultiplier;
+
+			if ( ViewportSizeXY.X > ViewportSizeXY.Y )
+			{
+				//if the viewport is wider than it is tall
+				XAxisMultiplier = 1.0f;
+				YAxisMultiplier = ViewportSizeXY.X / (float)ViewportSizeXY.Y;
+			}
+			else
+			{
+				//if the viewport is taller than it is wide
+				XAxisMultiplier = ViewportSizeXY.Y / (float)ViewportSizeXY.X;
+				YAxisMultiplier = 1.0f;
+			}
+
+			ViewInitOptions.ProjectionMatrix = FReversedZPerspectiveMatrix(
+				MatrixFOV,
+				MatrixFOV,
+				XAxisMultiplier,
+				YAxisMultiplier,
+				MinZ,
+				MaxZ
+				);
+		}
+	}
+	else if ( ProjectionType == ECameraProjectionMode::Orthographic )
+	{
+		float ZScale = 0.5f / HALF_WORLD_MAX;
+		float ZOffset = HALF_WORLD_MAX;
+
+		//The divisor for the matrix needs to match the translation code.
+		const float Zoom = GetOrthoUnitsPerPixel(Viewport);
+
+		float OrthoWidth = Zoom * ViewportSizeXY.X / 2.0f;
+		float OrthoHeight = Zoom * ViewportSizeXY.Y / 2.0f;
+
+		ViewInitOptions.ViewMatrix = ViewInitOptions.ViewMatrix * FInverseRotationMatrix(ViewRotation);
+		ViewInitOptions.ViewMatrix = ViewInitOptions.ViewMatrix * FMatrix(
+			FPlane(0, 0, 1, 0),
+			FPlane(1, 0, 0, 0),
+			FPlane(0, 1, 0, 0),
+			FPlane(0, 0, 0, 1));
+
+		const float EffectiveAspectRatio = bConstrainAspectRatio ? AspectRatio : ( ViewportSizeXY.X / (float)ViewportSizeXY.Y );
+		const float YScale = 1.0f / EffectiveAspectRatio;
+		OrthoWidth = ViewInfo.OrthoWidth / 2.0f;
+		OrthoHeight = ( ViewInfo.OrthoWidth / 2.0f ) * YScale;
+
+		const float NearViewPlane = ViewInitOptions.ViewMatrix.GetOrigin().Z;
+		const float FarViewPlane = NearViewPlane - 2.0f*WORLD_MAX;
+
+		const float InverseRange = 1.0f / ( FarViewPlane - NearViewPlane );
+		ZScale = -2.0f * InverseRange;
+		ZOffset = -( FarViewPlane + NearViewPlane ) * InverseRange;
+
+		ViewInitOptions.ProjectionMatrix = FReversedZOrthoMatrix(
+			OrthoWidth,
+			OrthoHeight,
+			ZScale,
+			ZOffset
+			);
+
+		if ( bConstrainAspectRatio )
+		{
+			ViewInitOptions.SetConstrainedViewRectangle(Viewport->CalculateViewExtents(AspectRatio, ViewRect));
+		}
+	}
+	else
+	{
+		check(false);
+	}
+
+	ViewInitOptions.ViewFamily = ViewFamily;
+	ViewInitOptions.SceneViewStateInterface = ViewState.GetReference();
+	ViewInitOptions.ViewElementDrawer = this;
+
+	ViewInitOptions.BackgroundColor = GetBackgroundColor();
+
+	//ViewInitOptions.EditorViewBitflag = 0, // send the bit for this view - each actor will check it's visibility bits against this
+
+	// for ortho views to steal perspective view origin
+	ViewInitOptions.OverrideLODViewOrigin = FVector::ZeroVector;
+	ViewInitOptions.bUseFauxOrthoViewPos = true;
+
+	ViewInitOptions.OverrideFarClippingPlaneDistance = FarPlane;
+	//ViewInitOptions.CursorPos = CurrentMousePos;
+
+	FSceneView* View = new FSceneView(ViewInitOptions);
+
+	ViewFamily->Views.Add(View);
+
+	View->StartFinalPostprocessSettings(ViewLocation);
+
+	//OverridePostProcessSettings(*View);
+
+	View->EndFinalPostprocessSettings();
+
+	return View;
+}
+
+/////////////////////////////////////////////////////
+// SAutoRefreshViewport
+
+class SAutoRefreshViewport : public SViewport
+{
+	SLATE_BEGIN_ARGS(SAutoRefreshViewport)
+	{
+	}
+	SLATE_END_ARGS()
+
+	void Construct(const FArguments& InArgs)
+	{
+		SViewport::FArguments ParentArgs;
+		//ParentArgs.RenderDirectlyToWindow(true);
+		SViewport::Construct(ParentArgs);
+
+		ViewportClient = MakeShareable(new FUMGViewportClient(&PreviewScene));
+		Viewport = MakeShareable(new FSceneViewport(ViewportClient.Get(), SharedThis(this)));
+
+		// The viewport widget needs an interface so it knows what should render
+		SetViewportInterface(Viewport.ToSharedRef());
+	}
+
+	void Tick(const FGeometry& AllottedGeometry, const double InCurrentTime, const float InDeltaTime) override
+	{
+		Viewport->Invalidate();
+		Viewport->InvalidateDisplay();
+
+		Viewport->Tick(InDeltaTime);
+		
+		//ViewportClient->Tick(InDeltaTime);
+	}
+
+public:
+	TSharedPtr<FUMGViewportClient> ViewportClient;
+	
+	TSharedPtr<FSceneViewport> Viewport;
+
+	/** preview scene */
+	FPreviewScene PreviewScene;
+};
+
+/////////////////////////////////////////////////////
+// UViewport
+
+UViewport::UViewport(const FPostConstructInitializeProperties& PCIP)
+	: Super(PCIP)
+{
+	bIsVariable = true;
+}
+
+void UViewport::ReleaseNativeWidget()
+{
+	Super::ReleaseNativeWidget();
+
+	ViewportWidget.Reset();
+}
+
+TSharedRef<SWidget> UViewport::RebuildWidget()
+{
+	ViewportWidget = SNew(SAutoRefreshViewport);
+
+	if ( GetChildrenCount() > 0 )
+	{
+		ViewportWidget->SetContent(GetContentSlot()->Content ? GetContentSlot()->Content->TakeWidget() : SNullWidget::NullWidget);
+	}
+
+	return BuildDesignTimeWidget( ViewportWidget.ToSharedRef() );
+}
+
+void UViewport::SyncronizeProperties()
+{
+	Super::SyncronizeProperties();
+}
+
+void UViewport::OnSlotAdded(UPanelSlot* Slot)
+{
+	// Add the child to the live canvas if it already exists
+	if ( ViewportWidget.IsValid() )
+	{
+		ViewportWidget->SetContent(Slot->Content ? Slot->Content->TakeWidget() : SNullWidget::NullWidget);
+	}
+}
+
+void UViewport::OnSlotRemoved(UPanelSlot* Slot)
+{
+	// Remove the widget from the live slot if it exists.
+	if ( ViewportWidget.IsValid() )
+	{
+		ViewportWidget->SetContent(SNullWidget::NullWidget);
+	}
+}
+
+#if WITH_EDITOR
+
+const FSlateBrush* UViewport::GetEditorIcon()
+{
+	return FUMGStyle::Get().GetBrush("Widget.Viewport");
+}
+
+#endif
+
+/////////////////////////////////////////////////////
+
+#undef LOCTEXT_NAMESPACE
diff --git a/Engine/Source/Runtime/UMG/Private/Components/Widget.cpp b/Engine/Source/Runtime/UMG/Private/Components/Widget.cpp
--- a/Engine/Source/Runtime/UMG/Private/Components/Widget.cpp
+++ b/Engine/Source/Runtime/UMG/Private/Components/Widget.cpp
@@ -86,21 +86,19 @@ UPanelWidget* UWidget::GetParent() const
 	return NULL;
 }
 
-TSharedRef<SWidget> UWidget::GetWidget() const
+TSharedRef<SWidget> UWidget::TakeWidget()
 {
 	TSharedPtr<SWidget> SafeWidget;
 
-	UWidget* MutableThis = const_cast<UWidget*>( this );
-
 	if ( !MyWidget.IsValid() )
 	{
 		// We lie a bit about this being a const function.  If this is the first call it's not really const
 		// and we need to construct and cache the widget for the first run.  But instead of forcing everyone
 		// downstream to make RebuildWidget const and force every implementation to make things mutable, we
 		// just blow away the const here.
-		SafeWidget = MutableThis->RebuildWidget();
+		SafeWidget = RebuildWidget();
 		MyWidget = SafeWidget;
-		MutableThis->SyncronizeProperties();
+		SyncronizeProperties();
 	}
 	else
 	{
@@ -108,12 +106,23 @@ TSharedRef<SWidget> UWidget::GetWidget() const
 	}
 
 	// If it is a user widget wrap it in a SObjectWidget to keep the instance from being GC'ed
-	if ( MutableThis->IsA(UUserWidget::StaticClass()) )
+	if ( IsA(UUserWidget::StaticClass()) )
 	{
-		return SNew(SObjectWidget, Cast<UUserWidget>(MutableThis))
-		[
-			SafeWidget.ToSharedRef()
-		];
+		if ( MyGCWidget.IsValid() )
+		{
+			return MyGCWidget.Pin().ToSharedRef();
+		}
+		else
+		{
+			TSharedPtr<SWidget> SafeGCWidget = SNew(SObjectWidget, Cast<UUserWidget>(this))
+				[
+					SafeWidget.ToSharedRef()
+				];
+
+			MyGCWidget = SafeGCWidget;
+
+			return SafeGCWidget.ToSharedRef();
+		}
 	}
 	else
 	{
@@ -123,6 +132,11 @@ TSharedRef<SWidget> UWidget::GetWidget() const
 
 TSharedPtr<SWidget> UWidget::GetCachedWidget() const
 {
+	if ( MyGCWidget.IsValid() )
+	{
+		return MyGCWidget.Pin();
+	}
+
 	return MyWidget.Pin();
 }
 
diff --git a/Engine/Source/Runtime/UMG/Private/Components/WidgetSwitcher.cpp b/Engine/Source/Runtime/UMG/Private/Components/WidgetSwitcher.cpp
--- a/Engine/Source/Runtime/UMG/Private/Components/WidgetSwitcher.cpp
+++ b/Engine/Source/Runtime/UMG/Private/Components/WidgetSwitcher.cpp
@@ -73,7 +73,11 @@ void UWidgetSwitcher::OnSlotRemoved(UPanelSlot* Slot)
 	// Remove the widget from the live slot if it exists.
 	if ( MyWidgetSwitcher.IsValid() )
 	{
-		MyWidgetSwitcher->RemoveSlot(Slot->Content->GetWidget());
+		TSharedPtr<SWidget> Widget = Slot->Content->GetCachedWidget();
+		if ( Widget.IsValid() )
+		{
+			MyWidgetSwitcher->RemoveSlot(Widget.ToSharedRef());
+		}
 	}
 }
 
diff --git a/Engine/Source/Runtime/UMG/Private/Components/WidgetSwitcherSlot.cpp b/Engine/Source/Runtime/UMG/Private/Components/WidgetSwitcherSlot.cpp
--- a/Engine/Source/Runtime/UMG/Private/Components/WidgetSwitcherSlot.cpp
+++ b/Engine/Source/Runtime/UMG/Private/Components/WidgetSwitcherSlot.cpp
@@ -27,7 +27,7 @@ void UWidgetSwitcherSlot::BuildSlot(TSharedRef<SWidgetSwitcher> WidgetSwitcher)
 		.HAlign(HorizontalAlignment)
 		.VAlign(VerticalAlignment)
 		[
-			Content == NULL ? SNullWidget::NullWidget : Content->GetWidget()
+			Content == NULL ? SNullWidget::NullWidget : Content->TakeWidget()
 		];
 }
 
diff --git a/Engine/Source/Runtime/UMG/Private/Slate/SObjectWidget.cpp b/Engine/Source/Runtime/UMG/Private/Slate/SObjectWidget.cpp
--- a/Engine/Source/Runtime/UMG/Private/Slate/SObjectWidget.cpp
+++ b/Engine/Source/Runtime/UMG/Private/Slate/SObjectWidget.cpp
@@ -33,9 +33,11 @@ void SObjectWidget::Tick(const FGeometry& AllottedGeometry, const double InCurre
 	}
 }
 
-//bool OnHitTest(const FGeometry& MyGeometry, FVector2D InAbsoluteCursorPosition)
+//bool SObjectWidget::OnHitTest(const FGeometry& MyGeometry, FVector2D InAbsoluteCursorPosition)
 //{
-
+//	if ( WidgetObject && !WidgetObject->IsDesignTime() )
+//	{
+//	}
 //}
 
 int32 SObjectWidget::OnPaint(const FGeometry& AllottedGeometry, const FSlateRect& MyClippingRect, FSlateWindowElementList& OutDrawElements, int32 LayerId, const FWidgetStyle& InWidgetStyle, bool bParentEnabled) const
diff --git a/Engine/Source/Runtime/UMG/Private/UMGStyle.cpp b/Engine/Source/Runtime/UMG/Private/UMGStyle.cpp
--- a/Engine/Source/Runtime/UMG/Private/UMGStyle.cpp
+++ b/Engine/Source/Runtime/UMG/Private/UMGStyle.cpp
@@ -66,6 +66,7 @@ TSharedRef< FSlateStyleSet > FUMGStyle::Create()
 	Style->Set("Widget.WidgetSwitcher", new IMAGE_BRUSH(TEXT("WidgetSwitcher"), Icon16x16));
 	Style->Set("Widget.MultiLineEditableText", new IMAGE_BRUSH(TEXT("MultiLineEditableText"), Icon16x16));
 	Style->Set("Widget.MultiLineEditableTextBox", new IMAGE_BRUSH(TEXT("MultiLineEditableTextBox"), Icon16x16));
+	Style->Set("Widget.Viewport", new IMAGE_BRUSH(TEXT("Viewport"), Icon16x16));
 
 	Style->Set("Widget.UserWidget", new IMAGE_BRUSH(TEXT("UserWidget"), Icon16x16));
 
diff --git a/Engine/Source/Runtime/UMG/Private/UserWidget.cpp b/Engine/Source/Runtime/UMG/Private/UserWidget.cpp
--- a/Engine/Source/Runtime/UMG/Private/UserWidget.cpp
+++ b/Engine/Source/Runtime/UMG/Private/UserWidget.cpp
@@ -260,7 +260,7 @@ TSharedRef<SWidget> UUserWidget::RebuildWidget()
 	// Add the first component to the root of the widget surface.
 	if ( Components.Num() > 0 )
 	{
-		UserRootWidget = Components[0]->GetWidget();
+		UserRootWidget = Components[0]->TakeWidget();
 	}
 	else
 	{
@@ -271,9 +271,11 @@ TSharedRef<SWidget> UUserWidget::RebuildWidget()
 	for ( int32 ComponentIndex = 0; ComponentIndex < Components.Num(); ++ComponentIndex )
 	{
 		UWidget* Handle = Components[ComponentIndex];
-		TSharedRef<SWidget> Widget = Handle->GetWidget();
-
-		WidgetToComponent.Add(Widget, Handle);
+		TSharedPtr<SWidget> Widget = Handle->GetCachedWidget();
+		if ( Widget.IsValid() )
+		{
+			WidgetToComponent.Add(Widget, Handle);
+		}
 	}
 
 	if ( !IsDesignTime() )
@@ -334,7 +336,7 @@ TSharedRef<SWidget> UUserWidget::MakeFullScreenWidget()
 			.VAlign(VerticalAlignment)
 			.HAlign(HorizontalAlignment)
 			[
-				GetWidget()
+				TakeWidget()
 			];
 	}
 	else
@@ -346,7 +348,7 @@ TSharedRef<SWidget> UUserWidget::MakeFullScreenWidget()
 			.HAlign(HorizontalAlignment)
 			.VAlign(VerticalAlignment)
 			[
-				GetWidget()
+				TakeWidget()
 			];
 
 		NewSlot.SizeParam = UWidget::ConvertSerializedSizeParamToRuntime(Size);
diff --git a/Engine/Source/Runtime/UMG/Public/Components/EditableTextBox.h b/Engine/Source/Runtime/UMG/Public/Components/EditableTextBox.h
--- a/Engine/Source/Runtime/UMG/Public/Components/EditableTextBox.h
+++ b/Engine/Source/Runtime/UMG/Public/Components/EditableTextBox.h
@@ -103,6 +103,9 @@ public:
 	UFUNCTION(BlueprintCallable, Category="Widget")
 	void SetError(FText InError);
 
+	UFUNCTION(BlueprintCallable, Category="Widget")
+	void ClearError();
+
 	// UWidget interface
 	virtual void SyncronizeProperties() override;
 	// End of UWidget interface
diff --git a/Engine/Source/Runtime/UMG/Public/Components/Viewport.h b/Engine/Source/Runtime/UMG/Public/Components/Viewport.h
new file mode 100644
--- /dev/null
+++ b/Engine/Source/Runtime/UMG/Public/Components/Viewport.h
@@ -0,0 +1,260 @@
+// Copyright 1998-2014 Epic Games, Inc. All Rights Reserved.
+
+#pragma once
+
+#include "Camera/CameraTypes.h"
+
+#include "Viewport.generated.h"
+
+class FPreviewScene;
+
+/**
+* Stores the transformation data for the viewport camera
+*/
+struct UMG_API FUMGViewportCameraTransform
+{
+public:
+	FUMGViewportCameraTransform();
+
+	/** Sets the transform's location */
+	void SetLocation(const FVector& Position);
+
+	/** Sets the transform's rotation */
+	void SetRotation(const FRotator& Rotation)
+	{
+		ViewRotation = Rotation;
+	}
+
+	/** Sets the location to look at during orbit */
+	void SetLookAt(const FVector& InLookAt)
+	{
+		LookAt = InLookAt;
+	}
+
+	/** Set the ortho zoom amount */
+	void SetOrthoZoom(float InOrthoZoom)
+	{
+		OrthoZoom = InOrthoZoom;
+	}
+
+
+	/** @return The transform's location */
+	FORCEINLINE const FVector& GetLocation() const { return ViewLocation; }
+
+	/** @return The transform's rotation */
+	FORCEINLINE const FRotator& GetRotation() const { return ViewRotation; }
+
+	/** @return The look at point for orbiting */
+	FORCEINLINE const FVector& GetLookAt() const { return LookAt; }
+
+	/** @return The ortho zoom amount */
+	FORCEINLINE float GetOrthoZoom() const { return OrthoZoom; }
+
+	/**
+	* Animates from the current location to the desired location
+	*
+	* @param InDesiredLocation	The location to transition to
+	* @param bInstant			If the desired location should be set instantly rather than transitioned to over time
+	*/
+	void TransitionToLocation(const FVector& InDesiredLocation, bool bInstant);
+
+	/**
+	* Updates any current location transitions
+	*
+	* @return true if there is currently a transition
+	*/
+	bool UpdateTransition();
+
+	/**
+	* Computes a matrix to use for viewport location and rotation when orbiting
+	*/
+	FMatrix ComputeOrbitMatrix() const;
+private:
+	/** Curve for animating between locations */
+	TSharedPtr<struct FCurveSequence> TransitionCurve;
+	/** Current viewport Position. */
+	FVector	ViewLocation;
+	/** Current Viewport orientation; valid only for perspective projections. */
+	FRotator ViewRotation;
+	/** Desired viewport location when animating between two locations */
+	FVector	DesiredLocation;
+	/** When orbiting, the point we are looking at */
+	FVector LookAt;
+	/** Viewport start location when animating to another location */
+	FVector StartLocation;
+	/** Ortho zoom amount */
+	float OrthoZoom;
+};
+
+
+class UMG_API FUMGViewportClient : public FCommonViewportClient, public FViewElementDrawer, public FGCObject
+{
+public:
+	FUMGViewportClient(FPreviewScene* InPreviewScene = nullptr);
+	virtual ~FUMGViewportClient();
+
+	virtual void AddReferencedObjects(FReferenceCollector & Collector) override;
+
+	virtual void Draw(FViewport* InViewport, FCanvas* Canvas) override;
+
+	virtual UWorld* GetWorld() const override;
+
+	/**
+	* @return The scene being rendered in this viewport
+	*/
+	virtual FSceneInterface* GetScene() const;
+
+	virtual FSceneView* CalcSceneView(FSceneViewFamily* ViewFamily);
+
+	bool IsOrtho() const;
+
+	bool IsPerspective() const;
+
+	bool IsAspectRatioConstrained() const;
+
+	FLinearColor GetBackgroundColor() const;
+
+	float GetNearClipPlane() const;
+
+	void OverrideNearClipPlane(float InNearPlane);
+
+	float GetFarClipPlaneOverride() const;
+
+	void OverrideFarClipPlane(const float InFarPlane);
+
+	ECameraProjectionMode::Type GetProjectionType() const;
+
+	/** Sets the location of the viewport's camera */
+	void SetViewLocation(const FVector& NewLocation)
+	{
+		ViewTransform.SetLocation(NewLocation);
+	}
+
+	/** Sets the location of the viewport's camera */
+	void SetViewRotation(const FRotator& NewRotation)
+	{
+		ViewTransform.SetRotation(NewRotation);
+	}
+
+	/**
+	* Sets the look at location of the viewports camera for orbit *
+	*
+	* @param LookAt The new look at location
+	* @param bRecalulateView	If true, will recalculate view location and rotation to look at the new point immediatley
+	*/
+	void SetLookAtLocation(const FVector& LookAt, bool bRecalculateView = false)
+	{
+		ViewTransform.SetLookAt(LookAt);
+
+		if ( bRecalculateView )
+		{
+			FMatrix OrbitMatrix = ViewTransform.ComputeOrbitMatrix();
+			OrbitMatrix = OrbitMatrix.Inverse();
+
+			ViewTransform.SetRotation(OrbitMatrix.Rotator());
+			ViewTransform.SetLocation(OrbitMatrix.GetOrigin());
+		}
+	}
+
+	/** Sets ortho zoom amount */
+	void SetOrthoZoom(float InOrthoZoom)
+	{
+		// A zero ortho zoom is not supported and causes NaN/div0 errors
+		check(InOrthoZoom != 0);
+		ViewTransform.SetOrthoZoom(InOrthoZoom);
+	}
+
+	/** @return the current viewport camera location */
+	const FVector& GetViewLocation() const
+	{
+		return ViewTransform.GetLocation();
+	}
+
+	/** @return the current viewport camera rotation */
+	const FRotator& GetViewRotation() const
+	{
+		return ViewTransform.GetRotation();
+	}
+
+	/** @return the current look at location */
+	const FVector& GetLookAtLocation() const
+	{
+		return ViewTransform.GetLookAt();
+	}
+
+	/** @return the current ortho zoom amount */
+	float GetOrthoZoom() const
+	{
+		return ViewTransform.GetOrthoZoom();
+	}
+
+	/** @return The number of units per pixel displayed in this viewport */
+	float GetOrthoUnitsPerPixel(const FViewport* Viewport) const;
+
+protected:
+
+	/** The scene used for the viewport. Owned externally */
+	FPreviewScene* PreviewScene;
+
+	FMinimalViewInfo ViewInfo;
+
+	FLinearColor BackgroundColor;
+
+	/** Viewport camera transform data */
+	FUMGViewportCameraTransform ViewTransform;
+
+	FViewport* Viewport;
+
+	/** near plane adjustable for each editor view, if < 0 GNearClippingPlane should be used. */
+	float NearPlane;
+
+	/** If > 0, overrides the view's far clipping plane with a plane at the specified distance. */
+	float FarPlane;
+
+	/** Viewport's current horizontal field of view (can be modified by locked cameras etc.) */
+	float ViewFOV;
+	/** Viewport's stored horizontal field of view (saved in ini files). */
+	float FOVAngle;
+
+	float AspectRatio;
+
+	/** The viewport's scene view state. */
+	FSceneViewStateReference ViewState;
+
+	/** A set of flags that determines visibility for various scene elements. */
+	FEngineShowFlags		EngineShowFlags;
+};
+
+/**
+ * 
+ */
+UCLASS(meta=( Category="Common" ), ClassGroup=UserInterface)
+class UMG_API UViewport : public UContentWidget
+{
+	GENERATED_UCLASS_BODY()
+
+	// UWidget interface
+	virtual void SyncronizeProperties() override;
+	// End of UWidget interface
+
+	virtual void ReleaseNativeWidget() override;
+
+#if WITH_EDITOR
+	virtual const FSlateBrush* GetEditorIcon() override;
+#endif
+
+protected:
+
+	// UPanelWidget
+	virtual void OnSlotAdded(UPanelSlot* Slot) override;
+	virtual void OnSlotRemoved(UPanelSlot* Slot) override;
+	// End UPanelWidget
+
+protected:
+	TSharedPtr<class SAutoRefreshViewport> ViewportWidget;
+
+protected:
+	// UWidget interface
+	virtual TSharedRef<SWidget> RebuildWidget() override;
+	// End of UWidget interface
+};
diff --git a/Engine/Source/Runtime/UMG/Public/Components/Widget.h b/Engine/Source/Runtime/UMG/Public/Components/Widget.h
--- a/Engine/Source/Runtime/UMG/Public/Components/Widget.h
+++ b/Engine/Source/Runtime/UMG/Public/Components/Widget.h
@@ -140,8 +140,9 @@ public:
 	 * virtual however, you should not inherit this function unless you're very aware of what you're
 	 * doing.  Normal derived versions should only ever override RebuildWidget.
 	 */
-	TSharedRef<SWidget> GetWidget() const;
+	TSharedRef<SWidget> TakeWidget();
 
+	/** Gets the last created widget does not recreate the gc container for the widget if one is needed. */
 	TSharedPtr<SWidget> GetCachedWidget() const;
 
 	/**
@@ -233,7 +234,10 @@ protected:
 
 protected:
 	/** The underlying SWidget. */
-	mutable TWeakPtr<SWidget> MyWidget;
+	TWeakPtr<SWidget> MyWidget;
+
+	/** The underlying SWidget contained in a SObjectWidget */
+	TWeakPtr<SWidget> MyGCWidget;
 	
 	/** Is this widget being displayed on a designer surface */
 	bool bDesignTime;
diff --git a/Engine/Source/Runtime/UMG/Public/UMG.h b/Engine/Source/Runtime/UMG/Public/UMG.h
--- a/Engine/Source/Runtime/UMG/Public/UMG.h
+++ b/Engine/Source/Runtime/UMG/Public/UMG.h
@@ -83,6 +83,8 @@
 #include "ListView.h"
 #include "TileView.h"
 
+#include "Viewport.h"
+
 // Slate
 #include "SObjectWidget.h"
 
