Fix non-unity

diff --git a/Engine/Source/Runtime/UMG/Private/Animation/UMGSequencePlayer.cpp b/Engine/Source/Runtime/UMG/Private/Animation/UMGSequencePlayer.cpp
--- a/Engine/Source/Runtime/UMG/Private/Animation/UMGSequencePlayer.cpp
+++ b/Engine/Source/Runtime/UMG/Private/Animation/UMGSequencePlayer.cpp
@@ -4,6 +4,7 @@
 #include "UMGSequencePlayer.h"
 #include "MovieSceneBindings.h"
 #include "MovieSceneInstance.h"
+#include "MovieScene.h"
 
 UUMGSequencePlayer::UUMGSequencePlayer(const FPostConstructInitializeProperties& PCIP)
 	: Super(PCIP)
