UMG  - Fixing build.

diff --git a/Engine/Source/Runtime/UMG/Private/Components/Viewport.cpp b/Engine/Source/Runtime/UMG/Private/Components/Viewport.cpp
--- a/Engine/Source/Runtime/UMG/Private/Components/Viewport.cpp
+++ b/Engine/Source/Runtime/UMG/Private/Components/Viewport.cpp
@@ -397,8 +397,8 @@ FSceneView* FUMGViewportClient::CalcSceneView(FSceneViewFamily* ViewFamily)
 	//ViewInitOptions.EditorViewBitflag = 0, // send the bit for this view - each actor will check it's visibility bits against this
 
 	// for ortho views to steal perspective view origin
-	ViewInitOptions.OverrideLODViewOrigin = FVector::ZeroVector;
-	ViewInitOptions.bUseFauxOrthoViewPos = true;
+	//ViewInitOptions.OverrideLODViewOrigin = FVector::ZeroVector;
+	//ViewInitOptions.bUseFauxOrthoViewPos = true;
 
 	ViewInitOptions.OverrideFarClippingPlaneDistance = FarPlane;
 	//ViewInitOptions.CursorPos = CurrentMousePos;
