- Clang fix for multiple inheritance

diff --git a/Engine/Source/Runtime/UMG/Public/Components/Viewport.h b/Engine/Source/Runtime/UMG/Public/Components/Viewport.h
--- a/Engine/Source/Runtime/UMG/Public/Components/Viewport.h
+++ b/Engine/Source/Runtime/UMG/Public/Components/Viewport.h
@@ -96,6 +96,7 @@ public:
 	virtual void AddReferencedObjects(FReferenceCollector & Collector) override;
 
 	virtual void Draw(FViewport* InViewport, FCanvas* Canvas) override;
+	using FViewElementDrawer::Draw;
 
 	virtual UWorld* GetWorld() const override;
 
