Build - Fixing non unity build.

diff --git a/Engine/Source/Runtime/UMG/Private/Components/Viewport.cpp b/Engine/Source/Runtime/UMG/Private/Components/Viewport.cpp
--- a/Engine/Source/Runtime/UMG/Private/Components/Viewport.cpp
+++ b/Engine/Source/Runtime/UMG/Private/Components/Viewport.cpp
@@ -4,6 +4,7 @@
 #include "Slate/SlateBrushAsset.h"
 #include "PreviewScene.h"
 #include "EngineModule.h"
+#include "SceneViewport.h"
 
 #define LOCTEXT_NAMESPACE "UMG"
 
