UMG - Images and their color/opacity are now bindable in the UImage.

diff --git a/Engine/Source/Runtime/UMG/Private/Components/Image.cpp b/Engine/Source/Runtime/UMG/Private/Components/Image.cpp
--- a/Engine/Source/Runtime/UMG/Private/Components/Image.cpp
+++ b/Engine/Source/Runtime/UMG/Private/Components/Image.cpp
@@ -30,8 +30,11 @@ TSharedRef<SWidget> UImage::RebuildWidget()
 
 void UImage::SyncronizeProperties()
 {
-	MyImage->SetImage(GetImageBrush());
-	MyImage->SetColorAndOpacity(ColorAndOpacity);
+	TAttribute<FSlateColor> ColorAndOpacityBinding = OPTIONAL_BINDING(FSlateColor, ColorAndOpacity);
+	TAttribute<const FSlateBrush*> ImageBinding = OPTIONAL_BINDING_CONVERT(USlateBrushAsset*, Image, const FSlateBrush*, ConvertImage);
+
+	MyImage->SetImage(ImageBinding);
+	MyImage->SetColorAndOpacity(ColorAndOpacityBinding);
 	MyImage->SetOnMouseButtonDown(BIND_UOBJECT_DELEGATE(FPointerEventHandler, HandleMouseButtonDown));
 }
 
@@ -53,8 +56,41 @@ void UImage::SetOpacity(float InOpacity)
 	}
 }
 
+const FSlateBrush* UImage::ConvertImage(TAttribute<USlateBrushAsset*> InImageAsset) const
+{
+	USlateBrushAsset* ImageAsset = InImageAsset.Get();
+	if ( ImageAsset != Image )
+	{
+		UImage* MutableThis = const_cast< UImage* >( this );
+		MutableThis->DynamicBrush = TOptional<FSlateBrush>();
+		MutableThis->Image = ImageAsset;
+	}
+	else if ( DynamicBrush.IsSet() )
+	{
+		return &DynamicBrush.GetValue();
+	}
+
+	if ( Image == NULL )
+	{
+		SImage::FArguments ImageDefaults;
+		return ImageDefaults._Image.Get();
+	}
+
+	return &Image->Brush;
+}
+
+const FSlateBrush* UImage::GetImageBrush() const
+{
+	return ConvertImage(Image);
+}
+
 void UImage::SetImage(USlateBrushAsset* InImage)
 {
+	if ( InImage != Image )
+	{
+		DynamicBrush = TOptional<FSlateBrush>();
+	}
+
 	Image = InImage;
 	if ( MyImage.IsValid() )
 	{
@@ -66,7 +102,7 @@ UMaterialInstanceDynamic* UImage::GetDynamicMaterial()
 {
 	if ( !DynamicBrush.IsSet() )
 	{
-		const FSlateBrush* Brush = UImage::GetImageBrush();
+		const FSlateBrush* Brush = GetImageBrush();
 		FSlateBrush ClonedBrush = *Brush;
 		UObject* Resource = ClonedBrush.GetResourceObject();
 		UMaterialInterface* Material = Cast<UMaterialInterface>(Resource);
@@ -97,22 +133,6 @@ UMaterialInstanceDynamic* UImage::GetDynamicMaterial()
 	return NULL;
 }
 
-const FSlateBrush* UImage::GetImageBrush() const
-{
-	if ( DynamicBrush.IsSet() )
-	{
-		return &DynamicBrush.GetValue();
-	}
-
-	if ( Image == NULL )
-	{
-		SImage::FArguments ImageDefaults;
-		return ImageDefaults._Image.Get();
-	}
-
-	return &Image->Brush;
-}
-
 FReply UImage::HandleMouseButtonDown(const FGeometry& Geometry, const FPointerEvent& MouseEvent)
 {
 	if ( OnMouseButtonDownEvent.IsBound() )
diff --git a/Engine/Source/Runtime/UMG/Public/Components/Image.h b/Engine/Source/Runtime/UMG/Public/Components/Image.h
--- a/Engine/Source/Runtime/UMG/Public/Components/Image.h
+++ b/Engine/Source/Runtime/UMG/Public/Components/Image.h
@@ -18,10 +18,18 @@ public:
 	UPROPERTY(EditDefaultsOnly, Category=Appearance, meta=( DisplayThumbnail = "true" ))
 	USlateBrushAsset* Image;
 
+	/** A bindable delegate for the Image. */
+	UPROPERTY()
+	FGetSlateBrushAsset ImageDelegate;
+
 	/** Color and opacity */
 	UPROPERTY(EditDefaultsOnly, Category=Appearance)
 	FLinearColor ColorAndOpacity;
 
+	/** A bindable delegate for the ColorAndOpacity. */
+	UPROPERTY()
+	FGetSlateColor ColorAndOpacityDelegate;
+
 	UPROPERTY(EditDefaultsOnly, Category=Events)
 	FOnPointerEvent OnMouseButtonDownEvent;
 
@@ -60,6 +68,8 @@ protected:
 
 	const FSlateBrush* GetImageBrush() const;
 
+	const FSlateBrush* ConvertImage(TAttribute<USlateBrushAsset*> InImageAsset) const;
+
 	FReply HandleMouseButtonDown(const FGeometry& Geometry, const FPointerEvent& MouseEvent);
 
 protected:
