UMG - User widgets now have an associated player controller when spawned.

diff --git a/Engine/Source/Runtime/UMG/Private/UserWidget.cpp b/Engine/Source/Runtime/UMG/Private/UserWidget.cpp
--- a/Engine/Source/Runtime/UMG/Private/UserWidget.cpp
+++ b/Engine/Source/Runtime/UMG/Private/UserWidget.cpp
@@ -177,17 +177,15 @@ void UUserWidget::PostInitProperties()
 
 UWorld* UUserWidget::GetWorld() const
 {
-	UObject* Outer = GetOuter();
-	if ( Outer == NULL )
+	if ( PlayerContext.IsValid() )
 	{
-		return NULL;
+		return PlayerContext.GetWorld();
 	}
 
-	// TODO UMG Global UI elements should go where?  who will be their outer.  Currently CreateWidget node makes the level their owner.
-	if ( ULevel* Level = Cast<ULevel>(Outer) )
-	{
-		return Level->OwningWorld;
-	}
+	//if ( ULevel* Level = Cast<ULevel>(Outer) )
+	//{
+	//	return Level->OwningWorld;
+	//}
 
 	return NULL;
 }
@@ -444,3 +442,24 @@ TEnumAsByte<ESlateVisibility::Type> UUserWidget::GetVisiblity()
 
 	return ESlateVisibility::Collapsed;
 }
+
+void UUserWidget::SetPlayerContext(FLocalPlayerContext InPlayerContext)
+{
+	PlayerContext = InPlayerContext;
+}
+
+const FLocalPlayerContext& UUserWidget::GetPlayerContext() const
+{
+	return PlayerContext;
+}
+
+ULocalPlayer* UUserWidget::GetLocalPlayer() const
+{
+	APlayerController* PC = PlayerContext.IsValid() ? PlayerContext.GetPlayerController() : NULL;
+	return PC ? Cast<ULocalPlayer>(PC->Player) : NULL;
+}
+
+APlayerController* UUserWidget::GetPlayerController() const
+{
+	return PlayerContext.IsValid() ? PlayerContext.GetPlayerController() : NULL;
+}
diff --git a/Engine/Source/Runtime/UMG/Private/WidgetBlueprintLibrary.cpp b/Engine/Source/Runtime/UMG/Private/WidgetBlueprintLibrary.cpp
--- a/Engine/Source/Runtime/UMG/Private/WidgetBlueprintLibrary.cpp
+++ b/Engine/Source/Runtime/UMG/Private/WidgetBlueprintLibrary.cpp
@@ -14,7 +14,7 @@ UWidgetBlueprintLibrary::UWidgetBlueprintLibrary(const FPostConstructInitializeP
 {
 }
 
-UUserWidget* UWidgetBlueprintLibrary::Create(UObject* WorldContextObject, TSubclassOf<UUserWidget> WidgetType)
+UUserWidget* UWidgetBlueprintLibrary::Create(UObject* WorldContextObject, TSubclassOf<UUserWidget> WidgetType, APlayerController* OwningPlayer)
 {
 	if ( WidgetType == NULL || WidgetType->HasAnyClassFlags(CLASS_Abstract) )
 	{
@@ -22,10 +22,22 @@ UUserWidget* UWidgetBlueprintLibrary::Create(UObject* WorldContextObject, TSubcl
 		return NULL;
 	}
 
-	UWorld* World = GEngine->GetWorldFromContextObject(WorldContextObject);
-	UUserWidget* NewWidget = ConstructObject<UUserWidget>(WidgetType, World->GetCurrentLevel());
-	NewWidget->SetFlags(RF_Transactional);
-	
+	UUserWidget* NewWidget = NULL;
+
+	if ( OwningPlayer == NULL )
+	{
+		UWorld* World = GEngine->GetWorldFromContextObject(WorldContextObject);
+		ULocalPlayer* Player = GEngine->GetGamePlayer(World, 0);
+
+		NewWidget = ConstructObject<UUserWidget>(WidgetType, Player);
+		NewWidget->SetPlayerContext(FLocalPlayerContext(Player));
+	}
+	else
+	{
+		NewWidget = ConstructObject<UUserWidget>(WidgetType, OwningPlayer);
+		NewWidget->SetPlayerContext(FLocalPlayerContext(OwningPlayer));
+	}
+
 	return NewWidget;
 }
 
diff --git a/Engine/Source/Runtime/UMG/Public/Blueprint/UserWidget.h b/Engine/Source/Runtime/UMG/Public/Blueprint/UserWidget.h
--- a/Engine/Source/Runtime/UMG/Public/Blueprint/UserWidget.h
+++ b/Engine/Source/Runtime/UMG/Public/Blueprint/UserWidget.h
@@ -122,6 +122,20 @@ public:
 	UFUNCTION(BlueprintPure, Category="Appearance")
 	TEnumAsByte<ESlateVisibility::Type> GetVisiblity();
 
+	/** Sets the player context associated with this UI. */
+	void SetPlayerContext(FLocalPlayerContext InPlayerContext);
+
+	/** Gets the player context associated with this UI. */
+	const FLocalPlayerContext& GetPlayerContext() const;
+
+	/** Gets the local player associated with this UI. */
+	UFUNCTION(BlueprintCallable, Category="Player")
+	class ULocalPlayer* GetLocalPlayer() const;
+
+	/** Gets the player controller associated with this UI. */
+	UFUNCTION(BlueprintCallable, Category="Player")
+	class APlayerController* GetPlayerController() const;
+
 	/** Called when the widget is constructed */
 	UFUNCTION(BlueprintImplementableEvent, Category="User Interface")
 	void Construct();
@@ -236,6 +250,7 @@ public:
 
 	/** Ticks this widget and forwards to the underlying widget to tick */
 	void NativeTick(const FGeometry& MyGeometry, float InDeltaTime );
+
 public:
 	/** Called when the visibility changes. */
 	UPROPERTY(BlueprintAssignable)
@@ -315,4 +330,6 @@ private:
 	TMap< TWeakPtr<SWidget>, TWeakObjectPtr<UWidget> > WidgetToComponent;
 
 	TWeakPtr<SWidget> FullScreenWidget;
+
+	FLocalPlayerContext PlayerContext;
 };
diff --git a/Engine/Source/Runtime/UMG/Public/Blueprint/WidgetBlueprintLibrary.h b/Engine/Source/Runtime/UMG/Public/Blueprint/WidgetBlueprintLibrary.h
--- a/Engine/Source/Runtime/UMG/Public/Blueprint/WidgetBlueprintLibrary.h
+++ b/Engine/Source/Runtime/UMG/Public/Blueprint/WidgetBlueprintLibrary.h
@@ -9,8 +9,8 @@ class UWidgetBlueprintLibrary : public UBlueprintFunctionLibrary
 	GENERATED_UCLASS_BODY()
 
 	/** Creates a widget */
-	UFUNCTION(BlueprintCallable, meta=( HidePin="WorldContextObject", DefaultToSelf="WorldContextObject", UnsafeDuringActorConstruction = "true", FriendlyName = "Create Widget" ), Category="User Interface|Widget")
-	static class UUserWidget* Create(UObject* WorldContextObject, TSubclassOf<class UUserWidget> WidgetType);
+	UFUNCTION(BlueprintCallable, meta=( HidePin="WorldContextObject", DefaultToSelf="WorldContextObject", FriendlyName = "Create Widget" ), Category="User Interface|Widget")
+	static class UUserWidget* Create(UObject* WorldContextObject, TSubclassOf<class UUserWidget> WidgetType, APlayerController* OwningPlayer);
 
 	///** Show a widget as a popup */
 	//UFUNCTION(BlueprintCallable, meta=( HidePin="WorldContextObject", DefaultToSelf="WorldContextObject", UnsafeDuringActorConstruction = "true" ), Category="User Interface" )
