UMG - Trimming down properties from the multiline editable textboxes that don't work.  Allowing justification enum to be used in UObjects.

diff --git a/Engine/Source/Runtime/Slate/Public/Framework/Text/TextLayout.h b/Engine/Source/Runtime/Slate/Public/Framework/Text/TextLayout.h
--- a/Engine/Source/Runtime/Slate/Public/Framework/Text/TextLayout.h
+++ b/Engine/Source/Runtime/Slate/Public/Framework/Text/TextLayout.h
@@ -1,8 +1,11 @@
 // Copyright 1998-2014 Epic Games, Inc. All Rights Reserved.
 #pragma once
 
+#include "TextLayout.generated.h"
+
 #define TEXT_LAYOUT_DEBUG 0
 
+UENUM()
 namespace ETextJustify
 {
 	enum Type
diff --git a/Engine/Source/Runtime/UMG/Private/Components/MultiLineEditableText.cpp b/Engine/Source/Runtime/UMG/Private/Components/MultiLineEditableText.cpp
--- a/Engine/Source/Runtime/UMG/Private/Components/MultiLineEditableText.cpp
+++ b/Engine/Source/Runtime/UMG/Private/Components/MultiLineEditableText.cpp
@@ -10,21 +10,8 @@
 UMultiLineEditableText::UMultiLineEditableText(const FPostConstructInitializeProperties& PCIP)
 	: Super(PCIP)
 {
-	ColorAndOpacity = FLinearColor::Black;
-
 	// HACK Special font initialization hack since there are no font assets yet for slate.
 	Font = FSlateFontInfo(TEXT("Slate/Fonts/Roboto-Bold.ttf"), 12);
-
-	// Grab other defaults from slate arguments.
-	SMultiLineEditableTextBox::FArguments Defaults;
-	IsReadOnly = Defaults._IsReadOnly.Get();
-	IsPassword = Defaults._IsReadOnly.Get();
-//	MinimumDesiredWidth = Defaults._MinDesiredWidth.Get();
-	IsCaretMovedWhenGainFocus = Defaults._IsCaretMovedWhenGainFocus.Get();
-	SelectAllTextWhenFocused = Defaults._SelectAllTextWhenFocused.Get();
-	RevertTextOnEscape = Defaults._RevertTextOnEscape.Get();
-	ClearKeyboardFocusOnCommit = Defaults._ClearKeyboardFocusOnCommit.Get();
-	SelectAllTextOnCommit = Defaults._SelectAllTextOnCommit.Get();
 }
 
 TSharedRef<SWidget> UMultiLineEditableText::RebuildWidget()
@@ -42,6 +29,7 @@ TSharedRef<SWidget> UMultiLineEditableText::RebuildWidget()
 	MyMultiLineEditableText = SNew(SMultiLineEditableText)
 	//.Style(StylePtr)
 	.Font(FSlateFontInfo(FontPath, Font.Size))
+	.Justification(Justification)
 //	.MinDesiredWidth(MinimumDesiredWidth)
 //	.IsCaretMovedWhenGainFocus(IsCaretMovedWhenGainFocus)
 //	.SelectAllTextWhenFocused(SelectAllTextWhenFocused)
diff --git a/Engine/Source/Runtime/UMG/Private/Components/MultiLineEditableTextBox.cpp b/Engine/Source/Runtime/UMG/Private/Components/MultiLineEditableTextBox.cpp
--- a/Engine/Source/Runtime/UMG/Private/Components/MultiLineEditableTextBox.cpp
+++ b/Engine/Source/Runtime/UMG/Private/Components/MultiLineEditableTextBox.cpp
@@ -16,18 +16,6 @@ UMultiLineEditableTextBox::UMultiLineEditableTextBox(const FPostConstructInitial
 
 	// HACK Special font initialization hack since there are no font assets yet for slate.
 	Font = FSlateFontInfo(TEXT("Slate/Fonts/Roboto-Bold.ttf"), 12);
-
-	// Grab other defaults from slate arguments.
-	SMultiLineEditableTextBox::FArguments Defaults;
-	IsReadOnly = Defaults._IsReadOnly.Get();
-	IsPassword = Defaults._IsReadOnly.Get();
-//	MinimumDesiredWidth = Defaults._MinDesiredWidth.Get();
-	Padding = Defaults._Padding.Get();
-	IsCaretMovedWhenGainFocus = Defaults._IsCaretMovedWhenGainFocus.Get();
-	SelectAllTextWhenFocused = Defaults._SelectAllTextWhenFocused.Get();
-	RevertTextOnEscape = Defaults._RevertTextOnEscape.Get();
-	ClearKeyboardFocusOnCommit = Defaults._ClearKeyboardFocusOnCommit.Get();
-	SelectAllTextOnCommit = Defaults._SelectAllTextOnCommit.Get();
 }
 
 TSharedRef<SWidget> UMultiLineEditableTextBox::RebuildWidget()
@@ -36,6 +24,7 @@ TSharedRef<SWidget> UMultiLineEditableTextBox::RebuildWidget()
 
 	MyEditableTextBlock = SNew(SMultiLineEditableTextBox)
 		.Font(FSlateFontInfo(FontPath, Font.Size))
+		.Justification(Justification)
 		.ForegroundColor(ForegroundColor)
 		.BackgroundColor(BackgroundColor)
 		.ReadOnlyForegroundColor(ReadOnlyForegroundColor)
diff --git a/Engine/Source/Runtime/UMG/Public/Components/EditableTextBox.h b/Engine/Source/Runtime/UMG/Public/Components/EditableTextBox.h
--- a/Engine/Source/Runtime/UMG/Public/Components/EditableTextBox.h
+++ b/Engine/Source/Runtime/UMG/Public/Components/EditableTextBox.h
@@ -35,15 +35,15 @@ public:
 
 	/** Text color and opacity (overrides Style) */
 	UPROPERTY(EditDefaultsOnly, Category=Appearance)
-	FSlateColor ForegroundColor;
+	FLinearColor ForegroundColor;
 
 	/** The color of the background/border around the editable text (overrides Style) */
 	UPROPERTY(EditDefaultsOnly, Category=Appearance)
-	FSlateColor BackgroundColor;
+	FLinearColor BackgroundColor;
 
 	/** Text color and opacity when read-only (overrides Style) */
 	UPROPERTY(EditDefaultsOnly, Category=Appearance)
-	FSlateColor ReadOnlyForegroundColor;
+	FLinearColor ReadOnlyForegroundColor;
 
 	/** Sets whether this text box can actually be modified interactively by the user */
 	UPROPERTY(EditDefaultsOnly, Category=Appearance)
diff --git a/Engine/Source/Runtime/UMG/Public/Components/MultiLineEditableText.h b/Engine/Source/Runtime/UMG/Public/Components/MultiLineEditableText.h
--- a/Engine/Source/Runtime/UMG/Public/Components/MultiLineEditableText.h
+++ b/Engine/Source/Runtime/UMG/Public/Components/MultiLineEditableText.h
@@ -21,70 +21,14 @@ public:
 	UPROPERTY(EditDefaultsOnly, Category=Content)
 	FText Text;
 
-	/** Hint text that appears when there is no text in the text box */
+	/** The justification of the text in the multilinebox */
 	UPROPERTY(EditDefaultsOnly, Category=Content)
-	FText HintText;
-
-	/** Text style */
-	UPROPERTY(EditDefaultsOnly, Category=Style, meta=( DisplayThumbnail = "true" ))
-	USlateWidgetStyleAsset* Style;
-
-	/** Background image for the selected text (overrides Style) */
-	UPROPERTY(EditDefaultsOnly, Category="Style", meta=( DisplayThumbnail = "true" ), AdvancedDisplay)
-	USlateBrushAsset* BackgroundImageSelected;
-
-	/** Background image for the selection targeting effect (overrides Style) */
-	UPROPERTY(EditDefaultsOnly, Category="Style", meta=( DisplayThumbnail = "true" ), AdvancedDisplay)
-	USlateBrushAsset* BackgroundImageSelectionTarget;
-
-	/** Background image for the composing text (overrides Style) */
-	UPROPERTY(EditDefaultsOnly, Category="Style", meta=( DisplayThumbnail = "true" ), AdvancedDisplay)
-	USlateBrushAsset* BackgroundImageComposing;
-
-	/** Image brush used for the caret (overrides Style) */
-	UPROPERTY(EditDefaultsOnly, Category="Style", meta=( DisplayThumbnail = "true" ), AdvancedDisplay)
-	USlateBrushAsset* CaretImage;
+	TEnumAsByte<ETextJustify::Type> Justification;
 
 	/** Font color and opacity (overrides Style) */
 	UPROPERTY(EditDefaultsOnly, Category=Appearance)
 	FSlateFontInfo Font;
 
-	/** Text color and opacity (overrides Style) */
-	UPROPERTY(EditDefaultsOnly, Category=Appearance)
-	FSlateColor ColorAndOpacity;
-
-	/** Sets whether this text box can actually be modified interactively by the user */
-	UPROPERTY(EditDefaultsOnly, Category=Appearance)
-	bool IsReadOnly;
-
-	/** Sets whether this text box is for storing a password */
-	UPROPERTY(EditDefaultsOnly, Category=Appearance)
-	bool IsPassword;
-
-	/** Minimum width that a text block should be */
-	UPROPERTY(EditDefaultsOnly, Category=Appearance)
-	float MinimumDesiredWidth;
-
-	/** Workaround as we loose focus when the auto completion closes. */
-	UPROPERTY(EditDefaultsOnly, Category=Behavior, AdvancedDisplay)
-	bool IsCaretMovedWhenGainFocus;
-
-	/** Whether to select all text when the user clicks to give focus on the widget */
-	UPROPERTY(EditDefaultsOnly, Category=Behavior, AdvancedDisplay)
-	bool SelectAllTextWhenFocused;
-
-	/** Whether to allow the user to back out of changes when they press the escape key */
-	UPROPERTY(EditDefaultsOnly, Category=Behavior, AdvancedDisplay)
-	bool RevertTextOnEscape;
-
-	/** Whether to clear keyboard focus when pressing enter to commit changes */
-	UPROPERTY(EditDefaultsOnly, Category=Behavior, AdvancedDisplay)
-	bool ClearKeyboardFocusOnCommit;
-
-	/** Whether to select all text when pressing enter to commit changes */
-	UPROPERTY(EditDefaultsOnly, Category=Behavior, AdvancedDisplay)
-	bool SelectAllTextOnCommit;
-
 	/** Called whenever the text is changed interactively by the user */
 	UPROPERTY(BlueprintAssignable, Category="Widget Event")
 	FOnMultiLineEditableTextChangedEvent OnTextChanged;
diff --git a/Engine/Source/Runtime/UMG/Public/Components/MultiLineEditableTextBox.h b/Engine/Source/Runtime/UMG/Public/Components/MultiLineEditableTextBox.h
--- a/Engine/Source/Runtime/UMG/Public/Components/MultiLineEditableTextBox.h
+++ b/Engine/Source/Runtime/UMG/Public/Components/MultiLineEditableTextBox.h
@@ -25,9 +25,9 @@ public:
 	UPROPERTY(EditDefaultsOnly, Category=Content)
 	FText Text;
 
-	/** Hint text that appears when there is no text in the text box */
+	/** The justification of the text in the multilinebox */
 	UPROPERTY(EditDefaultsOnly, Category=Content)
-	FText HintText;
+	TEnumAsByte<ETextJustify::Type> Justification;
 
 	/** Font color and opacity (overrides Style) */
 	UPROPERTY(EditDefaultsOnly, Category=Appearance)
@@ -35,51 +35,15 @@ public:
 
 	/** Text color and opacity (overrides Style) */
 	UPROPERTY(EditDefaultsOnly, Category=Appearance)
-	FSlateColor ForegroundColor;
+	FLinearColor ForegroundColor;
 
 	/** The color of the background/border around the editable text (overrides Style) */
 	UPROPERTY(EditDefaultsOnly, Category=Appearance)
-	FSlateColor BackgroundColor;
+	FLinearColor BackgroundColor;
 
 	/** Text color and opacity when read-only (overrides Style) */
 	UPROPERTY(EditDefaultsOnly, Category=Appearance)
-	FSlateColor ReadOnlyForegroundColor;
-
-	/** Sets whether this text box can actually be modified interactively by the user */
-	UPROPERTY(EditDefaultsOnly, Category=Appearance)
-	bool IsReadOnly;
-
-	/** Sets whether this text box is for storing a password */
-	UPROPERTY(EditDefaultsOnly, Category=Appearance)
-	bool IsPassword;
-
-	/** Minimum width that a text block should be */
-	UPROPERTY(EditDefaultsOnly, Category=Appearance)
-	float MinimumDesiredWidth;
-
-	/** Padding between the box/border and the text widget inside (overrides Style) */
-	UPROPERTY(EditDefaultsOnly, Category=Appearance)
-	FMargin Padding;
-
-	/** Workaround as we loose focus when the auto completion closes. */
-	UPROPERTY(EditDefaultsOnly, Category=Behavior, AdvancedDisplay)
-	bool IsCaretMovedWhenGainFocus;
-
-	/** Whether to select all text when the user clicks to give focus on the widget */
-	UPROPERTY(EditDefaultsOnly, Category=Behavior, AdvancedDisplay)
-	bool SelectAllTextWhenFocused;
-
-	/** Whether to allow the user to back out of changes when they press the escape key */
-	UPROPERTY(EditDefaultsOnly, Category=Behavior, AdvancedDisplay)
-	bool RevertTextOnEscape;
-
-	/** Whether to clear keyboard focus when pressing enter to commit changes */
-	UPROPERTY(EditDefaultsOnly, Category=Behavior, AdvancedDisplay)
-	bool ClearKeyboardFocusOnCommit;
-
-	/** Whether to select all text when pressing enter to commit changes */
-	UPROPERTY(EditDefaultsOnly, Category=Behavior, AdvancedDisplay)
-	bool SelectAllTextOnCommit;
+	FLinearColor ReadOnlyForegroundColor;
 
 	/** Called whenever the text is changed interactively by the user */
 	UPROPERTY(BlueprintAssignable, Category="Widget Event")
