UMG - Reparenting is now limited to UUserWidget base classes in the blueprint editor.  Exposing a way to get documentation for palette items.  Adding convience methods for setting the brush/texture/material of a UImage widget.

diff --git a/Engine/Source/Editor/UMGEditor/Private/Palette/SPaletteView.cpp b/Engine/Source/Editor/UMGEditor/Private/Palette/SPaletteView.cpp
--- a/Engine/Source/Editor/UMGEditor/Private/Palette/SPaletteView.cpp
+++ b/Engine/Source/Editor/UMGEditor/Private/Palette/SPaletteView.cpp
@@ -41,6 +41,7 @@ public:
 		ChildSlot
 		[
 			SNew(SHorizontalBox)
+			.ToolTip(Template->GetToolTip())
 
 			+ SHorizontalBox::Slot()
 			.AutoWidth()
@@ -297,8 +298,6 @@ void SPaletteView::BuildClassWidgetList()
 			if ( bIsValidClass && !bIsSameClass )
 			{
 				TSharedPtr<FWidgetTemplateClass> Template = MakeShareable(new FWidgetTemplateClass(WidgetClass));
-				Template->Name = WidgetClass->GetDisplayNameText();
-				Template->ToolTip = WidgetClass->GetDisplayNameText();
 
 				AddWidgetTemplate(Template);
 			}
diff --git a/Engine/Source/Editor/UMGEditor/Private/Templates/WidgetTemplateClass.cpp b/Engine/Source/Editor/UMGEditor/Private/Templates/WidgetTemplateClass.cpp
--- a/Engine/Source/Editor/UMGEditor/Private/Templates/WidgetTemplateClass.cpp
+++ b/Engine/Source/Editor/UMGEditor/Private/Templates/WidgetTemplateClass.cpp
@@ -3,12 +3,14 @@
 #include "UMGEditorPrivatePCH.h"
 
 #include "WidgetTemplateClass.h"
+#include "IDocumentation.h"
 
 #define LOCTEXT_NAMESPACE "UMGEditor"
 
 FWidgetTemplateClass::FWidgetTemplateClass(TSubclassOf<UWidget> InWidgetClass)
 	: WidgetClass(InWidgetClass)
 {
+	Name = WidgetClass->GetDisplayNameText();
 }
 
 FText FWidgetTemplateClass::GetCategory() const
@@ -46,4 +48,9 @@ const FSlateBrush* FWidgetTemplateClass::GetIcon() const
 	return DefaultWidget->GetEditorIcon();
 }
 
+TSharedRef<IToolTip> FWidgetTemplateClass::GetToolTip() const
+{
+	return IDocumentation::Get()->CreateToolTip(WidgetClass->GetDisplayNameText(), nullptr, FString(TEXT("Shared/Types/")) + WidgetClass->GetName(), TEXT("Class"));
+}
+
 #undef LOCTEXT_NAMESPACE
diff --git a/Engine/Source/Editor/UMGEditor/Private/Templates/WidgetTemplateClass.h b/Engine/Source/Editor/UMGEditor/Private/Templates/WidgetTemplateClass.h
--- a/Engine/Source/Editor/UMGEditor/Private/Templates/WidgetTemplateClass.h
+++ b/Engine/Source/Editor/UMGEditor/Private/Templates/WidgetTemplateClass.h
@@ -22,6 +22,9 @@ public:
 	/** The icon coming from the default object of the class */
 	virtual const FSlateBrush* GetIcon() const override;
 
+	/** Gets the tooltip widget for this palette item. */
+	virtual TSharedRef<IToolTip> GetToolTip() const override;
+
 protected:
 	/** The widget class that will be created by this template */
 	TSubclassOf<UWidget> WidgetClass;
diff --git a/Engine/Source/Editor/UMGEditor/Private/WidgetBlueprint.cpp b/Engine/Source/Editor/UMGEditor/Private/WidgetBlueprint.cpp
--- a/Engine/Source/Editor/UMGEditor/Private/WidgetBlueprint.cpp
+++ b/Engine/Source/Editor/UMGEditor/Private/WidgetBlueprint.cpp
@@ -121,6 +121,11 @@ bool UWidgetBlueprint::ValidateGeneratedClass(const UClass* InClass)
 	return Result;
 }
 
+void UWidgetBlueprint::GetReparentingRules(TSet< const UClass* >& AllowedChildrenOfClasses, TSet< const UClass* >& DisallowedChildrenOfClasses) const
+{
+	AllowedChildrenOfClasses.Add( UUserWidget::StaticClass() );
+}
+
 void UWidgetBlueprint::PostLoadSubobjects(FObjectInstancingGraph* OuterInstanceGraph)
 {
 	Super::PostLoadSubobjects(OuterInstanceGraph);
diff --git a/Engine/Source/Editor/UMGEditor/Public/WidgetBlueprint.h b/Engine/Source/Editor/UMGEditor/Public/WidgetBlueprint.h
--- a/Engine/Source/Editor/UMGEditor/Public/WidgetBlueprint.h
+++ b/Engine/Source/Editor/UMGEditor/Public/WidgetBlueprint.h
@@ -76,6 +76,8 @@ public:
 	{
 		return false;
 	}
+
+	virtual void GetReparentingRules(TSet< const UClass* >& AllowedChildrenOfClasses, TSet< const UClass* >& DisallowedChildrenOfClasses) const;
 	// End of UBlueprint interface
 
 	FWidgetAnimation* FindAnimationDataForMovieScene( UMovieScene& MovieScene );
diff --git a/Engine/Source/Editor/UMGEditor/Public/WidgetTemplate.h b/Engine/Source/Editor/UMGEditor/Public/WidgetTemplate.h
--- a/Engine/Source/Editor/UMGEditor/Public/WidgetTemplate.h
+++ b/Engine/Source/Editor/UMGEditor/Public/WidgetTemplate.h
@@ -25,10 +25,10 @@ public:
 		return &NullBrush;
 	}
 
+	/** Gets tooltip widget for this palette item. */
+	virtual TSharedRef<IToolTip> GetToolTip() const = 0;
+
 public:
 	/** The name of the widget template. */
 	FText Name;
-
-	/** A tooltip to display for the template. */
-	FText ToolTip;
 };
diff --git a/Engine/Source/Editor/UnrealEd/Private/Kismet2/BlueprintEditorUtils.cpp b/Engine/Source/Editor/UnrealEd/Private/Kismet2/BlueprintEditorUtils.cpp
--- a/Engine/Source/Editor/UnrealEd/Private/Kismet2/BlueprintEditorUtils.cpp
+++ b/Engine/Source/Editor/UnrealEd/Private/Kismet2/BlueprintEditorUtils.cpp
@@ -5380,6 +5380,11 @@ TSharedRef<SWidget> FBlueprintEditorUtils::ConstructBlueprintParentClassPicker(
 		Filter->DisallowedChildrenOfClasses.Add(*ClassIt);
 	}
 
+	for ( UBlueprint* Blueprint : Blueprints )
+	{
+		Blueprint->GetReparentingRules(Filter->AllowedChildrenOfClasses, Filter->DisallowedChildrenOfClasses);
+	}
+
 	if(bIsActor)
 	{
 		if(bIsLevelScriptActor)
diff --git a/Engine/Source/Runtime/Engine/Classes/Engine/Blueprint.h b/Engine/Source/Runtime/Engine/Classes/Engine/Blueprint.h
--- a/Engine/Source/Runtime/Engine/Classes/Engine/Blueprint.h
+++ b/Engine/Source/Runtime/Engine/Classes/Engine/Blueprint.h
@@ -475,6 +475,8 @@ public:
 
 	virtual void SetWorldBeingDebugged(UWorld* NewWorld);
 
+	virtual void GetReparentingRules(TSet< const UClass* >& AllowedChildrenOfClasses, TSet< const UClass* >& DisallowedChildrenOfClasses) const;
+
 protected:
 	/** Gets asset registry tags */
 	virtual void GetAssetRegistryTags(TArray<FAssetRegistryTag>& OutTags) const override;
diff --git a/Engine/Source/Runtime/Engine/Private/Blueprint.cpp b/Engine/Source/Runtime/Engine/Private/Blueprint.cpp
--- a/Engine/Source/Runtime/Engine/Private/Blueprint.cpp
+++ b/Engine/Source/Runtime/Engine/Private/Blueprint.cpp
@@ -495,6 +495,11 @@ void UBlueprint::SetWorldBeingDebugged(UWorld *NewWorld)
 	CurrentWorldBeingDebugged = NewWorld;
 }
 
+void UBlueprint::GetReparentingRules(TSet< const UClass* >& AllowedChildrenOfClasses, TSet< const UClass* >& DisallowedChildrenOfClasses) const
+{
+
+}
+
 UObject* UBlueprint::GetObjectBeingDebugged()
 {
 	UObject* DebugObj = CurrentObjectBeingDebugged.Get();
diff --git a/Engine/Source/Runtime/UMG/Private/Components/Image.cpp b/Engine/Source/Runtime/UMG/Private/Components/Image.cpp
--- a/Engine/Source/Runtime/UMG/Private/Components/Image.cpp
+++ b/Engine/Source/Runtime/UMG/Private/Components/Image.cpp
@@ -98,20 +98,70 @@ void UImage::SetImage(USlateBrushAsset* InImage)
 	}
 }
 
+void UImage::SetImageFromBrush(FSlateBrush InImage)
+{
+	DynamicBrush = InImage;
+
+	if ( MyImage.IsValid() )
+	{
+		MyImage->SetImage(&DynamicBrush.GetValue());
+	}
+}
+
+void UImage::SetImageFromTexture(UTexture2D* Texture)
+{
+	FSlateBrush TextureBrush;
+	TextureBrush.SetResourceObject(Texture);
+
+	DynamicBrush = TextureBrush;
+
+	if ( MyImage.IsValid() )
+	{
+		MyImage->SetImage(&DynamicBrush.GetValue());
+	}
+}
+
+void UImage::SetImageFromMaterial(UMaterialInterface* Material)
+{
+	FSlateBrush MaterialBrush;
+	MaterialBrush.SetResourceObject(Material);
+
+	DynamicBrush = MaterialBrush;
+
+	if ( MyImage.IsValid() )
+	{
+		MyImage->SetImage(&DynamicBrush.GetValue());
+	}
+}
+
 UMaterialInstanceDynamic* UImage::GetDynamicMaterial()
 {
+	UMaterialInterface* Material = NULL;
+
 	if ( !DynamicBrush.IsSet() )
 	{
 		const FSlateBrush* Brush = GetImageBrush();
-		FSlateBrush ClonedBrush = *Brush;
-		UObject* Resource = ClonedBrush.GetResourceObject();
-		UMaterialInterface* Material = Cast<UMaterialInterface>(Resource);
+		UObject* Resource = Brush->GetResourceObject();
+		Material = Cast<UMaterialInterface>(Resource);
+	}
+	else
+	{
+		UObject* Resource = DynamicBrush.GetValue().GetResourceObject();
+		Material = Cast<UMaterialInterface>(Resource);
+	}
+
+	if ( Material )
+	{
+		UMaterialInstanceDynamic* DynamicMaterial = Cast<UMaterialInstanceDynamic>(Material);
 
-		if ( Material )
+		if ( !DynamicMaterial )
 		{
-			UMaterialInstanceDynamic* DynamicMaterial = UMaterialInstanceDynamic::Create(Material, this);
+			DynamicMaterial = UMaterialInstanceDynamic::Create(Material, this);
+
+			const FSlateBrush* Brush = GetImageBrush();
+			FSlateBrush ClonedBrush = *Brush;
 			ClonedBrush.SetResourceObject(DynamicMaterial);
-			
+
 			DynamicBrush = ClonedBrush;
 
 			if ( MyImage.IsValid() )
@@ -119,17 +169,12 @@ UMaterialInstanceDynamic* UImage::GetDynamicMaterial()
 				MyImage->SetImage(&DynamicBrush.GetValue());
 			}
 		}
-		else
-		{
-			//TODO UMG Log error for debugging.
-		}
-	}
-	else
-	{
-		UObject* Resource = DynamicBrush.GetValue().GetResourceObject();
-		return Cast<UMaterialInstanceDynamic>(Resource);
+		
+		return DynamicMaterial;
 	}
 
+	//TODO UMG can we do something for textures?  General purpose dynamic material for them?
+
 	return NULL;
 }
 
diff --git a/Engine/Source/Runtime/UMG/Private/UserWidget.cpp b/Engine/Source/Runtime/UMG/Private/UserWidget.cpp
--- a/Engine/Source/Runtime/UMG/Private/UserWidget.cpp
+++ b/Engine/Source/Runtime/UMG/Private/UserWidget.cpp
@@ -319,6 +319,16 @@ void UUserWidget::NativeTick(const FGeometry& MyGeometry, float InDeltaTime )
 		Player->Tick( InDeltaTime );
 	}
 
+	if ( !bDesignTime )
+	{
+		UWorld* World = GetWorld();
+		if ( World )
+		{
+			// Update any latent actions we have for this actor
+			World->GetLatentActionManager().ProcessLatentActions(this, InDeltaTime);
+		}
+	}
+
 	Tick( MyGeometry, InDeltaTime );
 }
 
diff --git a/Engine/Source/Runtime/UMG/Public/Components/Image.h b/Engine/Source/Runtime/UMG/Public/Components/Image.h
--- a/Engine/Source/Runtime/UMG/Public/Components/Image.h
+++ b/Engine/Source/Runtime/UMG/Public/Components/Image.h
@@ -45,6 +45,16 @@ public:
 	UFUNCTION(BlueprintCallable, Category="Appearance")
 	void SetImage(USlateBrushAsset* InImage);
 
+	/**  */
+	UFUNCTION(BlueprintCallable, Category="Appearance")
+	void SetImageFromBrush(FSlateBrush Brush);
+
+	UFUNCTION(BlueprintCallable, Category="Appearance")
+	void SetImageFromTexture(UTexture2D* Texture);
+
+	UFUNCTION(BlueprintCallable, Category="Appearance")
+	void SetImageFromMaterial(UMaterialInterface* Material);
+
 	/**  */
 	UFUNCTION(BlueprintCallable, Category="Appearance")
 	UMaterialInstanceDynamic* GetDynamicMaterial();
