UMG - The viewport now exposes the pocket dimension world where actors can be spawned.  Super experimental.

diff --git a/Engine/Source/Runtime/UMG/Private/Components/Viewport.cpp b/Engine/Source/Runtime/UMG/Private/Components/Viewport.cpp
--- a/Engine/Source/Runtime/UMG/Private/Components/Viewport.cpp
+++ b/Engine/Source/Runtime/UMG/Private/Components/Viewport.cpp
@@ -103,6 +103,25 @@ void FUMGViewportClient::AddReferencedObjects(FReferenceCollector & Collector)
 
 }
 
+void FUMGViewportClient::Tick(float InDeltaTime)
+{
+	if ( !GIntraFrameDebuggingGameThread )
+	{
+		// Begin Play
+		if ( !PreviewScene->GetWorld()->bBegunPlay )
+		{
+			for ( FActorIterator It(PreviewScene->GetWorld()); It; ++It )
+			{
+				It->BeginPlay();
+			}
+			PreviewScene->GetWorld()->bBegunPlay = true;
+		}
+
+		// Tick
+		PreviewScene->GetWorld()->Tick(LEVELTICK_All, InDeltaTime);
+	}
+}
+
 void FUMGViewportClient::Draw(FViewport* InViewport, FCanvas* Canvas)
 {
 	FViewport* ViewportBackup = Viewport;
@@ -163,7 +182,8 @@ void FUMGViewportClient::Draw(FViewport* InViewport, FCanvas* Canvas)
 		Canvas->Clear(FLinearColor::Black);
 	}
 
-	Canvas->Clear(FLinearColor::Black);
+	//Canvas->AlphaModulate = 0.5f;
+	Canvas->Clear(BackgroundColor);
 
 	GetRendererModule().BeginRenderingViewFamily(Canvas, &ViewFamily);
 
@@ -188,7 +208,7 @@ void FUMGViewportClient::Draw(FViewport* InViewport, FCanvas* Canvas)
 
 	Viewport = ViewportBackup;
 
-	Canvas->Clear(FLinearColor::Red);
+	//Canvas->Clear(FLinearColor::Red);
 }
 
 FSceneInterface* FUMGViewportClient::GetScene() const
@@ -237,6 +257,11 @@ ECameraProjectionMode::Type FUMGViewportClient::GetProjectionType() const
 	return ViewInfo.ProjectionMode;
 }
 
+void FUMGViewportClient::SetBackgroundColor(FLinearColor InBackgroundColor)
+{
+	BackgroundColor = InBackgroundColor;
+}
+
 FLinearColor FUMGViewportClient::GetBackgroundColor() const
 {
 	return BackgroundColor;
@@ -430,6 +455,8 @@ class SAutoRefreshViewport : public SViewport
 	void Construct(const FArguments& InArgs)
 	{
 		SViewport::FArguments ParentArgs;
+		//ParentArgs.IgnoreTextureAlpha(false);
+		ParentArgs.EnableBlending(true);
 		//ParentArgs.RenderDirectlyToWindow(true);
 		SViewport::Construct(ParentArgs);
 
@@ -446,8 +473,8 @@ class SAutoRefreshViewport : public SViewport
 		Viewport->InvalidateDisplay();
 
 		Viewport->Tick(InDeltaTime);
-		
-		//ViewportClient->Tick(InDeltaTime);
+
+		ViewportClient->Tick(InDeltaTime);
 	}
 
 public:
@@ -466,6 +493,8 @@ UViewport::UViewport(const FPostConstructInitializeProperties& PCIP)
 	: Super(PCIP)
 {
 	bIsVariable = true;
+
+	BackgroundColor = FLinearColor::Transparent;
 }
 
 void UViewport::ReleaseNativeWidget()
@@ -490,6 +519,8 @@ TSharedRef<SWidget> UViewport::RebuildWidget()
 void UViewport::SyncronizeProperties()
 {
 	Super::SyncronizeProperties();
+
+	ViewportWidget->ViewportClient->SetBackgroundColor(BackgroundColor);
 }
 
 void UViewport::OnSlotAdded(UPanelSlot* Slot)
@@ -510,6 +541,65 @@ void UViewport::OnSlotRemoved(UPanelSlot* Slot)
 	}
 }
 
+UWorld* UViewport::GetViewportWorld() const
+{
+	if ( ViewportWidget.IsValid() )
+	{
+		return ViewportWidget->PreviewScene.GetWorld();
+	}
+
+	return NULL;
+}
+
+FVector UViewport::GetViewLocation() const
+{
+	if ( ViewportWidget.IsValid() )
+	{
+		return ViewportWidget->ViewportClient->GetViewLocation();
+	}
+
+	return FVector();
+}
+
+void UViewport::SetViewLocation(FVector Vector)
+{
+	if ( ViewportWidget.IsValid() )
+	{
+		ViewportWidget->ViewportClient->SetViewLocation(Vector);
+	}
+}
+
+FRotator UViewport::GetViewRotation() const
+{
+	if ( ViewportWidget.IsValid() )
+	{
+		return ViewportWidget->ViewportClient->GetViewRotation();
+	}
+
+	return FRotator();
+}
+
+void UViewport::SetViewRotation(FRotator Rotator)
+{
+	if ( ViewportWidget.IsValid() )
+	{
+		ViewportWidget->ViewportClient->SetViewRotation(Rotator);
+	}
+}
+
+AActor* UViewport::Spawn(TSubclassOf<AActor> ActorClass)
+{
+	if ( ViewportWidget.IsValid() )
+	{
+		UWorld* World = GetViewportWorld();
+		return World->SpawnActor<AActor>(ActorClass, FVector(0, 0, 0), FRotator());
+	}
+
+	// TODO UMG Report spawning actor error before the world is ready.
+
+	return NULL;
+}
+
 #if WITH_EDITOR
 
 const FSlateBrush* UViewport::GetEditorIcon()
diff --git a/Engine/Source/Runtime/UMG/Public/Components/Viewport.h b/Engine/Source/Runtime/UMG/Public/Components/Viewport.h
--- a/Engine/Source/Runtime/UMG/Public/Components/Viewport.h
+++ b/Engine/Source/Runtime/UMG/Public/Components/Viewport.h
@@ -96,6 +96,9 @@ public:
 	virtual void AddReferencedObjects(FReferenceCollector & Collector) override;
 
 	virtual void Draw(FViewport* InViewport, FCanvas* Canvas) override;
+
+	virtual void Tick(float InDeltaTime);
+
 	using FViewElementDrawer::Draw;
 
 	virtual UWorld* GetWorld() const override;
@@ -113,6 +116,7 @@ public:
 
 	bool IsAspectRatioConstrained() const;
 
+	void SetBackgroundColor(FLinearColor InBackgroundColor);
 	FLinearColor GetBackgroundColor() const;
 
 	float GetNearClipPlane() const;
@@ -234,6 +238,27 @@ class UMG_API UViewport : public UContentWidget
 {
 	GENERATED_UCLASS_BODY()
 
+	UPROPERTY(EditDefaultsOnly, Category=Appearance)
+	FLinearColor BackgroundColor;
+
+	UFUNCTION(BlueprintCallable, Category="Viewport")
+	UWorld* GetViewportWorld() const;
+
+	UFUNCTION(BlueprintCallable, Category="Camera")
+	FVector GetViewLocation() const;
+
+	UFUNCTION(BlueprintCallable, Category="Camera")
+	void SetViewLocation(FVector Location);
+
+	UFUNCTION(BlueprintCallable, Category="Camera")
+	FRotator GetViewRotation() const;
+
+	UFUNCTION(BlueprintCallable, Category="Camera")
+	void SetViewRotation(FRotator Rotation);
+
+	UFUNCTION(BlueprintCallable, Category="Viewport")
+	AActor* Spawn(TSubclassOf<AActor> ActorClass);
+
 	// UWidget interface
 	virtual void SyncronizeProperties() override;
 	// End of UWidget interface
