UMG - Tweaking some defaults of the viewport.

diff --git a/Engine/Source/Runtime/UMG/Private/Components/Viewport.cpp b/Engine/Source/Runtime/UMG/Private/Components/Viewport.cpp
--- a/Engine/Source/Runtime/UMG/Private/Components/Viewport.cpp
+++ b/Engine/Source/Runtime/UMG/Private/Components/Viewport.cpp
@@ -182,7 +182,6 @@ void FUMGViewportClient::Draw(FViewport* InViewport, FCanvas* Canvas)
 		Canvas->Clear(FLinearColor::Black);
 	}
 
-	//Canvas->AlphaModulate = 0.5f;
 	Canvas->Clear(BackgroundColor);
 
 	GetRendererModule().BeginRenderingViewFamily(Canvas, &ViewFamily);
@@ -207,8 +206,6 @@ void FUMGViewportClient::Draw(FViewport* InViewport, FCanvas* Canvas)
 	//FlushRenderingCommands();
 
 	Viewport = ViewportBackup;
-
-	//Canvas->Clear(FLinearColor::Red);
 }
 
 FSceneInterface* FUMGViewportClient::GetScene() const
@@ -455,8 +452,8 @@ class SAutoRefreshViewport : public SViewport
 	void Construct(const FArguments& InArgs)
 	{
 		SViewport::FArguments ParentArgs;
-		//ParentArgs.IgnoreTextureAlpha(false);
-		ParentArgs.EnableBlending(true);
+		ParentArgs.IgnoreTextureAlpha(false);
+		ParentArgs.EnableBlending(false);
 		//ParentArgs.RenderDirectlyToWindow(true);
 		SViewport::Construct(ParentArgs);
 
@@ -491,10 +488,15 @@ public:
 
 UViewport::UViewport(const FPostConstructInitializeProperties& PCIP)
 	: Super(PCIP)
+	, ShowFlags(ESFIM_Game)
 {
 	bIsVariable = true;
 
-	BackgroundColor = FLinearColor::Transparent;
+	BackgroundColor = FLinearColor::Black;
+	ShowFlags.DisableAdvancedFeatures();
+	//ParentArgs.IgnoreTextureAlpha(false);
+	//ParentArgs.EnableBlending(true);
+	////ParentArgs.RenderDirectlyToWindow(true);
 }
 
 void UViewport::ReleaseNativeWidget()
@@ -521,6 +523,7 @@ void UViewport::SyncronizeProperties()
 	Super::SyncronizeProperties();
 
 	ViewportWidget->ViewportClient->SetBackgroundColor(BackgroundColor);
+	ViewportWidget->ViewportClient->SetEngineShowFlags(ShowFlags);
 }
 
 void UViewport::OnSlotAdded(UPanelSlot* Slot)
diff --git a/Engine/Source/Runtime/UMG/Public/Components/Viewport.h b/Engine/Source/Runtime/UMG/Public/Components/Viewport.h
--- a/Engine/Source/Runtime/UMG/Public/Components/Viewport.h
+++ b/Engine/Source/Runtime/UMG/Public/Components/Viewport.h
@@ -196,6 +196,11 @@ public:
 	/** @return The number of units per pixel displayed in this viewport */
 	float GetOrthoUnitsPerPixel(const FViewport* Viewport) const;
 
+	void SetEngineShowFlags(FEngineShowFlags InEngineShowFlags)
+	{
+		EngineShowFlags = InEngineShowFlags;
+	}
+
 protected:
 
 	/** The scene used for the viewport. Owned externally */
@@ -283,4 +288,7 @@ protected:
 	// UWidget interface
 	virtual TSharedRef<SWidget> RebuildWidget() override;
 	// End of UWidget interface
+
+	/** Show flags for the engine for this viewport */
+	FEngineShowFlags ShowFlags;
 };
