UMG - Disable drag detection for now in SObjectWidget, and fix the issue with uninitialized capture mouse problems.

diff --git a/Engine/Source/Runtime/UMG/Private/Slate/SObjectWidget.cpp b/Engine/Source/Runtime/UMG/Private/Slate/SObjectWidget.cpp
--- a/Engine/Source/Runtime/UMG/Private/Slate/SObjectWidget.cpp
+++ b/Engine/Source/Runtime/UMG/Private/Slate/SObjectWidget.cpp
@@ -135,10 +135,10 @@ FReply SObjectWidget::OnMouseButtonDown(const FGeometry& MyGeometry, const FPoin
 		FReply Reply = WidgetObject->OnMouseButtonDown(MyGeometry, MouseEvent).ToReply(SlateWidget.ToSharedRef());
 
 		//TODO UMG Figure out how to let the user more easily manage when drags start.
-		if ( MouseEvent.GetEffectingButton() == EKeys::LeftMouseButton )
-		{
-			return Reply.DetectDrag(SlateWidget.ToSharedRef(), EKeys::LeftMouseButton);
-		}
+		//if ( MouseEvent.GetEffectingButton() == EKeys::LeftMouseButton )
+		//{
+		//	return Reply.DetectDrag(SlateWidget.ToSharedRef(), EKeys::LeftMouseButton);
+		//}
 
 		return Reply;
 	}
diff --git a/Engine/Source/Runtime/UMG/Public/Components/SlateWrapperTypes.h b/Engine/Source/Runtime/UMG/Public/Components/SlateWrapperTypes.h
--- a/Engine/Source/Runtime/UMG/Public/Components/SlateWrapperTypes.h
+++ b/Engine/Source/Runtime/UMG/Public/Components/SlateWrapperTypes.h
@@ -55,6 +55,13 @@ struct FSReply
 	UPROPERTY(EditAnywhere, Category=User)
 	uint32 bCaptureMouse : 1;
 
+	FSReply()
+		: bIsHandled(false)
+		, bCaptureMouse(false)
+	{
+
+	}
+
 	FReply ToReply( TSharedRef<SWidget> Widget ) const
 	{
 		FReply Reply = FReply::Unhandled();
