#UMG

- Making Visibility a Blueprint type

diff --git a/Engine/Source/Runtime/UMG/Public/Components/SlateWrapperTypes.h b/Engine/Source/Runtime/UMG/Public/Components/SlateWrapperTypes.h
--- a/Engine/Source/Runtime/UMG/Public/Components/SlateWrapperTypes.h
+++ b/Engine/Source/Runtime/UMG/Public/Components/SlateWrapperTypes.h
@@ -12,7 +12,7 @@
 	Type::CreateUObject( this, &ThisClass::Function )
 
 /** Is an entity visible? */
-UENUM()
+UENUM(BlueprintType)
 namespace ESlateVisibility
 {
 	enum Type
