UMG - Show/Hide have been replaced with AddToViewport RemoveFromViewport.  TextBlock now has SetText().  Fullscreen widgets no longer are all modal, so you can make a HUD with UMG.  It's now possible to manipulate the position of a widget added to the viewport after adding it dynamically.

diff --git a/Engine/Source/Editor/UMGEditor/Private/Designer/SDesignerView.cpp b/Engine/Source/Editor/UMGEditor/Private/Designer/SDesignerView.cpp
--- a/Engine/Source/Editor/UMGEditor/Private/Designer/SDesignerView.cpp
+++ b/Engine/Source/Editor/UMGEditor/Private/Designer/SDesignerView.cpp
@@ -899,7 +899,10 @@ void SDesignerView::UpdatePreviewWidget()
 		PreviewWidget = LatestPreviewWidget;
 		if ( PreviewWidget )
 		{
-			TSharedRef<SWidget> CurrentWidget = PreviewWidget->MakeFullScreenWidget();
+			const bool bAbsoluteLayout = false;
+			const bool bModal = false;
+			const bool bShowCursor = false;
+			TSharedRef<SWidget> CurrentWidget = PreviewWidget->MakeViewportWidget(bAbsoluteLayout, bModal, bShowCursor);
 			CurrentWidget->SlatePrepass();
 
 			if ( CurrentWidget != PreviewSlateWidget.Pin() )
diff --git a/Engine/Source/Runtime/UMG/Private/Components/TextBlock.cpp b/Engine/Source/Runtime/UMG/Private/Components/TextBlock.cpp
--- a/Engine/Source/Runtime/UMG/Private/Components/TextBlock.cpp
+++ b/Engine/Source/Runtime/UMG/Private/Components/TextBlock.cpp
@@ -107,6 +107,15 @@ void UTextBlock::SyncronizeProperties()
 
 #if WITH_EDITOR
 
+void UTextBlock::SetText(FText InText)
+{
+	Text = InText;
+	if ( MyTextBlock.IsValid() )
+	{
+		MyTextBlock->SetText(Text);
+	}
+}
+
 FString UTextBlock::GetLabelMetadata() const
 {
 	const int32 MaxSampleLength = 15;
diff --git a/Engine/Source/Runtime/UMG/Private/UserWidget.cpp b/Engine/Source/Runtime/UMG/Private/UserWidget.cpp
--- a/Engine/Source/Runtime/UMG/Private/UserWidget.cpp
+++ b/Engine/Source/Runtime/UMG/Private/UserWidget.cpp
@@ -49,14 +49,9 @@ class SViewportWidgetHost : public SCompoundWidget
 			];
 	}
 
-	//virtual bool OnHitTest(const FGeometry& MyGeometry, FVector2D InAbsoluteCursorPosition) override
-	//{
-	//	return false;
-	//}
-
 	virtual bool SupportsKeyboardFocus() const override
 	{
-		return true;
+		return bModal ? true : false;
 	}
 
 	FReply OnKeyboardFocusReceived(const FGeometry& MyGeometry, const FKeyboardFocusEvent& InKeyboardFocusEvent)
@@ -85,46 +80,6 @@ class SViewportWidgetHost : public SCompoundWidget
 		}
 	}
 
-	//FReply OnKeyDown(const FGeometry& MyGeometry, const FKeyboardEvent& InKeyboardEvent)
-	//{
-	//	return FReply::Handled();
-	//}
-
-	//FReply OnKeyUp(const FGeometry& MyGeometry, const FKeyboardEvent& InKeyboardEvent)
-	//{
-	//	return FReply::Handled();
-	//}
-
-	//FReply OnMouseButtonDown(const FGeometry& MyGeometry, const FPointerEvent& MouseEvent)
-	//{
-	//	if ( bModal )
-	//	{
-	//		return FReply::Handled();
-	//	}
-
-	//	return FReply::Unhandled();
-	//}
-
-	//FReply OnMouseButtonUp(const FGeometry& MyGeometry, const FPointerEvent& MouseEvent)
-	//{
-	//	if ( bModal )
-	//	{
-	//		return FReply::Handled();
-	//	}
-
-	//	return FReply::Unhandled();
-	//}
-
-	//FReply OnMouseMove(const FGeometry& MyGeometry, const FPointerEvent& MouseEvent)
-	//{
-	//	if ( bModal )
-	//	{
-	//		return FReply::Handled();
-	//	}
-
-	//	return FReply::Unhandled();
-	//}
-
 protected:
 	bool bModal;
 };
@@ -136,12 +91,6 @@ protected:
 UUserWidget::UUserWidget(const FPostConstructInitializeProperties& PCIP)
 	: Super(PCIP)
 {
-	/*PrimaryActorTick.bCanEverTick = true;
-	PrimaryActorTick.bTickEvenWhenPaused = true;*/
-	bShowCursorWhenVisible = false;
-	bAbsoluteLayout = false;
-	bModal = false;
-	AbsoluteSize = FVector2D(100, 100);
 	HorizontalAlignment = HAlign_Fill;
 	VerticalAlignment = VAlign_Fill;
 
@@ -182,11 +131,6 @@ UWorld* UUserWidget::GetWorld() const
 		return PlayerContext.GetWorld();
 	}
 
-	//if ( ULevel* Level = Cast<ULevel>(Outer) )
-	//{
-	//	return Level->OwningWorld;
-	//}
-
 	return NULL;
 }
 
@@ -332,17 +276,18 @@ void UUserWidget::NativeTick(const FGeometry& MyGeometry, float InDeltaTime )
 	Tick( MyGeometry, InDeltaTime );
 }
 
-TSharedRef<SWidget> UUserWidget::MakeFullScreenWidget()
+TSharedRef<SWidget> UUserWidget::MakeViewportWidget(bool bAbsoluteLayout, bool bModal, bool bShowCursor)
 {
 	if ( bAbsoluteLayout )
 	{
-		return SNew(SCanvas)
+		//BIND_UOBJECT_ATTRIBUTE
 
-			+ SCanvas::Slot()
-			.Position(AbsolutePosition)
-			.Size(AbsoluteSize)
-			.VAlign(VerticalAlignment)
-			.HAlign(HorizontalAlignment)
+		return SNew(SConstraintCanvas)
+
+			+ SConstraintCanvas::Slot()
+			.Offset(BIND_UOBJECT_ATTRIBUTE(FMargin, GetFullScreenOffset))
+			.Alignment(BIND_UOBJECT_ATTRIBUTE(FVector2D, GetFullScreenAlignment))
+			.ZOrder(BIND_UOBJECT_ATTRIBUTE(int32, GetFullScreenZOrder))
 			[
 				TakeWidget()
 			];
@@ -375,11 +320,11 @@ UWidget* UUserWidget::GetRootWidgetComponent()
 	return NULL;
 }
 
-void UUserWidget::Show()
+void UUserWidget::AddToViewport(bool bAbsoluteLayout, bool bModal, bool bShowCursor)
 {
 	if ( !FullScreenWidget.IsValid() )
 	{
-		TSharedRef<SWidget> RootWidget = MakeFullScreenWidget();
+		TSharedRef<SWidget> RootWidget = MakeViewportWidget(bAbsoluteLayout, bModal, bShowCursor);
 
 		TSharedRef<SViewportWidgetHost> WidgetHost = SNew(SViewportWidgetHost, (bool)bModal)
 			[
@@ -388,9 +333,6 @@ void UUserWidget::Show()
 
 		FullScreenWidget = WidgetHost;
 
-		//WidgetHost->SetVisibility(EVisibility::Visible);
-		//OnVisibilityChanged.Broadcast(ESlateVisibility::Visible);
-
 		// If this is a game world add the widget to the current worlds viewport.
 		UWorld* World = GetWorld();
 		if ( World && World->IsGameWorld() )
@@ -408,15 +350,12 @@ void UUserWidget::Show()
 	}
 }
 
-void UUserWidget::Hide()
+void UUserWidget::RemoveFromViewport()
 {
 	if ( FullScreenWidget.IsValid() )
 	{
 		TSharedPtr<SWidget> RootWidget = FullScreenWidget.Pin();
 
-		//RootWidget->SetVisibility(EVisibility::Hidden);
-		//OnVisibilityChanged.Broadcast(ESlateVisibility::Hidden);
-
 		// If this is a game world add the widget to the current worlds viewport.
 		UWorld* World = GetWorld();
 		if ( World && World->IsGameWorld() )
@@ -474,6 +413,31 @@ APlayerController* UUserWidget::GetPlayerController() const
 	return PlayerContext.IsValid() ? PlayerContext.GetPlayerController() : NULL;
 }
 
+FMargin UUserWidget::GetFullScreenOffset() const
+{
+	FVector2D FinalSize = FullScreenSize;
+	if ( FinalSize.IsZero() )
+	{
+		TSharedPtr<SWidget> CachedWidget = GetCachedWidget();
+		if ( CachedWidget.IsValid() )
+		{
+			FinalSize = CachedWidget->GetDesiredSize();
+		}
+	}
+
+	return FMargin(FullScreenPosition.X, FullScreenPosition.Y, FinalSize.X, FinalSize.Y);
+}
+
+FVector2D UUserWidget::GetFullScreenAlignment() const
+{
+	return FullScreenAlignment;
+}
+
+int32 UUserWidget::GetFullScreenZOrder() const
+{
+	return FullScreenZOrder;
+}
+
 #if WITH_EDITOR
 
 const FSlateBrush* UUserWidget::GetEditorIcon()
diff --git a/Engine/Source/Runtime/UMG/Public/Blueprint/UserWidget.h b/Engine/Source/Runtime/UMG/Public/Blueprint/UserWidget.h
--- a/Engine/Source/Runtime/UMG/Public/Blueprint/UserWidget.h
+++ b/Engine/Source/Runtime/UMG/Public/Blueprint/UserWidget.h
@@ -110,11 +110,11 @@ public:
 
 	/*  */
 	UFUNCTION(BlueprintCallable, Category="Appearance")
-	void Show();
+	void AddToViewport(bool bAbsoluteLayout = false, bool bModal = false, bool bShowCursor = false);
 
 	/*  */
 	UFUNCTION(BlueprintCallable, Category="Appearance")
-	void Hide();
+	void RemoveFromViewport();
 
 	UFUNCTION(BlueprintPure, Category="Appearance")
 	bool GetIsVisible();
@@ -237,7 +237,7 @@ public:
 	UWidget* GetWidgetHandle(TSharedRef<SWidget> InWidget);
 
 	/** Creates a fullscreen host widget, that wraps this widget. */
-	TSharedRef<SWidget> MakeFullScreenWidget();
+	TSharedRef<SWidget> MakeViewportWidget(bool bAbsoluteLayout, bool bModal, bool bShowCursor);
 
 	/** @returns The root UObject widget wrapper */
 	UWidget* GetRootWidgetComponent();
@@ -262,31 +262,28 @@ public:
 	UPROPERTY(BlueprintAssignable)
 	FOnVisibilityChangedEvent OnVisibilityChanged;
 
-	/** Controls whether the cursor is automatically visible when this widget is visible. */
-	UPROPERTY(EditDefaultsOnly, Category=Behavior)
-	uint32 bShowCursorWhenVisible : 1;
-
-	UPROPERTY(EditDefaultsOnly, Category=Behavior)
-	uint32 bModal : 1;
-
-	/**  */
 	UPROPERTY(EditDefaultsOnly, Category=Appearance)
-	bool bAbsoluteLayout;
-
-	UPROPERTY(EditAnywhere, Category=Appearance, meta=( EditCondition="!bAbsoluteLayout" ))
 	FMargin Padding;
 
 	/** How much space this slot should occupy in the direction of the panel. */
-	UPROPERTY(EditAnywhere, Category=Appearance, meta=( EditCondition="!bAbsoluteLayout" ))
+	UPROPERTY(EditDefaultsOnly, Category=Appearance)
 	FSlateChildSize Size;
 
-	/** Position. */
-	UPROPERTY(EditDefaultsOnly, Category=Appearance, meta=( EditCondition="bAbsoluteLayout" ))
-	FVector2D AbsolutePosition;
+	/** The position on the screen to place the UI. */
+	UPROPERTY(EditAnywhere, BlueprintReadWrite, Category=Appearance)
+	FVector2D FullScreenPosition;
 
-	/** Size. */
-	UPROPERTY(EditDefaultsOnly, Category=Appearance, meta=( EditCondition="bAbsoluteLayout" ))
-	FVector2D AbsoluteSize;
+	/** The size on the screen the UI should be. */
+	UPROPERTY(EditAnywhere, BlueprintReadWrite, Category=Appearance)
+	FVector2D FullScreenSize;
+
+	/** The normalized UI alignment on the screen, 0..1.  Think of this as the pivot point. */
+	UPROPERTY(EditAnywhere, BlueprintReadWrite, Category=Appearance)
+	FVector2D FullScreenAlignment;
+
+	/** The Z-Order when the UI is fullscreen. */
+	UPROPERTY(EditAnywhere, BlueprintReadWrite, Category=Appearance)
+	int32 FullScreenZOrder;
 
 	/**
 	* Horizontal pivot position
@@ -334,6 +331,10 @@ public:
 protected:
 	virtual TSharedRef<SWidget> RebuildWidget() override;
 
+	FMargin GetFullScreenOffset() const;
+	FVector2D GetFullScreenAlignment() const;
+	int32 GetFullScreenZOrder() const;
+
 private:
 	TMap< TWeakPtr<SWidget>, TWeakObjectPtr<UWidget> > WidgetToComponent;
 
diff --git a/Engine/Source/Runtime/UMG/Public/Components/TextBlock.h b/Engine/Source/Runtime/UMG/Public/Components/TextBlock.h
--- a/Engine/Source/Runtime/UMG/Public/Components/TextBlock.h
+++ b/Engine/Source/Runtime/UMG/Public/Components/TextBlock.h
@@ -85,6 +85,10 @@ public:
 	///** Called when this text is double clicked */
 	//SLATE_EVENT(FOnClicked, OnDoubleClicked)
 
+	/**  */
+	UFUNCTION(BlueprintCallable, Category="Widget")
+	void SetText(FText InText);
+
 	// UWidget interface
 	void SyncronizeProperties() override;
 	// End of UWidget interface
diff --git a/Engine/Source/Runtime/UMG/Public/Slate/SConstraintCanvas.h b/Engine/Source/Runtime/UMG/Public/Slate/SConstraintCanvas.h
--- a/Engine/Source/Runtime/UMG/Public/Slate/SConstraintCanvas.h
+++ b/Engine/Source/Runtime/UMG/Public/Slate/SConstraintCanvas.h
@@ -63,7 +63,7 @@ public:
 		FSlot()
 			: TSlotBase<FSlot>()
 			, OffsetAttr( FMargin( 0, 0, 1, 1 ) )
-			, AnchorsAttr( FAnchors( 0.5f, 0.5f ) )
+			, AnchorsAttr( FAnchors( 0.0f, 0.0f ) )
 			, AlignmentAttr( FVector2D( 0.5f, 0.5f ) )
 			, ZOrderAttr( 0 )
 		{ }
