Added missing files

diff --git a/Engine/Source/Runtime/UMG/Private/Animation/MovieSceneMarginSection.cpp b/Engine/Source/Runtime/UMG/Private/Animation/MovieSceneMarginSection.cpp
--- a/Engine/Source/Runtime/UMG/Private/Animation/MovieSceneMarginSection.cpp
+++ b/Engine/Source/Runtime/UMG/Private/Animation/MovieSceneMarginSection.cpp
@@ -30,12 +30,12 @@ void UMovieSceneMarginSection::DilateSection( float DilationFactor, float Origin
 	BottomCurve.ScaleCurve(Origin, DilationFactor);
 }
 
-FMargin UMovieSceneMarginSection::Eval( float Position ) const
+FMargin UMovieSceneMarginSection::Eval( float Position, const FMargin& DefaultValue ) const
 {
-	return FMargin(	LeftCurve.Eval(Position),
-					TopCurve.Eval(Position),
-					RightCurve.Eval(Position),
-					BottomCurve.Eval(Position));
+	return FMargin(	LeftCurve.Eval(Position, DefaultValue.Left),
+					TopCurve.Eval(Position, DefaultValue.Top),
+					RightCurve.Eval(Position, DefaultValue.Right),
+					BottomCurve.Eval(Position, DefaultValue.Bottom));
 }
 
 void UMovieSceneMarginSection::AddKey( float Time, const FMarginKey& MarginKey )
@@ -61,7 +61,7 @@ bool UMovieSceneMarginSection::NewKeyIsNewData(float Time, const FMargin& Value)
 		LeftCurve.GetNumKeys() == 0 ||
 		RightCurve.GetNumKeys() == 0 ||
 		BottomCurve.GetNumKeys() == 0 ||
-		(Eval(Time) != Value);
+		(Eval(Time,Value) != Value);
 }
 
 void UMovieSceneMarginSection::AddKeyToNamedCurve( float Time, const FMarginKey& MarginKey )
diff --git a/Engine/Source/Runtime/UMG/Private/Animation/MovieSceneMarginTrack.cpp b/Engine/Source/Runtime/UMG/Private/Animation/MovieSceneMarginTrack.cpp
--- a/Engine/Source/Runtime/UMG/Private/Animation/MovieSceneMarginTrack.cpp
+++ b/Engine/Source/Runtime/UMG/Private/Animation/MovieSceneMarginTrack.cpp
@@ -40,13 +40,15 @@ bool UMovieSceneMarginTrack::AddKeyToSection( float Time, const FMarginKey& Marg
 }
 
 
-bool UMovieSceneMarginTrack::Eval( float Position, float LastPosition, FMargin& OutMargin ) const
+bool UMovieSceneMarginTrack::Eval( float Position, float LastPosition, FMargin& InOutMargin ) const
 {
 	const UMovieSceneSection* Section = MovieSceneHelpers::FindSectionAtTime( Sections, Position );
 
 	if( Section )
 	{
-		OutMargin = CastChecked<UMovieSceneMarginSection>( Section )->Eval( Position );
+		const UMovieSceneMarginSection* MarginSection = CastChecked<UMovieSceneMarginSection>( Section );
+
+		InOutMargin = MarginSection->Eval( Position, InOutMargin );
 	}
 
 	return Section != NULL;
diff --git a/Engine/Source/Runtime/UMG/Public/Animation/MovieSceneMarginTrack.h b/Engine/Source/Runtime/UMG/Public/Animation/MovieSceneMarginTrack.h
--- a/Engine/Source/Runtime/UMG/Public/Animation/MovieSceneMarginTrack.h
+++ b/Engine/Source/Runtime/UMG/Public/Animation/MovieSceneMarginTrack.h
@@ -33,14 +33,15 @@ public:
 	 * @return True if the key was successfully added.
 	 */
 	virtual bool AddKeyToSection( float Time, const FMarginKey& MarginKey );
-	
+
 	/**
 	 * Evaluates the track at the playback position
 	 *
 	 * @param Position	The current playback position
-	 * @param LastPosition	The last plackback position
-	 * @param The color at the playback position
+	 * @param LastPosition	The last playback position
+	 * @param InOutMargin 	The margin at the playback position
 	 * @return true if anything was evaluated. Note: if false is returned OutMargin remains unchanged
 	 */
-	virtual bool Eval( float Position, float LastPostion, FMargin& OutMargin ) const;
+	bool Eval( float Position, float LastPostion, FMargin& InOutMargin ) const;
+
 };
