UMG - Updating icon references.

diff --git a/Engine/Source/Runtime/UMG/Private/Components/CircularThrobber.cpp b/Engine/Source/Runtime/UMG/Private/Components/CircularThrobber.cpp
--- a/Engine/Source/Runtime/UMG/Private/Components/CircularThrobber.cpp
+++ b/Engine/Source/Runtime/UMG/Private/Components/CircularThrobber.cpp
@@ -94,6 +94,14 @@ void UCircularThrobber::SetPieceImage(USlateBrushAsset* InPieceImage)
 	}
 }
 
+#if WITH_EDITOR
+
+const FSlateBrush* UCircularThrobber::GetEditorIcon()
+{
+	return FUMGStyle::Get().GetBrush("Widget.CircularThrobber");
+}
+
+#endif
 
 /////////////////////////////////////////////////////
 
diff --git a/Engine/Source/Runtime/UMG/Private/Components/ComboBox.cpp b/Engine/Source/Runtime/UMG/Private/Components/ComboBox.cpp
--- a/Engine/Source/Runtime/UMG/Private/Components/ComboBox.cpp
+++ b/Engine/Source/Runtime/UMG/Private/Components/ComboBox.cpp
@@ -43,6 +43,15 @@ TSharedRef<SWidget> UComboBox::HandleGenerateWidget(UObject* Item) const
 	return SNew(STextBlock).Text(Item ? FText::FromString(Item->GetName()) : LOCTEXT("null", "null"));
 }
 
+#if WITH_EDITOR
+
+const FSlateBrush* UComboBox::GetEditorIcon()
+{
+	return FUMGStyle::Get().GetBrush("Widget.ComboBox");
+}
+
+#endif
+
 /////////////////////////////////////////////////////
 
 #undef LOCTEXT_NAMESPACE
diff --git a/Engine/Source/Runtime/UMG/Private/Components/ListView.cpp b/Engine/Source/Runtime/UMG/Private/Components/ListView.cpp
--- a/Engine/Source/Runtime/UMG/Private/Components/ListView.cpp
+++ b/Engine/Source/Runtime/UMG/Private/Components/ListView.cpp
@@ -59,6 +59,15 @@ TSharedRef<ITableRow> UListView::HandleOnGenerateRow(UObject* Item, const TShare
 		];
 }
 
+#if WITH_EDITOR
+
+const FSlateBrush* UListView::GetEditorIcon()
+{
+	return FUMGStyle::Get().GetBrush("Widget.ListView");
+}
+
+#endif
+
 /////////////////////////////////////////////////////
 
 #undef LOCTEXT_NAMESPACE
diff --git a/Engine/Source/Runtime/UMG/Private/Components/Overlay.cpp b/Engine/Source/Runtime/UMG/Private/Components/Overlay.cpp
--- a/Engine/Source/Runtime/UMG/Private/Components/Overlay.cpp
+++ b/Engine/Source/Runtime/UMG/Private/Components/Overlay.cpp
@@ -63,3 +63,12 @@ TSharedRef<SWidget> UOverlay::RebuildWidget()
 
 	return BuildDesignTimeWidget( MyOverlay.ToSharedRef() );
 }
+
+#if WITH_EDITOR
+
+const FSlateBrush* UOverlay::GetEditorIcon()
+{
+	return FUMGStyle::Get().GetBrush("Widget.Overlay");
+}
+
+#endif
diff --git a/Engine/Source/Runtime/UMG/Private/Components/Throbber.cpp b/Engine/Source/Runtime/UMG/Private/Components/Throbber.cpp
--- a/Engine/Source/Runtime/UMG/Private/Components/Throbber.cpp
+++ b/Engine/Source/Runtime/UMG/Private/Components/Throbber.cpp
@@ -105,6 +105,15 @@ void UThrobber::SetPieceImage(USlateBrushAsset* InPieceImage)
 	}
 }
 
+#if WITH_EDITOR
+
+const FSlateBrush* UThrobber::GetEditorIcon()
+{
+	return FUMGStyle::Get().GetBrush("Widget.Throbber");
+}
+
+#endif
+
 /////////////////////////////////////////////////////
 
 #undef LOCTEXT_NAMESPACE
diff --git a/Engine/Source/Runtime/UMG/Private/Components/TileView.cpp b/Engine/Source/Runtime/UMG/Private/Components/TileView.cpp
--- a/Engine/Source/Runtime/UMG/Private/Components/TileView.cpp
+++ b/Engine/Source/Runtime/UMG/Private/Components/TileView.cpp
@@ -76,6 +76,15 @@ void UTileView::RequestListRefresh()
 	MyTileView->RequestListRefresh();
 }
 
+#if WITH_EDITOR
+
+const FSlateBrush* UTileView::GetEditorIcon()
+{
+	return FUMGStyle::Get().GetBrush("Widget.TileView");
+}
+
+#endif
+
 /////////////////////////////////////////////////////
 
 #undef LOCTEXT_NAMESPACE
diff --git a/Engine/Source/Runtime/UMG/Private/UMGStyle.cpp b/Engine/Source/Runtime/UMG/Private/UMGStyle.cpp
--- a/Engine/Source/Runtime/UMG/Private/UMGStyle.cpp
+++ b/Engine/Source/Runtime/UMG/Private/UMGStyle.cpp
@@ -67,6 +67,12 @@ TSharedRef< FSlateStyleSet > FUMGStyle::Create()
 	Style->Set("Widget.MultiLineEditableText", new IMAGE_BRUSH(TEXT("MultiLineEditableText"), Icon16x16));
 	Style->Set("Widget.MultiLineEditableTextBox", new IMAGE_BRUSH(TEXT("MultiLineEditableTextBox"), Icon16x16));
 	Style->Set("Widget.Viewport", new IMAGE_BRUSH(TEXT("Viewport"), Icon16x16));
+	Style->Set("Widget.ComboBox", new IMAGE_BRUSH(TEXT("icon_umg_ComboBox_16x"), Icon16x16));
+	Style->Set("Widget.ListView", new IMAGE_BRUSH(TEXT("icon_umg_ListView_16x"), Icon16x16));
+	Style->Set("Widget.TileView", new IMAGE_BRUSH(TEXT("icon_umg_TileView_16x"), Icon16x16));
+	Style->Set("Widget.Overlay", new IMAGE_BRUSH(TEXT("icon_umg_Overlay_16x"), Icon16x16));
+	Style->Set("Widget.Throbber", new IMAGE_BRUSH(TEXT("icon_umg_ThrobberHorizontal_16x"), Icon16x16));
+	Style->Set("Widget.CircularThrobber", new IMAGE_BRUSH(TEXT("icon_umg_ThrobberA_16x"), Icon16x16));
 
 	Style->Set("Widget.UserWidget", new IMAGE_BRUSH(TEXT("UserWidget"), Icon16x16));
 
diff --git a/Engine/Source/Runtime/UMG/Public/Components/CircularThrobber.h b/Engine/Source/Runtime/UMG/Public/Components/CircularThrobber.h
--- a/Engine/Source/Runtime/UMG/Public/Components/CircularThrobber.h
+++ b/Engine/Source/Runtime/UMG/Public/Components/CircularThrobber.h
@@ -36,6 +36,10 @@ public:
 
 	virtual void ReleaseNativeWidget() override;
 
+#if WITH_EDITOR
+	virtual const FSlateBrush* GetEditorIcon() override;
+#endif
+
 protected:
 
 	/** How many pieces there are */
diff --git a/Engine/Source/Runtime/UMG/Public/Components/ComboBox.h b/Engine/Source/Runtime/UMG/Public/Components/ComboBox.h
--- a/Engine/Source/Runtime/UMG/Public/Components/ComboBox.h
+++ b/Engine/Source/Runtime/UMG/Public/Components/ComboBox.h
@@ -22,6 +22,10 @@ public:
 
 	virtual void ReleaseNativeWidget() override;
 	
+#if WITH_EDITOR
+	virtual const FSlateBrush* GetEditorIcon() override;
+#endif
+
 protected:
 	TSharedRef<SWidget> HandleGenerateWidget(UObject* Item) const;
 
diff --git a/Engine/Source/Runtime/UMG/Public/Components/ListView.h b/Engine/Source/Runtime/UMG/Public/Components/ListView.h
--- a/Engine/Source/Runtime/UMG/Public/Components/ListView.h
+++ b/Engine/Source/Runtime/UMG/Public/Components/ListView.h
@@ -28,6 +28,10 @@ public:
 	UPROPERTY(EditDefaultsOnly, Category=Events)
 	FOnGenerateRowUObject OnGenerateRowEvent;
 
+#if WITH_EDITOR
+	virtual const FSlateBrush* GetEditorIcon() override;
+#endif
+
 protected:
 	TSharedPtr< SListView<UObject*> > MyListView;
 
diff --git a/Engine/Source/Runtime/UMG/Public/Components/Overlay.h b/Engine/Source/Runtime/UMG/Public/Components/Overlay.h
--- a/Engine/Source/Runtime/UMG/Public/Components/Overlay.h
+++ b/Engine/Source/Runtime/UMG/Public/Components/Overlay.h
@@ -11,6 +11,10 @@ class UMG_API UOverlay : public UPanelWidget
 
 	virtual void ReleaseNativeWidget() override;
 
+#if WITH_EDITOR
+	virtual const FSlateBrush* GetEditorIcon() override;
+#endif
+
 protected:
 
 	// UPanelWidget
diff --git a/Engine/Source/Runtime/UMG/Public/Components/Throbber.h b/Engine/Source/Runtime/UMG/Public/Components/Throbber.h
--- a/Engine/Source/Runtime/UMG/Public/Components/Throbber.h
+++ b/Engine/Source/Runtime/UMG/Public/Components/Throbber.h
@@ -38,6 +38,10 @@ public:
 	virtual void SyncronizeProperties() override;
 	// End of UWidget interface
 
+#if WITH_EDITOR
+	virtual const FSlateBrush* GetEditorIcon() override;
+#endif
+
 protected:
 
 	/** How many pieces there are */
diff --git a/Engine/Source/Runtime/UMG/Public/Components/TileView.h b/Engine/Source/Runtime/UMG/Public/Components/TileView.h
--- a/Engine/Source/Runtime/UMG/Public/Components/TileView.h
+++ b/Engine/Source/Runtime/UMG/Public/Components/TileView.h
@@ -42,6 +42,10 @@ class UMG_API UTileView : public UTableViewBase
 	UFUNCTION(BlueprintCallable, Category="Behavior")
 	void RequestListRefresh();
 
+#if WITH_EDITOR
+	virtual const FSlateBrush* GetEditorIcon() override;
+#endif
+
 protected:
 	TSharedRef<ITableRow> HandleOnGenerateTile(UObject* Item, const TSharedRef< STableViewBase >& OwnerTable) const;
 
