UMG - Adding a UWidget that can directly host Slate widgets inside of it.

diff --git a/Engine/Source/Runtime/UMG/Private/Components/NativeWidgetHost.cpp b/Engine/Source/Runtime/UMG/Private/Components/NativeWidgetHost.cpp
new file mode 100644
--- /dev/null
+++ b/Engine/Source/Runtime/UMG/Private/Components/NativeWidgetHost.cpp
@@ -0,0 +1,59 @@
+// Copyright 1998-2014 Epic Games, Inc. All Rights Reserved.
+
+#include "UMGPrivatePCH.h"
+
+#define LOCTEXT_NAMESPACE "UMG"
+
+/////////////////////////////////////////////////////
+// UNativeWidgetHost
+
+UNativeWidgetHost::UNativeWidgetHost(const FPostConstructInitializeProperties& PCIP)
+	: Super(PCIP)
+{
+	bIsVariable = false;
+}
+
+void UNativeWidgetHost::SetContent(TSharedRef<SWidget> InContent)
+{
+	MyWidget = InContent;
+}
+
+void UNativeWidgetHost::ReleaseNativeWidget()
+{
+	Super::ReleaseNativeWidget();
+
+	MyWidget.Reset();
+}
+
+TSharedRef<SWidget> UNativeWidgetHost::RebuildWidget()
+{
+	if ( MyWidget.IsValid() )
+	{
+		return MyWidget.ToSharedRef();
+	}
+	else
+	{
+		return SNew(SBorder)
+		.Visibility(EVisibility::HitTestInvisible)
+		.BorderImage(FUMGStyle::Get().GetBrush("MarchingAnts"))
+		.HAlign(HAlign_Center)
+		.VAlign(VAlign_Center)
+		[
+			SNew(STextBlock)
+			.Text(LOCTEXT("NativeWidgetHostText", "Slate Widget Host"))
+		];
+	}
+}
+
+#if WITH_EDITOR
+
+const FSlateBrush* UNativeWidgetHost::GetEditorIcon()
+{
+	return FUMGStyle::Get().GetBrush("Widget.NativeWidgetHost");
+}
+
+#endif
+
+/////////////////////////////////////////////////////
+
+#undef LOCTEXT_NAMESPACE
diff --git a/Engine/Source/Runtime/UMG/Private/UMGStyle.cpp b/Engine/Source/Runtime/UMG/Private/UMGStyle.cpp
--- a/Engine/Source/Runtime/UMG/Private/UMGStyle.cpp
+++ b/Engine/Source/Runtime/UMG/Private/UMGStyle.cpp
@@ -73,6 +73,7 @@ TSharedRef< FSlateStyleSet > FUMGStyle::Create()
 	Style->Set("Widget.Overlay", new IMAGE_BRUSH(TEXT("icon_umg_Overlay_16x"), Icon16x16));
 	Style->Set("Widget.Throbber", new IMAGE_BRUSH(TEXT("icon_umg_ThrobberHorizontal_16x"), Icon16x16));
 	Style->Set("Widget.CircularThrobber", new IMAGE_BRUSH(TEXT("icon_umg_ThrobberA_16x"), Icon16x16));
+	Style->Set("Widget.NativeWidgetHost", new IMAGE_BRUSH(TEXT("NativeWidgetHost"), Icon16x16));
 
 	Style->Set("Widget.UserWidget", new IMAGE_BRUSH(TEXT("UserWidget"), Icon16x16));
 
diff --git a/Engine/Source/Runtime/UMG/Public/Components/MenuAnchor.h b/Engine/Source/Runtime/UMG/Public/Components/MenuAnchor.h
--- a/Engine/Source/Runtime/UMG/Public/Components/MenuAnchor.h
+++ b/Engine/Source/Runtime/UMG/Public/Components/MenuAnchor.h
@@ -8,7 +8,7 @@
  * The Menu Anchor allows you to specify an location that a popup menu should be anchored to, and should be
  * summoned from.
  */
-UCLASS(meta=( Category="Popup" ), ClassGroup=UserInterface, meta=( IsAdvanced = "True" ))
+UCLASS(meta=( Category="Advanced" ), ClassGroup=UserInterface, meta=( IsAdvanced = "True" ))
 class UMG_API UMenuAnchor : public UContentWidget
 {
 	GENERATED_UCLASS_BODY()
diff --git a/Engine/Source/Runtime/UMG/Public/Components/NativeWidgetHost.h b/Engine/Source/Runtime/UMG/Public/Components/NativeWidgetHost.h
new file mode 100644
--- /dev/null
+++ b/Engine/Source/Runtime/UMG/Public/Components/NativeWidgetHost.h
@@ -0,0 +1,31 @@
+// Copyright 1998-2014 Epic Games, Inc. All Rights Reserved.
+
+#pragma once
+
+#include "NativeWidgetHost.generated.h"
+
+/**
+ * A NativeWidgetHost is a container widget that can contain one child slate widget.  This should
+ * be used when all you need is to nest a native widget inside a UMG widget.
+ */
+UCLASS(meta=( Category="Advanced" ), ClassGroup=UserInterface)
+class UMG_API UNativeWidgetHost : public UWidget
+{
+	GENERATED_UCLASS_BODY()
+
+	void SetContent(TSharedRef<SWidget> InContent);
+
+	virtual void ReleaseNativeWidget() override;
+
+#if WITH_EDITOR
+	virtual const FSlateBrush* GetEditorIcon() override;
+#endif
+
+protected:
+	TSharedPtr<SWidget> MyWidget;
+
+protected:
+	// UWidget interface
+	virtual TSharedRef<SWidget> RebuildWidget() override;
+	// End of UWidget interface
+};
diff --git a/Engine/Source/Runtime/UMG/Public/Components/ScrollBar.h b/Engine/Source/Runtime/UMG/Public/Components/ScrollBar.h
--- a/Engine/Source/Runtime/UMG/Public/Components/ScrollBar.h
+++ b/Engine/Source/Runtime/UMG/Public/Components/ScrollBar.h
@@ -5,7 +5,7 @@
 #include "ScrollBar.generated.h"
 
 /** */
-UCLASS(meta=( Category="Primitive" ), ClassGroup=UserInterface, meta=( IsAdvanced = "True" ))
+UCLASS(meta=( Category="Advanced" ), ClassGroup=UserInterface, meta=( IsAdvanced = "True" ))
 class UMG_API UScrollBar : public UPanelWidget
 {
 	GENERATED_UCLASS_BODY()
diff --git a/Engine/Source/Runtime/UMG/Public/UMG.h b/Engine/Source/Runtime/UMG/Public/UMG.h
--- a/Engine/Source/Runtime/UMG/Public/UMG.h
+++ b/Engine/Source/Runtime/UMG/Public/UMG.h
@@ -87,6 +87,8 @@
 
 #include "Viewport.h"
 
+#include "NativeWidgetHost.h"
+
 // Slate
 #include "SObjectWidget.h"
 
