UMG - Now attempting to load fonts from the game content directory first.

diff --git a/Engine/Source/Runtime/UMG/Private/Components/EditableText.cpp b/Engine/Source/Runtime/UMG/Private/Components/EditableText.cpp
--- a/Engine/Source/Runtime/UMG/Private/Components/EditableText.cpp
+++ b/Engine/Source/Runtime/UMG/Private/Components/EditableText.cpp
@@ -36,7 +36,12 @@ void UEditableText::ReleaseNativeWidget()
 
 TSharedRef<SWidget> UEditableText::RebuildWidget()
 {
-	FString FontPath = FPaths::EngineContentDir() / Font.FontName.ToString();
+	FString FontPath = FPaths::GameContentDir() / Font.FontName.ToString();
+
+	if ( !FPaths::FileExists(FontPath) )
+	{
+		FontPath = FPaths::EngineContentDir() / Font.FontName.ToString();
+	}
 	
 	SEditableText::FArguments Defaults;
 	
diff --git a/Engine/Source/Runtime/UMG/Private/Components/EditableTextBox.cpp b/Engine/Source/Runtime/UMG/Private/Components/EditableTextBox.cpp
--- a/Engine/Source/Runtime/UMG/Private/Components/EditableTextBox.cpp
+++ b/Engine/Source/Runtime/UMG/Private/Components/EditableTextBox.cpp
@@ -39,7 +39,12 @@ void UEditableTextBox::ReleaseNativeWidget()
 
 TSharedRef<SWidget> UEditableTextBox::RebuildWidget()
 {
-	FString FontPath = FPaths::EngineContentDir() / Font.FontName.ToString();
+	FString FontPath = FPaths::GameContentDir() / Font.FontName.ToString();
+
+	if ( !FPaths::FileExists(FontPath) )
+	{
+		FontPath = FPaths::EngineContentDir() / Font.FontName.ToString();
+	}
 
 	MyEditableTextBlock = SNew(SEditableTextBox)
 		.Font(FSlateFontInfo(FontPath, Font.Size))
diff --git a/Engine/Source/Runtime/UMG/Private/Components/MultiLineEditableText.cpp b/Engine/Source/Runtime/UMG/Private/Components/MultiLineEditableText.cpp
--- a/Engine/Source/Runtime/UMG/Private/Components/MultiLineEditableText.cpp
+++ b/Engine/Source/Runtime/UMG/Private/Components/MultiLineEditableText.cpp
@@ -16,7 +16,12 @@ UMultiLineEditableText::UMultiLineEditableText(const FPostConstructInitializePro
 
 TSharedRef<SWidget> UMultiLineEditableText::RebuildWidget()
 {
-	FString FontPath = FPaths::EngineContentDir() / Font.FontName.ToString();
+	FString FontPath = FPaths::GameContentDir() / Font.FontName.ToString();
+
+	if ( !FPaths::FileExists(FontPath) )
+	{
+		FontPath = FPaths::EngineContentDir() / Font.FontName.ToString();
+	}
 	
 	SMultiLineEditableText::FArguments Defaults;
 	
diff --git a/Engine/Source/Runtime/UMG/Private/Components/MultiLineEditableTextBox.cpp b/Engine/Source/Runtime/UMG/Private/Components/MultiLineEditableTextBox.cpp
--- a/Engine/Source/Runtime/UMG/Private/Components/MultiLineEditableTextBox.cpp
+++ b/Engine/Source/Runtime/UMG/Private/Components/MultiLineEditableTextBox.cpp
@@ -20,7 +20,12 @@ UMultiLineEditableTextBox::UMultiLineEditableTextBox(const FPostConstructInitial
 
 TSharedRef<SWidget> UMultiLineEditableTextBox::RebuildWidget()
 {
-	FString FontPath = FPaths::EngineContentDir() / Font.FontName.ToString();
+	FString FontPath = FPaths::GameContentDir() / Font.FontName.ToString();
+
+	if ( !FPaths::FileExists(FontPath) )
+	{
+		FontPath = FPaths::EngineContentDir() / Font.FontName.ToString();
+	}
 
 	MyEditableTextBlock = SNew(SMultiLineEditableTextBox)
 		.Font(FSlateFontInfo(FontPath, Font.Size))
diff --git a/Engine/Source/Runtime/UMG/Private/Components/TextBlock.cpp b/Engine/Source/Runtime/UMG/Private/Components/TextBlock.cpp
--- a/Engine/Source/Runtime/UMG/Private/Components/TextBlock.cpp
+++ b/Engine/Source/Runtime/UMG/Private/Components/TextBlock.cpp
@@ -73,7 +73,12 @@ void UTextBlock::SyncronizeProperties()
 {
 	Super::SyncronizeProperties();
 
-	FString FontPath = FPaths::EngineContentDir() / Font.FontName.ToString();
+	FString FontPath = FPaths::GameContentDir() / Font.FontName.ToString();
+
+	if ( !FPaths::FileExists(FontPath) )
+	{
+		FontPath = FPaths::EngineContentDir() / Font.FontName.ToString();
+	}
 
 	const FTextBlockStyle* StylePtr = ( Style != NULL ) ? Style->GetStyle<FTextBlockStyle>() : NULL;
 	if ( StylePtr == NULL )
