UMG - New Widget: "Scale Box". LetÆs you insert content that has a desired size larger or smaller than the allotted space, that you want to scale up or down in a way that can preserve aspect ratio.  Great for auto scaling text and images that fit but maintain aspect ratio when scaled to fit different devices.

diff --git a/Engine/Source/Runtime/UMG/Private/Components/ScaleBox.cpp b/Engine/Source/Runtime/UMG/Private/Components/ScaleBox.cpp
new file mode 100644
--- /dev/null
+++ b/Engine/Source/Runtime/UMG/Private/Components/ScaleBox.cpp
@@ -0,0 +1,80 @@
+// Copyright 1998-2014 Epic Games, Inc. All Rights Reserved.
+
+#include "UMGPrivatePCH.h"
+
+#define LOCTEXT_NAMESPACE "UMG"
+
+/////////////////////////////////////////////////////
+// UScaleBox
+
+UScaleBox::UScaleBox(const FPostConstructInitializeProperties& PCIP)
+	: Super(PCIP)
+{
+	bIsVariable = false;
+
+	StretchDirection = EStretchDirection::Both;
+	Stretch = EStretch::ScaleToFit;
+}
+
+void UScaleBox::ReleaseNativeWidget()
+{
+	Super::ReleaseNativeWidget();
+
+	MyScaleBox.Reset();
+}
+
+TSharedRef<SWidget> UScaleBox::RebuildWidget()
+{
+	MyScaleBox = SNew(SScaleBox);
+	
+	if ( GetChildrenCount() > 0 )
+	{
+		Cast<UScaleBoxSlot>(GetContentSlot())->BuildSlot(MyScaleBox.ToSharedRef());
+	}
+
+	return MyScaleBox.ToSharedRef();
+}
+
+void UScaleBox::SyncronizeProperties()
+{
+	Super::SyncronizeProperties();
+
+	MyScaleBox->SetStretchDirection(StretchDirection);
+	MyScaleBox->SetStretch(Stretch);
+}
+
+UClass* UScaleBox::GetSlotClass() const
+{
+	return UScaleBoxSlot::StaticClass();
+}
+
+void UScaleBox::OnSlotAdded(UPanelSlot* Slot)
+{
+	// Add the child to the live slot if it already exists
+	if ( MyScaleBox.IsValid() )
+	{
+		Cast<UScaleBoxSlot>(Slot)->BuildSlot(MyScaleBox.ToSharedRef());
+	}
+}
+
+void UScaleBox::OnSlotRemoved(UPanelSlot* Slot)
+{
+	// Remove the widget from the live slot if it exists.
+	if ( MyScaleBox.IsValid() )
+	{
+		MyScaleBox->SetContent(SNullWidget::NullWidget);
+	}
+}
+
+#if WITH_EDITOR
+
+const FSlateBrush* UScaleBox::GetEditorIcon()
+{
+	return FUMGStyle::Get().GetBrush("Widget.ScaleBox");
+}
+
+#endif
+
+/////////////////////////////////////////////////////
+
+#undef LOCTEXT_NAMESPACE
diff --git a/Engine/Source/Runtime/UMG/Private/Components/ScaleBoxSlot.cpp b/Engine/Source/Runtime/UMG/Private/Components/ScaleBoxSlot.cpp
new file mode 100644
--- /dev/null
+++ b/Engine/Source/Runtime/UMG/Private/Components/ScaleBoxSlot.cpp
@@ -0,0 +1,67 @@
+// Copyright 1998-2014 Epic Games, Inc. All Rights Reserved.
+
+#include "UMGPrivatePCH.h"
+
+/////////////////////////////////////////////////////
+// UScaleBoxSlot
+
+UScaleBoxSlot::UScaleBoxSlot(const FPostConstructInitializeProperties& PCIP)
+	: Super(PCIP)
+{
+	Padding = FMargin(0, 0);
+
+	HorizontalAlignment = HAlign_Center;
+	VerticalAlignment = VAlign_Center;
+}
+
+void UScaleBoxSlot::ReleaseNativeWidget()
+{
+	Super::ReleaseNativeWidget();
+
+	ScaleBox.Reset();
+}
+
+void UScaleBoxSlot::BuildSlot(TSharedRef<SScaleBox> InScaleBox)
+{
+	ScaleBox = InScaleBox;
+
+	//ScaleBox->SetPadding(Padding);
+	ScaleBox->SetHAlign(HorizontalAlignment);
+	ScaleBox->SetVAlign(VerticalAlignment);
+
+	ScaleBox->SetContent(Content ? Content->TakeWidget() : SNullWidget::NullWidget);
+}
+
+void UScaleBoxSlot::SetPadding(FMargin InPadding)
+{
+	Padding = InPadding;
+	if ( ScaleBox.IsValid() )
+	{
+		//ScaleBox->SetPadding(InPadding);
+	}
+}
+
+void UScaleBoxSlot::SetHorizontalAlignment(EHorizontalAlignment InHorizontalAlignment)
+{
+	HorizontalAlignment = InHorizontalAlignment;
+	if ( ScaleBox.IsValid() )
+	{
+		ScaleBox->SetHAlign(InHorizontalAlignment);
+	}
+}
+
+void UScaleBoxSlot::SetVerticalAlignment(EVerticalAlignment InVerticalAlignment)
+{
+	VerticalAlignment = InVerticalAlignment;
+	if ( ScaleBox.IsValid() )
+	{
+		ScaleBox->SetVAlign(InVerticalAlignment);
+	}
+}
+
+void UScaleBoxSlot::SyncronizeProperties()
+{
+	SetPadding(Padding);
+	SetHorizontalAlignment(HorizontalAlignment);
+	SetVerticalAlignment(VerticalAlignment);
+}
diff --git a/Engine/Source/Runtime/UMG/Private/Slate/SScaleBox.cpp b/Engine/Source/Runtime/UMG/Private/Slate/SScaleBox.cpp
new file mode 100644
--- /dev/null
+++ b/Engine/Source/Runtime/UMG/Private/Slate/SScaleBox.cpp
@@ -0,0 +1,114 @@
+// Copyright 1998-2014 Epic Games, Inc. All Rights Reserved.
+
+#include "UMGPrivatePCH.h"
+#include "LayoutUtils.h"
+#include "SScaleBox.h"
+
+
+/* SScaleBox interface
+ *****************************************************************************/
+
+void SScaleBox::Construct( const SScaleBox::FArguments& InArgs )
+{
+	ChildSlot
+	[
+		InArgs._Content.Widget
+	];
+}
+
+/* SWidget overrides
+ *****************************************************************************/
+
+void SScaleBox::OnArrangeChildren( const FGeometry& AllottedGeometry, FArrangedChildren& ArrangedChildren ) const
+{
+	const EVisibility ChildVisibility = ChildSlot.GetWidget()->GetVisibility();
+	if ( ArrangedChildren.Accepts(ChildVisibility) )
+	{
+		FVector2D DesiredSize = ChildSlot.GetWidget()->GetDesiredSize();
+
+		float FinalScale = 1;
+
+		EStretch::Type CurrentStretch = Stretch.Get();
+		EStretchDirection::Type CurrentStretchDirection = StretchDirection.Get();
+
+		if ( DesiredSize.X != 0 && DesiredSize.Y != 0 )
+		{
+			switch ( CurrentStretch )
+			{
+			case EStretch::None:
+				break;
+			case EStretch::Fill:
+				DesiredSize = AllottedGeometry.Size;
+				break;
+			case EStretch::ScaleToFit:
+				FinalScale = FMath::Min(AllottedGeometry.Size.X / DesiredSize.X, AllottedGeometry.Size.Y / DesiredSize.Y);
+				break;
+			case EStretch::ScaleToFill:
+				FinalScale = FMath::Max(AllottedGeometry.Size.X / DesiredSize.X, AllottedGeometry.Size.Y / DesiredSize.Y);
+				break;
+			}
+
+			switch ( CurrentStretchDirection )
+			{
+			case EStretchDirection::DownOnly:
+				FinalScale = FMath::Min(FinalScale, 1.0f);
+				break;
+			case EStretchDirection::UpOnly:
+				FinalScale = FMath::Max(FinalScale, 1.0f);
+				break;
+			}
+		}
+
+		FVector2D FinalOffset(0, 0);
+
+		if ( CurrentStretch != EStretch::Fill )
+		{
+			const FMargin SlotPadding(ChildSlot.SlotPadding.Get());
+			AlignmentArrangeResult XResult = AlignChild<Orient_Horizontal>(AllottedGeometry.Size.X, ChildSlot, SlotPadding, FinalScale, false);
+			AlignmentArrangeResult YResult = AlignChild<Orient_Vertical>(AllottedGeometry.Size.Y, ChildSlot, SlotPadding, FinalScale, false);
+
+			FinalOffset = FVector2D(XResult.Offset, YResult.Offset) * ( 1.0f / FinalScale );
+		}
+
+		ArrangedChildren.AddWidget(ChildVisibility, AllottedGeometry.MakeChild(
+			ChildSlot.GetWidget(),
+			FinalOffset,
+			DesiredSize,
+			FinalScale
+		) );
+	}
+}
+
+
+int32 SScaleBox::OnPaint( const FPaintArgs& Args, const FGeometry& AllottedGeometry, const FSlateRect& MyClippingRect, FSlateWindowElementList& OutDrawElements, int32 LayerId, const FWidgetStyle& InWidgetStyle, bool bParentEnabled ) const
+{
+	return SCompoundWidget::OnPaint(Args, AllottedGeometry, MyClippingRect, OutDrawElements, LayerId, InWidgetStyle, bParentEnabled);
+}
+
+void SScaleBox::SetContent(TSharedRef<SWidget> InContent)
+{
+	ChildSlot
+	[
+		InContent
+	];
+}
+
+void SScaleBox::SetHAlign(EHorizontalAlignment HAlign)
+{
+	ChildSlot.HAlignment = HAlign;
+}
+
+void SScaleBox::SetVAlign(EVerticalAlignment VAlign)
+{
+	ChildSlot.VAlignment = VAlign;
+}
+
+void SScaleBox::SetStretchDirection(EStretchDirection::Type InStretchDirection)
+{
+	StretchDirection = InStretchDirection;
+}
+
+void SScaleBox::SetStretch(EStretch::Type InStretch)
+{
+	Stretch = InStretch;
+}
diff --git a/Engine/Source/Runtime/UMG/Private/UMGStyle.cpp b/Engine/Source/Runtime/UMG/Private/UMGStyle.cpp
--- a/Engine/Source/Runtime/UMG/Private/UMGStyle.cpp
+++ b/Engine/Source/Runtime/UMG/Private/UMGStyle.cpp
@@ -74,6 +74,7 @@ TSharedRef< FSlateStyleSet > FUMGStyle::Create()
 	Style->Set("Widget.Throbber", new IMAGE_BRUSH(TEXT("icon_umg_ThrobberHorizontal_16x"), Icon16x16));
 	Style->Set("Widget.CircularThrobber", new IMAGE_BRUSH(TEXT("icon_umg_ThrobberA_16x"), Icon16x16));
 	Style->Set("Widget.NativeWidgetHost", new IMAGE_BRUSH(TEXT("NativeWidgetHost"), Icon16x16));
+	Style->Set("Widget.ScaleBox", new IMAGE_BRUSH(TEXT("ScaleBox"), Icon16x16));
 
 	Style->Set("Widget.UserWidget", new IMAGE_BRUSH(TEXT("UserWidget"), Icon16x16));
 
diff --git a/Engine/Source/Runtime/UMG/Public/Components/ScaleBox.h b/Engine/Source/Runtime/UMG/Public/Components/ScaleBox.h
new file mode 100644
--- /dev/null
+++ b/Engine/Source/Runtime/UMG/Public/Components/ScaleBox.h
@@ -0,0 +1,50 @@
+// Copyright 1998-2014 Epic Games, Inc. All Rights Reserved.
+
+#pragma once
+
+#include "ScaleBox.generated.h"
+
+/**
+ * Allows you to place content with a desired size and have it scale to meet the constraints placed on this box's alloted area.  If
+ * you needed to have a background image scale to fill an area but not become distorted with different aspect ratios, or if you need
+ * to auto fit some text to an area, this is the control for you.
+ */
+UCLASS(meta=( Category="Panel" ), ClassGroup=UserInterface)
+class UMG_API UScaleBox : public UContentWidget
+{
+	GENERATED_UCLASS_BODY()
+
+	/** Controls in what direction content can be scaled */
+	UPROPERTY(EditDefaultsOnly, Category="Stretching")
+	TEnumAsByte<EStretchDirection::Type> StretchDirection;
+
+	/** The stretching rule to apply when content is stretched */
+	UPROPERTY(EditDefaultsOnly, Category="Stretching")
+	TEnumAsByte<EStretch::Type> Stretch;
+
+	// UWidget interface
+	virtual void SyncronizeProperties() override;
+	// End of UWidget interface
+
+	virtual void ReleaseNativeWidget() override;
+
+#if WITH_EDITOR
+	virtual const FSlateBrush* GetEditorIcon() override;
+#endif
+
+protected:
+
+	// UPanelWidget
+	virtual UClass* GetSlotClass() const override;
+	virtual void OnSlotAdded(UPanelSlot* Slot) override;
+	virtual void OnSlotRemoved(UPanelSlot* Slot) override;
+	// End UPanelWidget
+
+protected:
+	TSharedPtr<SScaleBox> MyScaleBox;
+
+protected:
+	// UWidget interface
+	virtual TSharedRef<SWidget> RebuildWidget() override;
+	// End of UWidget interface
+};
diff --git a/Engine/Source/Runtime/UMG/Public/Components/ScaleBoxSlot.h b/Engine/Source/Runtime/UMG/Public/Components/ScaleBoxSlot.h
new file mode 100644
--- /dev/null
+++ b/Engine/Source/Runtime/UMG/Public/Components/ScaleBoxSlot.h
@@ -0,0 +1,49 @@
+// Copyright 1998-2014 Epic Games, Inc. All Rights Reserved.
+
+#pragma once
+
+#include "SlateWrapperTypes.h"
+
+#include "ScaleBoxSlot.generated.h"
+
+/** The Slot for the UScaleBoxSlot, contains the widget displayed in a button's single slot */
+UCLASS()
+class UMG_API UScaleBoxSlot : public UPanelSlot
+{
+	GENERATED_UCLASS_BODY()
+	
+	/** The padding area between the slot and the content it contains. */
+	UPROPERTY(EditDefaultsOnly, Category=Layout)
+	FMargin Padding;
+
+	/** The alignment of the object horizontally. */
+	UPROPERTY(EditDefaultsOnly, Category=Layout)
+	TEnumAsByte<EHorizontalAlignment> HorizontalAlignment;
+
+	/** The alignment of the object vertically. */
+	UPROPERTY(EditDefaultsOnly, Category=Layout)
+	TEnumAsByte<EVerticalAlignment> VerticalAlignment;
+
+	UFUNCTION(BlueprintCallable, Category="Appearance")
+	void SetPadding(FMargin InPadding);
+
+	UFUNCTION(BlueprintCallable, Category="Appearance")
+	void SetHorizontalAlignment(EHorizontalAlignment InHorizontalAlignment);
+
+	UFUNCTION(BlueprintCallable, Category="Appearance")
+	void SetVerticalAlignment(EVerticalAlignment InVerticalAlignment);
+
+	// UPanelSlot interface
+	virtual void SyncronizeProperties() override;
+	// End of UPanelSlot interface
+
+	/** Builds the underlying slot for the slate ScaleBox. */
+	void BuildSlot(TSharedRef<SScaleBox> InScaleBox);
+
+	virtual void ReleaseNativeWidget() override;
+
+private:
+
+	/** A pointer to the button to allow us to adjust the size, padding...etc at runtime. */
+	TSharedPtr<SScaleBox> ScaleBox;
+};
diff --git a/Engine/Source/Runtime/UMG/Public/Slate/SScaleBox.h b/Engine/Source/Runtime/UMG/Public/Slate/SScaleBox.h
new file mode 100644
--- /dev/null
+++ b/Engine/Source/Runtime/UMG/Public/Slate/SScaleBox.h
@@ -0,0 +1,90 @@
+// Copyright 1998-2014 Epic Games, Inc. All Rights Reserved.
+
+#pragma once
+
+#include "SScaleBox.generated.h"
+
+UENUM()
+namespace EStretchDirection
+{
+	enum Type
+	{
+		/** Will scale the content up or down */
+		Both,
+		/** Will only make the content smaller, will never scale it larger than the content's desired size. */
+		DownOnly,
+		/** Will only make the content larger, will never scale it smaller than the content's desired size. */
+		UpOnly
+	};
+}
+
+UENUM()
+namespace EStretch
+{
+	enum Type
+	{
+		/** Does not scale the content. */
+		None,
+		/** Scales the content non-uniformly filling the entire space of the area */
+		Fill,
+		/** Scales the content uniformly (preserving aspect ratio) until it can no longer scale the content without clipping it. */
+		ScaleToFit,
+		/** Scales the content uniformly (preserving aspect ratio), until all sides meet or exceed the size of the area.  Will result in clipping longer sides. */
+		ScaleToFill
+	};
+}
+
+/**
+ * Allows you to place content with a desired size and have it scale to meet the constraints placed on this box's alloted area.  If
+ * you needed to have a background image scale to fill an area but not become distorted with different aspect ratios, or if you need
+ * to auto fit some text to an area, this is the control for you.
+ */
+class UMG_API SScaleBox : public SCompoundWidget
+{
+public:
+	SLATE_BEGIN_ARGS(SScaleBox)
+	: _Content()
+	, _HAlign(HAlign_Center)
+	, _VAlign(VAlign_Center)
+	, _StretchDirection(EStretchDirection::Both)
+	, _Stretch(EStretch::None)
+	{}
+		/** Slot for this designers content (optional) */
+		SLATE_DEFAULT_SLOT(FArguments, Content)
+
+		SLATE_ARGUMENT(EHorizontalAlignment, HAlign)
+
+		SLATE_ARGUMENT(EVerticalAlignment, VAlign)
+		
+		/** Controls in what direction content can be scaled */
+		SLATE_ATTRIBUTE(EStretchDirection::Type, StretchDirection)
+		
+		/** The stretching rule to apply when content is stretched */
+		SLATE_ATTRIBUTE(EStretch::Type, Stretch)
+
+	SLATE_END_ARGS()
+
+	void Construct(const FArguments& InArgs);
+	
+	virtual void OnArrangeChildren(const FGeometry& AllottedGeometry, FArrangedChildren& ArrangedChildren) const override;
+	virtual int32 OnPaint(const FPaintArgs& Args, const FGeometry& AllottedGeometry, const FSlateRect& MyClippingRect, FSlateWindowElementList& OutDrawElements, int32 LayerId, const FWidgetStyle& InWidgetStyle, bool bParentEnabled) const override;
+
+	/** See Content slot */
+	void SetContent(TSharedRef<SWidget> InContent);
+
+	/** See HAlign argument */
+	void SetHAlign(EHorizontalAlignment HAlign);
+
+	/** See VAlign argument */
+	void SetVAlign(EVerticalAlignment VAlign);
+
+	/** See StretchDirection argument */
+	void SetStretchDirection(EStretchDirection::Type InStretchDirection);
+
+	/** See Stretch argument */
+	void SetStretch(EStretch::Type InStretch);
+	
+private:
+	TAttribute<EStretchDirection::Type> StretchDirection;
+	TAttribute<EStretch::Type> Stretch;
+};
diff --git a/Engine/Source/Runtime/UMG/Public/UMG.h b/Engine/Source/Runtime/UMG/Public/UMG.h
--- a/Engine/Source/Runtime/UMG/Public/UMG.h
+++ b/Engine/Source/Runtime/UMG/Public/UMG.h
@@ -74,6 +74,10 @@
 #include "ProgressBar.h"
 #include "ComboBox.h"
 
+#include "SScaleBox.h"
+#include "ScaleBoxSlot.h"
+#include "ScaleBox.h"
+
 #include "Spacer.h"
 #include "MenuAnchor.h"
 #include "ScrollBar.h"
