UMG - Fixing a crash that occurs when a sub widget blueprint is recompiled and is still referenced by a running designer.  (Stale object)

diff --git a/Engine/Source/Runtime/UMG/Private/UserWidget.cpp b/Engine/Source/Runtime/UMG/Private/UserWidget.cpp
--- a/Engine/Source/Runtime/UMG/Private/UserWidget.cpp
+++ b/Engine/Source/Runtime/UMG/Private/UserWidget.cpp
@@ -189,31 +189,17 @@ void UUserWidget::OnAnimationFinishedPlaying( UUMGSequencePlayer& Player )
 
 UWidget* UUserWidget::GetWidgetHandle(TSharedRef<SWidget> InWidget)
 {
-	TWeakObjectPtr<UWidget> VisualWidget = WidgetToComponent.FindRef(InWidget);
-	return VisualWidget.Get();
+	return WidgetTree->FindWidget(InWidget);
 }
 
 TSharedRef<SWidget> UUserWidget::RebuildWidget()
 {
-	WidgetToComponent.Reset();
-	
 	TSharedPtr<SWidget> UserRootWidget;
 
 	// Add the first component to the root of the widget surface.
 	if ( Components.Num() > 0 && Components[0] != NULL )
 	{
 		UserRootWidget = Components[0]->TakeWidget();
-
-		// Place all of our top-level children Slate wrapped components into the overlay
-		for ( int32 ComponentIndex = 0; ComponentIndex < Components.Num(); ++ComponentIndex )
-		{
-			UWidget* Handle = Components[ComponentIndex];
-			TSharedPtr<SWidget> Widget = Handle->GetCachedWidget();
-			if ( Widget.IsValid() )
-			{
-				WidgetToComponent.Add(Widget, Handle);
-			}
-		}
 	}
 	else
 	{
@@ -231,12 +217,10 @@ TSharedRef<SWidget> UUserWidget::RebuildWidget()
 
 TSharedPtr<SWidget> UUserWidget::GetWidgetFromName(const FString& Name) const
 {
-	for ( auto& Entry : WidgetToComponent )
+	UWidget* WidgetObject = WidgetTree->FindWidget(Name);
+	if ( WidgetObject )
 	{
-		if ( Entry.Value->GetName().Equals(Name, ESearchCase::IgnoreCase) )
-		{
-			return Entry.Key.Pin();
-		}
+		return WidgetObject->GetCachedWidget();
 	}
 
 	return TSharedPtr<SWidget>();
diff --git a/Engine/Source/Runtime/UMG/Private/WidgetTree.cpp b/Engine/Source/Runtime/UMG/Private/WidgetTree.cpp
--- a/Engine/Source/Runtime/UMG/Private/WidgetTree.cpp
+++ b/Engine/Source/Runtime/UMG/Private/WidgetTree.cpp
@@ -30,6 +30,25 @@ UWidget* UWidgetTree::FindWidget(const FString& Name) const
 	return NULL;
 }
 
+UWidget* UWidgetTree::FindWidget(TSharedRef<SWidget> InWidget) const
+{
+	FString ExistingName;
+
+	// TODO UMG Hacky, remove this find widget function, or make it faster.
+	TArray<UWidget*> Widgets;
+	GetAllWidgets(Widgets);
+
+	for ( UWidget* Widget : Widgets )
+	{
+		if ( Widget->GetCachedWidget() == InWidget )
+		{
+			return Widget;
+		}
+	}
+
+	return NULL;
+}
+
 UPanelWidget* UWidgetTree::FindWidgetParent(UWidget* Widget, int32& OutChildIndex)
 {
 	UPanelWidget* Parent = Widget->GetParent();
diff --git a/Engine/Source/Runtime/UMG/Public/Blueprint/UserWidget.h b/Engine/Source/Runtime/UMG/Public/Blueprint/UserWidget.h
--- a/Engine/Source/Runtime/UMG/Public/Blueprint/UserWidget.h
+++ b/Engine/Source/Runtime/UMG/Public/Blueprint/UserWidget.h
@@ -336,8 +336,6 @@ protected:
 	int32 GetFullScreenZOrder() const;
 
 private:
-	TMap< TWeakPtr<SWidget>, TWeakObjectPtr<UWidget> > WidgetToComponent;
-
 	TWeakPtr<SWidget> FullScreenWidget;
 
 	FLocalPlayerContext PlayerContext;
diff --git a/Engine/Source/Runtime/UMG/Public/Blueprint/WidgetTree.h b/Engine/Source/Runtime/UMG/Public/Blueprint/WidgetTree.h
--- a/Engine/Source/Runtime/UMG/Public/Blueprint/WidgetTree.h
+++ b/Engine/Source/Runtime/UMG/Public/Blueprint/WidgetTree.h
@@ -12,14 +12,22 @@ class UMG_API UWidgetTree : public UObject
 
 public:
 
+	/** Finds the widget in the tree by name. */
 	UWidget* FindWidget(const FString& Name) const;
 
+	/** Finds a widget in the tree using the native widget as the key. */
+	UWidget* FindWidget(TSharedRef<SWidget> InWidget) const;
+
+	/** Removes the widget from the hierarchy and all sub widgets. */
 	bool RemoveWidget(UWidget* Widget);
 
+	/** Gets the parent widget of a given widget, and potentially the child index. */
 	class UPanelWidget* FindWidgetParent(UWidget* Widget, int32& OutChildIndex);
 
+	/** Gathers all the widgets in the tree recursively */
 	void GetAllWidgets(TArray<UWidget*>& Widgets) const;
 
+	/** Gathers only the immediate child widgets of a parent widget. */
 	void GetChildWidgets(UWidget* Parent, TArray<UWidget*>& Widgets) const;
 
 	/** Constructs the widget, and adds it to the tree. */
