UMG: NULL to nullptr

diff --git a/Engine/Source/Runtime/UMG/Private/Components/ScrollBar.cpp b/Engine/Source/Runtime/UMG/Private/Components/ScrollBar.cpp
--- a/Engine/Source/Runtime/UMG/Private/Components/ScrollBar.cpp
+++ b/Engine/Source/Runtime/UMG/Private/Components/ScrollBar.cpp
@@ -19,8 +19,8 @@ UScrollBar::UScrollBar(const FPostConstructInitializeProperties& PCIP)
 
 TSharedRef<SWidget> UScrollBar::RebuildWidget()
 {
-	const FScrollBarStyle* StylePtr = ( Style != NULL ) ? Style->GetStyle<FScrollBarStyle>() : NULL;
-	if ( StylePtr == NULL )
+	const FScrollBarStyle* StylePtr = ( Style != nullptr ) ? Style->GetStyle<FScrollBarStyle>() : nullptr;
+	if ( StylePtr == nullptr )
 	{
 		SScrollBar::FArguments Defaults;
 		StylePtr = Defaults._Style;
diff --git a/Engine/Source/Runtime/UMG/Private/Components/ScrollBox.cpp b/Engine/Source/Runtime/UMG/Private/Components/ScrollBox.cpp
--- a/Engine/Source/Runtime/UMG/Private/Components/ScrollBox.cpp
+++ b/Engine/Source/Runtime/UMG/Private/Components/ScrollBox.cpp
@@ -53,15 +53,15 @@ void UScrollBox::OnSlotRemoved(UPanelSlot* Slot)
 
 TSharedRef<SWidget> UScrollBox::RebuildWidget()
 {
-	const FScrollBoxStyle* StylePtr = ( Style != NULL ) ? Style->GetStyle<FScrollBoxStyle>() : NULL;
-	if ( StylePtr == NULL )
+	const FScrollBoxStyle* StylePtr = ( Style != nullptr ) ? Style->GetStyle<FScrollBoxStyle>() : nullptr;
+	if ( StylePtr == nullptr )
 	{
 		SScrollBox::FArguments Defaults;
 		StylePtr = Defaults._Style;
 	}
 
-	const FScrollBarStyle* BarStylePtr = ( BarStyle != NULL ) ? BarStyle->GetStyle<FScrollBarStyle>() : NULL;
-	if ( BarStylePtr == NULL )
+	const FScrollBarStyle* BarStylePtr = ( BarStyle != nullptr ) ? BarStyle->GetStyle<FScrollBarStyle>() : nullptr;
+	if ( BarStylePtr == nullptr )
 	{
 		SScrollBox::FArguments Defaults;
 		BarStylePtr = Defaults._ScrollBarStyle;
diff --git a/Engine/Source/Runtime/UMG/Private/Components/ScrollBoxSlot.cpp b/Engine/Source/Runtime/UMG/Private/Components/ScrollBoxSlot.cpp
--- a/Engine/Source/Runtime/UMG/Private/Components/ScrollBoxSlot.cpp
+++ b/Engine/Source/Runtime/UMG/Private/Components/ScrollBoxSlot.cpp
@@ -7,7 +7,7 @@
 
 UScrollBoxSlot::UScrollBoxSlot(const FPostConstructInitializeProperties& PCIP)
 	: Super(PCIP)
-	, Slot(NULL)
+	, Slot(nullptr)
 {
 	HorizontalAlignment = HAlign_Fill;
 }
@@ -18,7 +18,7 @@ void UScrollBoxSlot::BuildSlot(TSharedRef<SScrollBox> ScrollBox)
 		.Padding(Padding)
 		.HAlign(HorizontalAlignment)
 		[
-			Content == NULL ? SNullWidget::NullWidget : Content->TakeWidget()
+			Content == nullptr ? SNullWidget::NullWidget : Content->TakeWidget()
 		];
 }
 
@@ -49,5 +49,5 @@ void UScrollBoxSlot::SyncronizeProperties()
 void UScrollBoxSlot::ReleaseNativeWidget()
 {
 	Super::ReleaseNativeWidget();
-	Slot = NULL;
-}
\ No newline at end of file
+	Slot = nullptr;
+}
