UMG - No longer locking the cursor to the viewport widget.  Modal widgets in the viewport now block mouse and keyboard while up.  Cut has been added to the designer, copy / paste now preserve / restore the old name if no other widget with that name exists.  Fixed a crash with undoing (causing a delete) of the currently selected widget.

diff --git a/Engine/Source/Editor/UMGEditor/Private/Hierarchy/SHierarchyViewItem.cpp b/Engine/Source/Editor/UMGEditor/Private/Hierarchy/SHierarchyViewItem.cpp
--- a/Engine/Source/Editor/UMGEditor/Private/Hierarchy/SHierarchyViewItem.cpp
+++ b/Engine/Source/Editor/UMGEditor/Private/Hierarchy/SHierarchyViewItem.cpp
@@ -164,21 +164,29 @@ void SHierarchyViewItem::OnNameTextCommited(const FText& InText, ETextCommit::Ty
 
 FSlateFontInfo SHierarchyViewItem::GetItemFont() const
 {
-	if ( !Item.GetTemplate()->IsGeneratedName() && Item.GetTemplate()->bIsVariable )
+	UWidget* WidgetTemplate = Item.GetTemplate();
+	if ( WidgetTemplate )
 	{
-		// TODO UMG Hacky move into style area
-		return FSlateFontInfo(FPaths::EngineContentDir() / TEXT("Slate/Fonts/Roboto-Bold.ttf"), 10);
-	}
-	else
-	{
-		static FName NormalFont("NormalFont");
-		return FCoreStyle::Get().GetFontStyle(NormalFont);
+		if ( !WidgetTemplate->IsGeneratedName() && WidgetTemplate->bIsVariable )
+		{
+			// TODO UMG Hacky move into style area
+			return FSlateFontInfo(FPaths::EngineContentDir() / TEXT("Slate/Fonts/Roboto-Bold.ttf"), 10);
+		}
 	}
+
+	static FName NormalFont("NormalFont");
+	return FCoreStyle::Get().GetFontStyle(NormalFont);
 }
 
 FText SHierarchyViewItem::GetItemText() const
 {
-	return FText::FromString(Item.GetTemplate()->GetLabel());
+	UWidget* WidgetTemplate = Item.GetTemplate();
+	if ( WidgetTemplate )
+	{
+		return FText::FromString(WidgetTemplate->GetLabel());
+	}
+
+	return FText::GetEmpty();
 }
 
 FReply SHierarchyViewItem::HandleDragDetected(const FGeometry& MyGeometry, const FPointerEvent& MouseEvent)
diff --git a/Engine/Source/Editor/UMGEditor/Private/WidgetBlueprintEditor.cpp b/Engine/Source/Editor/UMGEditor/Private/WidgetBlueprintEditor.cpp
--- a/Engine/Source/Editor/UMGEditor/Private/WidgetBlueprintEditor.cpp
+++ b/Engine/Source/Editor/UMGEditor/Private/WidgetBlueprintEditor.cpp
@@ -70,10 +70,10 @@ void FWidgetBlueprintEditor::InitWidgetBlueprintEditor(const EToolkitMode::Type
 		FCanExecuteAction::CreateSP(this, &FWidgetBlueprintEditor::CanCopySelectedWidgets)
 		);
 
-	//WidgetCommandList->MapAction(FGenericCommands::Get().Cut,
-	//	FExecuteAction::CreateSP(this, &FWidgetBlueprintEditor::CutSelectedNodes),
-	//	FCanExecuteAction::CreateSP(this, &FWidgetBlueprintEditor::CanCutNodes)
-	//	);
+	WidgetCommandList->MapAction(FGenericCommands::Get().Cut,
+		FExecuteAction::CreateSP(this, &FWidgetBlueprintEditor::CutSelectedWidgets),
+		FCanExecuteAction::CreateSP(this, &FWidgetBlueprintEditor::CanCutSelectedWidgets)
+		);
 
 	WidgetCommandList->MapAction(FGenericCommands::Get().Paste,
 		FExecuteAction::CreateSP(this, &FWidgetBlueprintEditor::PasteWidgets),
@@ -237,6 +237,18 @@ void FWidgetBlueprintEditor::CopySelectedWidgets()
 	FWidgetBlueprintEditorUtils::CopyWidgets(GetWidgetBlueprintObj(), Widgets);
 }
 
+bool FWidgetBlueprintEditor::CanCutSelectedWidgets()
+{
+	TSet<FWidgetReference> Widgets = GetSelectedWidgets();
+	return Widgets.Num() > 0;
+}
+
+void FWidgetBlueprintEditor::CutSelectedWidgets()
+{
+	TSet<FWidgetReference> Widgets = GetSelectedWidgets();
+	FWidgetBlueprintEditorUtils::CutWidgets(GetWidgetBlueprintObj(), Widgets);
+}
+
 bool FWidgetBlueprintEditor::CanPasteWidgets()
 {
 	TSet<FWidgetReference> Widgets = GetSelectedWidgets();
diff --git a/Engine/Source/Editor/UMGEditor/Private/WidgetBlueprintEditor.h b/Engine/Source/Editor/UMGEditor/Private/WidgetBlueprintEditor.h
--- a/Engine/Source/Editor/UMGEditor/Private/WidgetBlueprintEditor.h
+++ b/Engine/Source/Editor/UMGEditor/Private/WidgetBlueprintEditor.h
@@ -86,6 +86,9 @@ private:
 	bool CanPasteWidgets();
 	void PasteWidgets();
 
+	bool CanCutSelectedWidgets();
+	void CutSelectedWidgets();
+
 private:
 	/** Called whenever the blueprint is structurally changed. */
 	virtual void OnBlueprintChanged(UBlueprint* InBlueprint) override;
diff --git a/Engine/Source/Editor/UMGEditor/Private/WidgetBlueprintEditorUtils.cpp b/Engine/Source/Editor/UMGEditor/Private/WidgetBlueprintEditorUtils.cpp
--- a/Engine/Source/Editor/UMGEditor/Private/WidgetBlueprintEditorUtils.cpp
+++ b/Engine/Source/Editor/UMGEditor/Private/WidgetBlueprintEditorUtils.cpp
@@ -178,12 +178,17 @@ void FWidgetBlueprintEditorUtils::CreateWidgetContextMenu(FMenuBuilder& MenuBuil
 
 	MenuBuilder.PushCommandList(BlueprintEditor->WidgetCommandList.ToSharedRef());
 
-	MenuBuilder.BeginSection("Actions");
+	MenuBuilder.BeginSection("Edit", LOCTEXT("Edit", "Edit"));
 	{
+		MenuBuilder.AddMenuEntry(FGenericCommands::Get().Cut);
 		MenuBuilder.AddMenuEntry(FGenericCommands::Get().Copy);
 		MenuBuilder.AddMenuEntry(FGenericCommands::Get().Paste);
 		MenuBuilder.AddMenuEntry(FGenericCommands::Get().Delete);
+	}
+	MenuBuilder.EndSection();
 
+	MenuBuilder.BeginSection("Actions");
+	{
 		MenuBuilder.AddSubMenu(
 			LOCTEXT("WidgetTree_WrapWith", "Wrap With..."),
 			LOCTEXT("WidgetTree_WrapWithToolTip", "Wraps the currently selected widgets inside of another container widget"),
@@ -206,7 +211,30 @@ void FWidgetBlueprintEditorUtils::DeleteWidgets(UWidgetBlueprint* BP, TSet<FWidg
 		bool bRemoved = false;
 		for ( FWidgetReference& Item : Widgets )
 		{
-			bRemoved = BP->WidgetTree->RemoveWidget(Item.GetTemplate());
+			UWidget* WidgetTemplate = Item.GetTemplate();
+
+			// Modify the widget's parent
+			UPanelWidget* Parent = WidgetTemplate->GetParent();
+			if ( Parent )
+			{
+				Parent->Modify();
+			}
+			
+			// Modify the widget being removed.
+			WidgetTemplate->Modify();
+
+			bRemoved = BP->WidgetTree->RemoveWidget(WidgetTemplate);
+
+			// Rename the removed widget to the transient package so that it doesn't conflict with future widgets sharing the same name.
+			WidgetTemplate->Rename(NULL, NULL);
+
+			// Rename all child widgets as well, to the transient package so that they don't conflict with future widgets sharing the same name.
+			TArray<UWidget*> ChildWidgets;
+			BP->WidgetTree->GetChildWidgets(WidgetTemplate, ChildWidgets);
+			for ( UWidget* Widget : ChildWidgets )
+			{
+				Widget->Rename(NULL, NULL);
+			}
 		}
 
 		//TODO UMG There needs to be an event for widget removal so that caches can be updated, and selection
@@ -271,6 +299,12 @@ void FWidgetBlueprintEditorUtils::WrapWidgets(UWidgetBlueprint* BP, TSet<FWidget
 	FBlueprintEditorUtils::MarkBlueprintAsStructurallyModified(BP);
 }
 
+void FWidgetBlueprintEditorUtils::CutWidgets(UWidgetBlueprint* BP, TSet<FWidgetReference> Widgets)
+{
+	CopyWidgets(BP, Widgets);
+	DeleteWidgets(BP, Widgets);
+}
+
 void FWidgetBlueprintEditorUtils::CopyWidgets(UWidgetBlueprint* BP, TSet<FWidgetReference> Widgets)
 {
 	TSet<UWidget*> CopyableWidets;
@@ -430,7 +464,16 @@ void FWidgetBlueprintEditorUtils::ImportWidgetsFromText(UWidgetBlueprint* BP, co
 		ImportedWidgetSet.Add(Widget);
 
 		Widget->SetFlags(RF_Transactional);
-		Widget->Rename(nullptr, BP->WidgetTree);
+
+		// If there is an existing widget with the same name, rename the newly placed widget.
+		if ( FindObject<UObject>(BP->WidgetTree, *Widget->GetName()) )
+		{
+			Widget->Rename(nullptr, BP->WidgetTree);
+		}
+		else
+		{
+			Widget->Rename(*Widget->GetName(), BP->WidgetTree);
+		}
 	}
 
 	// Remove the temp package from the root now that it has served its purpose.
diff --git a/Engine/Source/Editor/UMGEditor/Private/WidgetBlueprintEditorUtils.h b/Engine/Source/Editor/UMGEditor/Private/WidgetBlueprintEditorUtils.h
--- a/Engine/Source/Editor/UMGEditor/Private/WidgetBlueprintEditorUtils.h
+++ b/Engine/Source/Editor/UMGEditor/Private/WidgetBlueprintEditorUtils.h
@@ -18,6 +18,8 @@ public:
 
 	static void DeleteWidgets(UWidgetBlueprint* BP, TSet<FWidgetReference> Widgets);
 
+	static void CutWidgets(UWidgetBlueprint* BP, TSet<FWidgetReference> Widgets);
+
 public:
 	static void ExportWidgetsToText(TSet<UWidget*> WidgetsToExport, /*out*/ FString& ExportedText);
 
diff --git a/Engine/Source/Runtime/UMG/Private/UserWidget.cpp b/Engine/Source/Runtime/UMG/Private/UserWidget.cpp
--- a/Engine/Source/Runtime/UMG/Private/UserWidget.cpp
+++ b/Engine/Source/Runtime/UMG/Private/UserWidget.cpp
@@ -28,6 +28,9 @@ TSharedPtr<SWidget> FUMGDragDropOp::GetDefaultDecorator() const
 	return DecoratorWidget;
 }
 
+/**
+ * This class holds onto the widget when it's placed into the viewport.
+ */
 class SViewportWidgetHost : public SCompoundWidget
 {
 	SLATE_BEGIN_ARGS(SViewportWidgetHost)
@@ -44,9 +47,9 @@ class SViewportWidgetHost : public SCompoundWidget
 		bModal = bInModal;
 
 		ChildSlot
-			[
-				InArgs._Content.Widget
-			];
+		[
+			InArgs._Content.Widget
+		];
 	}
 
 	virtual bool SupportsKeyboardFocus() const override
@@ -54,12 +57,12 @@ class SViewportWidgetHost : public SCompoundWidget
 		return bModal ? true : false;
 	}
 
-	FReply OnKeyboardFocusReceived(const FGeometry& MyGeometry, const FKeyboardFocusEvent& InKeyboardFocusEvent)
+	FReply OnKeyboardFocusReceived(const FGeometry& MyGeometry, const FKeyboardFocusEvent& InKeyboardFocusEvent) override
 	{
 		// if we support focus, release the focus captures and lock the focus to this widget 
 		if ( SupportsKeyboardFocus() )
 		{
-			return FReply::Handled().ReleaseMouseCapture().ReleaseJoystickCapture().LockMouseToWidget(SharedThis(this));
+			return FReply::Handled().ReleaseMouseCapture().ReleaseJoystickCapture();// .LockMouseToWidget(SharedThis(this));
 		}
 		else
 		{
@@ -67,7 +70,7 @@ class SViewportWidgetHost : public SCompoundWidget
 		}
 	}
 
-	virtual FCursorReply OnCursorQuery(const FGeometry& MyGeometry, const FPointerEvent& CursorEvent) const
+	virtual FCursorReply OnCursorQuery(const FGeometry& MyGeometry, const FPointerEvent& CursorEvent) const override
 	{
 		// If we support focus, show the default mouse cursor 
 		if ( SupportsKeyboardFocus() )
@@ -80,6 +83,56 @@ class SViewportWidgetHost : public SCompoundWidget
 		}
 	}
 
+	virtual FReply OnKeyChar(const FGeometry& MyGeometry, const FCharacterEvent& InCharacterEvent) override
+	{
+		return bModal ? FReply::Handled() : SCompoundWidget::OnKeyChar(MyGeometry, InCharacterEvent);
+	}
+			   
+	virtual FReply OnKeyDown(const FGeometry& MyGeometry, const FKeyboardEvent& InKeyboardEvent) override
+	{
+		return bModal ? FReply::Handled() : SCompoundWidget::OnKeyDown(MyGeometry, InKeyboardEvent);
+	}
+		   
+	virtual FReply OnKeyUp(const FGeometry& MyGeometry, const FKeyboardEvent& InKeyboardEvent) override
+	{
+		return bModal ? FReply::Handled() : SCompoundWidget::OnKeyUp(MyGeometry, InKeyboardEvent);
+	}
+		   
+	virtual FReply OnMouseButtonDown(const FGeometry& MyGeometry, const FPointerEvent& MouseEvent) override
+	{
+		return bModal ? FReply::Handled().CaptureMouse(AsShared()) : SCompoundWidget::OnMouseButtonDown(MyGeometry, MouseEvent);
+	}
+		   
+	virtual FReply OnMouseButtonUp(const FGeometry& MyGeometry, const FPointerEvent& MouseEvent) override
+	{
+		return bModal ? FReply::Handled().ReleaseMouseCapture() : SCompoundWidget::OnMouseButtonUp(MyGeometry, MouseEvent);
+	}
+
+	virtual FReply OnMouseButtonDoubleClick(const FGeometry& MyGeometry, const FPointerEvent& MouseEvent) override
+	{
+		return bModal ? FReply::Handled() : SCompoundWidget::OnMouseButtonDoubleClick(MyGeometry, MouseEvent);
+	}
+
+	virtual FReply OnMouseMove(const FGeometry& MyGeometry, const FPointerEvent& MouseEvent) override
+	{
+		return bModal ? FReply::Handled() : SCompoundWidget::OnMouseMove(MyGeometry, MouseEvent);
+	}
+
+	virtual void OnMouseEnter(const FGeometry& MyGeometry, const FPointerEvent& MouseEvent) override
+	{
+		return bModal ? FReply::Handled() : SCompoundWidget::OnMouseEnter(MyGeometry, MouseEvent);
+	}
+
+	virtual void OnMouseLeave(const FPointerEvent& MouseEvent) override
+	{
+		return bModal ? FReply::Handled() : SCompoundWidget::OnMouseLeave(MouseEvent);
+	}
+
+	virtual FReply OnMouseWheel(const FGeometry& MyGeometry, const FPointerEvent& MouseEvent) override
+	{
+		return bModal ? FReply::Handled() : SCompoundWidget::OnMouseWheel(MyGeometry, MouseEvent);
+	}
+
 protected:
 	bool bModal;
 };
diff --git a/Engine/Source/Runtime/UMG/Private/WidgetTree.cpp b/Engine/Source/Runtime/UMG/Private/WidgetTree.cpp
--- a/Engine/Source/Runtime/UMG/Private/WidgetTree.cpp
+++ b/Engine/Source/Runtime/UMG/Private/WidgetTree.cpp
@@ -66,22 +66,18 @@ UPanelWidget* UWidgetTree::FindWidgetParent(UWidget* Widget, int32& OutChildInde
 
 bool UWidgetTree::RemoveWidget(UWidget* InRemovedWidget)
 {
-	InRemovedWidget->Modify();
-
 	bool bRemoved = false;
 
 	UPanelWidget* InRemovedWidgetParent = InRemovedWidget->GetParent();
 	if ( InRemovedWidgetParent )
 	{
-		InRemovedWidgetParent->Modify();
-
 		if ( InRemovedWidgetParent->RemoveChild(InRemovedWidget) )
 		{
 			bRemoved = true;
 		}
 	}
-
-	if ( InRemovedWidget == RootWidget )
+	// If the widget being removed is the root, null it out.
+	else if ( InRemovedWidget == RootWidget )
 	{
 		RootWidget = NULL;
 		bRemoved = true;
diff --git a/Engine/Source/Runtime/UMG/Public/Blueprint/WidgetTree.h b/Engine/Source/Runtime/UMG/Public/Blueprint/WidgetTree.h
--- a/Engine/Source/Runtime/UMG/Public/Blueprint/WidgetTree.h
+++ b/Engine/Source/Runtime/UMG/Public/Blueprint/WidgetTree.h
@@ -27,7 +27,7 @@ public:
 	/** Gathers all the widgets in the tree recursively */
 	void GetAllWidgets(TArray<UWidget*>& Widgets) const;
 
-	/** Gathers only the immediate child widgets of a parent widget. */
+	/** Gathers descendant child widgets of a parent widget. */
 	void GetChildWidgets(UWidget* Parent, TArray<UWidget*>& Widgets) const;
 
 	/** Constructs the widget, and adds it to the tree. */
