Build Fix - Fixing mac editor build, returning a value on void.

diff --git a/Engine/Source/Runtime/UMG/Private/UserWidget.cpp b/Engine/Source/Runtime/UMG/Private/UserWidget.cpp
--- a/Engine/Source/Runtime/UMG/Private/UserWidget.cpp
+++ b/Engine/Source/Runtime/UMG/Private/UserWidget.cpp
@@ -118,16 +118,6 @@ class SViewportWidgetHost : public SCompoundWidget
 		return bModal ? FReply::Handled() : SCompoundWidget::OnMouseMove(MyGeometry, MouseEvent);
 	}
 
-	virtual void OnMouseEnter(const FGeometry& MyGeometry, const FPointerEvent& MouseEvent) override
-	{
-		return bModal ? FReply::Handled() : SCompoundWidget::OnMouseEnter(MyGeometry, MouseEvent);
-	}
-
-	virtual void OnMouseLeave(const FPointerEvent& MouseEvent) override
-	{
-		return bModal ? FReply::Handled() : SCompoundWidget::OnMouseLeave(MouseEvent);
-	}
-
 	virtual FReply OnMouseWheel(const FGeometry& MyGeometry, const FPointerEvent& MouseEvent) override
 	{
 		return bModal ? FReply::Handled() : SCompoundWidget::OnMouseWheel(MyGeometry, MouseEvent);
