UMG - General cleanup, comments, dead code removal and fixing a bug in checkbox that prevent the user from requesting the current checkbox state.  Buttons now use the standard Blueprint event system.

diff --git a/Engine/Source/Runtime/UMG/Private/Components/Button.cpp b/Engine/Source/Runtime/UMG/Private/Components/Button.cpp
--- a/Engine/Source/Runtime/UMG/Private/Components/Button.cpp
+++ b/Engine/Source/Runtime/UMG/Private/Components/Button.cpp
@@ -19,7 +19,10 @@ UButton::UButton(const FPostConstructInitializeProperties& PCIP)
 
 	ColorAndOpacity = FLinearColor::White;
 	BackgroundColor = FLinearColor::White;
-	ForegroundColor = FLinearColor::Black;
+	//ForegroundColor = FLinearColor::Black;
+
+	ClickMethod = EButtonClickMethod::DownAndUp;
+	TouchMethod = EButtonTouchMethod::DownAndUp;
 }
 
 void UButton::ReleaseNativeWidget()
@@ -31,7 +34,9 @@ void UButton::ReleaseNativeWidget()
 
 TSharedRef<SWidget> UButton::RebuildWidget()
 {
-	MyButton = SNew(SButton);
+	MyButton = SNew(SButton)
+		.ClickMethod(ClickMethod)
+		.TouchMethod(TouchMethod);
 
 	if ( GetChildrenCount() > 0 )
 	{
@@ -67,7 +72,6 @@ void UButton::SyncronizeProperties()
 
 	MyButton->SetButtonStyle( StylePtr );
 
-	MyButton->SetForegroundColor( ForegroundColor );
 	MyButton->SetColorAndOpacity( ColorAndOpacity );
 	MyButton->SetBorderBackgroundColor( BackgroundColor );
 	
@@ -76,7 +80,7 @@ void UButton::SyncronizeProperties()
 	MyButton->SetPressedSound( OptionalPressedSound );
 	MyButton->SetHoveredSound( OptionalHoveredSound );
 
-	MyButton->SetOnClicked(BIND_UOBJECT_DELEGATE(FOnClicked, HandleOnClicked));
+	MyButton->SetOnClicked(BIND_UOBJECT_DELEGATE(FOnClicked, SlateHandleClicked));
 }
 
 UClass* UButton::GetSlotClass() const
@@ -120,25 +124,6 @@ void UButton::SetBackgroundColor(FLinearColor Color)
 	}
 }
 
-void UButton::SetForegroundColor(FLinearColor InForegroundColor)
-{
-	ForegroundColor = InForegroundColor;
-	if ( MyButton.IsValid() )
-	{
-		MyButton->SetForegroundColor(InForegroundColor);
-	}
-}
-
-FReply UButton::HandleOnClicked()
-{
-	if ( OnClickedEvent.IsBound() && MyButton.IsValid() )
-	{
-		return OnClickedEvent.Execute().ToReply( MyButton.ToSharedRef() );
-	}
-	
-	return FReply::Unhandled();
-}
-
 bool UButton::IsPressed() const
 {
 	return MyButton->IsPressed();
@@ -165,6 +150,13 @@ void UButton::PostLoad()
 	}
 }
 
+FReply UButton::SlateHandleClicked()
+{
+	OnClicked.Broadcast();
+
+	return FReply::Handled();
+}
+
 #if WITH_EDITOR
 
 const FSlateBrush* UButton::GetEditorIcon()
diff --git a/Engine/Source/Runtime/UMG/Private/Components/VerticalBox.cpp b/Engine/Source/Runtime/UMG/Private/Components/VerticalBox.cpp
--- a/Engine/Source/Runtime/UMG/Private/Components/VerticalBox.cpp
+++ b/Engine/Source/Runtime/UMG/Private/Components/VerticalBox.cpp
@@ -48,7 +48,7 @@ void UVerticalBox::OnSlotRemoved(UPanelSlot* Slot)
 	}
 }
 
-UVerticalBoxSlot* UVerticalBox::Add(UWidget* Content)
+UVerticalBoxSlot* UVerticalBox::AddChildToVerticalBox(UWidget* Content)
 {
 	return Cast<UVerticalBoxSlot>( Super::AddChild(Content) );
 }
diff --git a/Engine/Source/Runtime/UMG/Private/Components/WidgetSwitcher.cpp b/Engine/Source/Runtime/UMG/Private/Components/WidgetSwitcher.cpp
--- a/Engine/Source/Runtime/UMG/Private/Components/WidgetSwitcher.cpp
+++ b/Engine/Source/Runtime/UMG/Private/Components/WidgetSwitcher.cpp
@@ -49,7 +49,18 @@ void UWidgetSwitcher::SetActiveWidgetIndex(int32 Index)
 	if ( MyWidgetSwitcher.IsValid() )
 	{
 		// Ensure the index is clamped to a valid range.
-		int32 SafeIndex = FMath::Clamp(Index, 0, FMath::Max(0, Slots.Num() - 1));
+		int32 SafeIndex = FMath::Clamp(ActiveWidgetIndex, 0, FMath::Max(0, Slots.Num() - 1));
+		MyWidgetSwitcher->SetActiveWidgetIndex(SafeIndex);
+	}
+}
+
+void UWidgetSwitcher::SetActiveWidget(UWidget* Widget)
+{
+	ActiveWidgetIndex = GetChildIndex(Widget);
+	if ( MyWidgetSwitcher.IsValid() )
+	{
+		// Ensure the index is clamped to a valid range.
+		int32 SafeIndex = FMath::Clamp(ActiveWidgetIndex, 0, FMath::Max(0, Slots.Num() - 1));
 		MyWidgetSwitcher->SetActiveWidgetIndex(SafeIndex);
 	}
 }
diff --git a/Engine/Source/Runtime/UMG/Public/Components/Border.h b/Engine/Source/Runtime/UMG/Public/Components/Border.h
--- a/Engine/Source/Runtime/UMG/Public/Components/Border.h
+++ b/Engine/Source/Runtime/UMG/Public/Components/Border.h
@@ -15,6 +15,8 @@ class UMG_API UBorder : public UContentWidget
 {
 	GENERATED_UCLASS_BODY()
 
+public:
+
 	/** Color and opacity multiplier of content in the border */
 	UPROPERTY(EditDefaultsOnly, Category="Content Layout")
 	FLinearColor ContentColorAndOpacity;
@@ -63,6 +65,8 @@ class UMG_API UBorder : public UContentWidget
 	UPROPERTY(EditDefaultsOnly, Category=Events)
 	FOnPointerEvent OnMouseDoubleClickEvent;
 
+public:
+
 	/**  */
 	UFUNCTION(BlueprintCallable, Category="Appearance")
 	void SetBrushColor(FLinearColor InColorAndOpacity);
@@ -75,12 +79,14 @@ class UMG_API UBorder : public UContentWidget
 	virtual void SyncronizeProperties() override;
 	// End of UWidget interface
 
+	// UVisual interface
+	virtual void ReleaseNativeWidget() override;
+	// End of UVisual interface
+
 	// Begin UObject interface
 	virtual void PostLoad() override;
 	// End of UObject interface
 
-	virtual void ReleaseNativeWidget() override;
-
 #if WITH_EDITOR
 	virtual const FSlateBrush* GetEditorIcon() override;
 #endif
diff --git a/Engine/Source/Runtime/UMG/Public/Components/Button.h b/Engine/Source/Runtime/UMG/Public/Components/Button.h
--- a/Engine/Source/Runtime/UMG/Public/Components/Button.h
+++ b/Engine/Source/Runtime/UMG/Public/Components/Button.h
@@ -6,7 +6,9 @@
 
 DECLARE_DYNAMIC_MULTICAST_DELEGATE(FOnButtonClickedEvent);
 
-/** Buttons are clickable widgets */
+/**
+ * The button is a clickable primitive widget to enable basic interaction.
+ */
 UCLASS(meta=( Category="Common" ), ClassGroup=UserInterface)
 class UMG_API UButton : public UContentWidget
 {
@@ -25,7 +27,7 @@ public:
 	UPROPERTY(EditDefaultsOnly, Category=Appearance, AdvancedDisplay)
 	FVector2D ContentScale;
 	
-	/** The color multiplier for the button images */
+	/** The color multiplier for the button content */
 	UPROPERTY(EditDefaultsOnly, Category=Appearance )
 	FLinearColor ColorAndOpacity;
 	
@@ -33,35 +35,37 @@ public:
 	UPROPERTY(EditDefaultsOnly, Category=Appearance )
 	FLinearColor BackgroundColor;
 
-	/** The foreground color of the button */
-	UPROPERTY(EditDefaultsOnly, Category=Appearance )
-	FLinearColor ForegroundColor;
+	/** The type of mouse action required by the user to trigger the buttons 'Click' */
+	UPROPERTY(EditDefaultsOnly, Category="Interaction")
+	TEnumAsByte<EButtonClickMethod::Type> ClickMethod;
+
+	/** The type of touch action required by the user to trigger the buttons 'Click' */
+	UPROPERTY(EditDefaultsOnly, Category="Interaction")
+	TEnumAsByte<EButtonTouchMethod::Type> TouchMethod;
 
+	/** The sound made when the user 'clicks' the button */
 	UPROPERTY(EditDefaultsOnly, Category=Sound)
 	FSlateSound PressedSound;
 
+	/** The sound made when the user hovers over the button */
 	UPROPERTY(EditDefaultsOnly, Category=Sound)
 	FSlateSound HoveredSound;
 
-	/** Called when the button is clicked */
-	UPROPERTY(EditDefaultsOnly, Category=Events)
-	FOnReply OnClickedEvent;
+public:
 
-	virtual void ReleaseNativeWidget() override;
+	/** Called when the button is clicked */
+	UPROPERTY(BlueprintAssignable)
+	FOnButtonClickedEvent OnClicked;
 	
-	/**  */
+	/** Sets the color multiplier for the button content */
 	UFUNCTION(BlueprintCallable, Category="Appearance")
 	void SetColorAndOpacity(FLinearColor InColorAndOpacity);
 
-	/**  */
+	/** Sets the color multiplier for the button background */
 	UFUNCTION(BlueprintCallable, Category="Appearance")
 	void SetBackgroundColor(FLinearColor InBackgroundColor);
 
-	/**  */
-	UFUNCTION(BlueprintCallable, Category="Appearance")
-	void SetForegroundColor(FLinearColor InForegroundColor);
-	
-	/**  */
+	/** @return true if the user is actively pressing the button otherwise false.  For detecting Clicks, use the OnClicked event. */
 	UFUNCTION(BlueprintCallable, Category="Widget")
 	bool IsPressed() const;
 
@@ -69,6 +73,10 @@ public:
 	virtual void SyncronizeProperties() override;
 	// End of UWidget interface
 
+	// UVisual interface
+	virtual void ReleaseNativeWidget() override;
+	// End of UVisual interface
+
 	// Begin UObject interface
 	virtual void PostLoad() override;
 	// End of UObject interface
@@ -85,13 +93,14 @@ protected:
 	virtual void OnSlotRemoved(UPanelSlot* Slot) override;
 	// End UPanelWidget
 
+protected:
+	FReply SlateHandleClicked();
+
 protected:
 	// UWidget interface
 	virtual TSharedRef<SWidget> RebuildWidget() override;
 	// End of UWidget interface
 
-	FReply HandleOnClicked();
-
 protected:
 	TSharedPtr<SButton> MyButton;
 };
diff --git a/Engine/Source/Runtime/UMG/Public/Components/ButtonSlot.h b/Engine/Source/Runtime/UMG/Public/Components/ButtonSlot.h
--- a/Engine/Source/Runtime/UMG/Public/Components/ButtonSlot.h
+++ b/Engine/Source/Runtime/UMG/Public/Components/ButtonSlot.h
@@ -11,6 +11,8 @@ UCLASS()
 class UMG_API UButtonSlot : public UPanelSlot
 {
 	GENERATED_UCLASS_BODY()
+
+public:
 	
 	/** The padding area between the slot and the content it contains. */
 	UPROPERTY(EditDefaultsOnly, Category=Layout)
diff --git a/Engine/Source/Runtime/UMG/Public/Components/CheckBox.h b/Engine/Source/Runtime/UMG/Public/Components/CheckBox.h
--- a/Engine/Source/Runtime/UMG/Public/Components/CheckBox.h
+++ b/Engine/Source/Runtime/UMG/Public/Components/CheckBox.h
@@ -77,11 +77,6 @@ protected:
 	UPROPERTY(EditDefaultsOnly, Category=Appearance)
 	FSlateColor BorderBackgroundColor;
 
-	//SLATE_ATTRIBUTE( bool, ReadOnly )
-	//SLATE_ARGUMENT( bool, IsFocusable )
-
-	//SLATE_EVENT( FOnGetContent, OnGetMenuContent )
-
 	UPROPERTY(EditDefaultsOnly, Category="Sound")
 	FSlateSound CheckedSound;
 
@@ -91,6 +86,8 @@ protected:
 	UPROPERTY(EditDefaultsOnly, Category="Sound")
 	FSlateSound HoveredSound;
 
+public:
+
 	/** Called when the checked state has changed */
 	UPROPERTY(BlueprintAssignable)
 	FOnCheckBoxComponentStateChanged OnCheckStateChanged;
diff --git a/Engine/Source/Runtime/UMG/Public/Components/EditableText.h b/Engine/Source/Runtime/UMG/Public/Components/EditableText.h
--- a/Engine/Source/Runtime/UMG/Public/Components/EditableText.h
+++ b/Engine/Source/Runtime/UMG/Public/Components/EditableText.h
@@ -85,6 +85,8 @@ public:
 	UPROPERTY(EditDefaultsOnly, Category=Behavior, AdvancedDisplay)
 	bool SelectAllTextOnCommit;
 
+public:
+
 	/** Called whenever the text is changed interactively by the user */
 	UPROPERTY(BlueprintAssignable, Category="Widget Event")
 	FOnEditableTextChangedEvent OnTextChanged;
diff --git a/Engine/Source/Runtime/UMG/Public/Components/EditableTextBox.h b/Engine/Source/Runtime/UMG/Public/Components/EditableTextBox.h
--- a/Engine/Source/Runtime/UMG/Public/Components/EditableTextBox.h
+++ b/Engine/Source/Runtime/UMG/Public/Components/EditableTextBox.h
@@ -92,6 +92,8 @@ public:
 	/** Provide a alternative mechanism for error reporting. */
 	//SLATE_ARGUMENT(TSharedPtr<class IErrorReportingWidget>, ErrorReporting)
 
+public:
+
 	/**  */
 	UFUNCTION(BlueprintCallable, Category="Widget")
 	FText GetText() const;
diff --git a/Engine/Source/Runtime/UMG/Public/Components/Image.h b/Engine/Source/Runtime/UMG/Public/Components/Image.h
--- a/Engine/Source/Runtime/UMG/Public/Components/Image.h
+++ b/Engine/Source/Runtime/UMG/Public/Components/Image.h
@@ -6,7 +6,9 @@
 
 class USlateBrushAsset;
 
-/** Image widget */
+/**
+ * The image widget allows you to display a Slate Brush, or texture or material in the UI.
+ */
 UCLASS(meta=( Category="Common" ), ClassGroup=UserInterface)
 class UMG_API UImage : public UWidget
 {
@@ -33,6 +35,8 @@ public:
 	UPROPERTY(EditDefaultsOnly, Category=Events)
 	FOnPointerEvent OnMouseButtonDownEvent;
 
+public:
+
 	/**  */
 	UFUNCTION(BlueprintCallable, Category="Appearance")
 	void SetColorAndOpacity(FLinearColor InColorAndOpacity);
diff --git a/Engine/Source/Runtime/UMG/Public/Components/MenuAnchor.h b/Engine/Source/Runtime/UMG/Public/Components/MenuAnchor.h
--- a/Engine/Source/Runtime/UMG/Public/Components/MenuAnchor.h
+++ b/Engine/Source/Runtime/UMG/Public/Components/MenuAnchor.h
@@ -27,6 +27,8 @@ public:
 	UPROPERTY(EditDefaultsOnly, Category=Menu)
 	TEnumAsByte<EMenuPlacement> Placement;
 
+public:
+
 	/**
 	 * Open or close the popup
 	 *
diff --git a/Engine/Source/Runtime/UMG/Public/Components/ProgressBar.h b/Engine/Source/Runtime/UMG/Public/Components/ProgressBar.h
--- a/Engine/Source/Runtime/UMG/Public/Components/ProgressBar.h
+++ b/Engine/Source/Runtime/UMG/Public/Components/ProgressBar.h
@@ -50,6 +50,8 @@ public:
 	/** Border Padding around fill bar */
 	UPROPERTY(EditDefaultsOnly, Category=Appearance)
 	FVector2D BorderPadding;
+
+public:
 	
 	/** Sets the current value of the ProgressBar. */
 	UFUNCTION(BlueprintCallable, Category="Behavior")
diff --git a/Engine/Source/Runtime/UMG/Public/Components/ScaleBox.h b/Engine/Source/Runtime/UMG/Public/Components/ScaleBox.h
--- a/Engine/Source/Runtime/UMG/Public/Components/ScaleBox.h
+++ b/Engine/Source/Runtime/UMG/Public/Components/ScaleBox.h
@@ -14,6 +14,8 @@ class UMG_API UScaleBox : public UContentWidget
 {
 	GENERATED_UCLASS_BODY()
 
+public:
+
 	/** Controls in what direction content can be scaled */
 	UPROPERTY(EditDefaultsOnly, Category="Stretching")
 	TEnumAsByte<EStretchDirection::Type> StretchDirection;
@@ -22,6 +24,8 @@ class UMG_API UScaleBox : public UContentWidget
 	UPROPERTY(EditDefaultsOnly, Category="Stretching")
 	TEnumAsByte<EStretch::Type> Stretch;
 
+public:
+
 	// UWidget interface
 	virtual void SyncronizeProperties() override;
 	// End of UWidget interface
diff --git a/Engine/Source/Runtime/UMG/Public/Components/ScaleBoxSlot.h b/Engine/Source/Runtime/UMG/Public/Components/ScaleBoxSlot.h
--- a/Engine/Source/Runtime/UMG/Public/Components/ScaleBoxSlot.h
+++ b/Engine/Source/Runtime/UMG/Public/Components/ScaleBoxSlot.h
@@ -11,6 +11,8 @@ UCLASS()
 class UMG_API UScaleBoxSlot : public UPanelSlot
 {
 	GENERATED_UCLASS_BODY()
+
+public:
 	
 	/** The padding area between the slot and the content it contains. */
 	UPROPERTY(EditDefaultsOnly, Category=Layout)
@@ -24,6 +26,8 @@ class UMG_API UScaleBoxSlot : public UPanelSlot
 	UPROPERTY(EditDefaultsOnly, Category=Layout)
 	TEnumAsByte<EVerticalAlignment> VerticalAlignment;
 
+public:
+
 	UFUNCTION(BlueprintCallable, Category="Appearance")
 	void SetPadding(FMargin InPadding);
 
diff --git a/Engine/Source/Runtime/UMG/Public/Components/ScrollBar.h b/Engine/Source/Runtime/UMG/Public/Components/ScrollBar.h
--- a/Engine/Source/Runtime/UMG/Public/Components/ScrollBar.h
+++ b/Engine/Source/Runtime/UMG/Public/Components/ScrollBar.h
@@ -10,6 +10,8 @@ class UMG_API UScrollBar : public UPanelWidget
 {
 	GENERATED_UCLASS_BODY()
 
+public:
+
 	/** Style of the scrollbar */
 	UPROPERTY(EditDefaultsOnly, Category="Style", meta=( DisplayThumbnail = "true" ))
 	USlateWidgetStyleAsset* Style;
@@ -26,6 +28,8 @@ class UMG_API UScrollBar : public UPanelWidget
 	UPROPERTY(EditDefaultsOnly, Category="Behavior")
 	FVector2D Thickness;
 
+public:
+
 	/**
 	* Set the offset and size of the track's thumb.
 	* Note that the maximum offset is 1.0-ThumbSizeFraction.
diff --git a/Engine/Source/Runtime/UMG/Public/Components/Spacer.h b/Engine/Source/Runtime/UMG/Public/Components/Spacer.h
--- a/Engine/Source/Runtime/UMG/Public/Components/Spacer.h
+++ b/Engine/Source/Runtime/UMG/Public/Components/Spacer.h
@@ -10,10 +10,14 @@ class UMG_API USpacer : public UWidget
 {
 	GENERATED_UCLASS_BODY()
 
+public:
+
 	/** The size of the spacer */
 	UPROPERTY(EditDefaultsOnly, Category=Appearance)
 	FVector2D Size;
 
+public:
+
 	/** Sets the size of the spacer */
 	UFUNCTION(BlueprintCallable, Category="Widget")
 	void SetSize(FVector2D InSize);
diff --git a/Engine/Source/Runtime/UMG/Public/Components/TileView.h b/Engine/Source/Runtime/UMG/Public/Components/TileView.h
--- a/Engine/Source/Runtime/UMG/Public/Components/TileView.h
+++ b/Engine/Source/Runtime/UMG/Public/Components/TileView.h
@@ -10,6 +10,8 @@ class UMG_API UTileView : public UTableViewBase
 {
 	GENERATED_UCLASS_BODY()
 
+public:
+
 	/**  */
 	UPROPERTY(EditDefaultsOnly, Category=Appearance)
 	float ItemWidth;
@@ -30,6 +32,8 @@ class UMG_API UTileView : public UTableViewBase
 	UPROPERTY(EditDefaultsOnly, Category=Events)
 	FOnGenerateRowUObject OnGenerateTileEvent;
 
+public:
+
 	/** Set item width */
 	UFUNCTION(BlueprintCallable, Category="Appearance")
 	void SetItemWidth(float Width);
diff --git a/Engine/Source/Runtime/UMG/Public/Components/UniformGridPanel.h b/Engine/Source/Runtime/UMG/Public/Components/UniformGridPanel.h
--- a/Engine/Source/Runtime/UMG/Public/Components/UniformGridPanel.h
+++ b/Engine/Source/Runtime/UMG/Public/Components/UniformGridPanel.h
@@ -10,6 +10,8 @@ class UMG_API UUniformGridPanel : public UPanelWidget
 {
 	GENERATED_UCLASS_BODY()
 
+public:
+
 	/** Padding given to each slot */
 	UPROPERTY(EditDefaultsOnly, Category=Layout)
 	FMargin SlotPadding;
@@ -22,6 +24,8 @@ class UMG_API UUniformGridPanel : public UPanelWidget
 	UPROPERTY(EditDefaultsOnly, Category=Layout)
 	float MinDesiredSlotHeight;
 
+public:
+
 	// UWidget interface
 	virtual void SyncronizeProperties() override;
 	// End of UWidget interface
diff --git a/Engine/Source/Runtime/UMG/Public/Components/VerticalBox.h b/Engine/Source/Runtime/UMG/Public/Components/VerticalBox.h
--- a/Engine/Source/Runtime/UMG/Public/Components/VerticalBox.h
+++ b/Engine/Source/Runtime/UMG/Public/Components/VerticalBox.h
@@ -13,9 +13,11 @@ class UMG_API UVerticalBox : public UPanelWidget
 {
 	GENERATED_UCLASS_BODY()
 
+public:
+
 	/**  */
 	UFUNCTION(BlueprintCallable, Category="Widget")
-	UVerticalBoxSlot* Add(UWidget* Content);
+	UVerticalBoxSlot* AddChildToVerticalBox(UWidget* Content);
 
 #if WITH_EDITOR
 	// UWidget interface
diff --git a/Engine/Source/Runtime/UMG/Public/Components/VerticalBoxSlot.h b/Engine/Source/Runtime/UMG/Public/Components/VerticalBoxSlot.h
--- a/Engine/Source/Runtime/UMG/Public/Components/VerticalBoxSlot.h
+++ b/Engine/Source/Runtime/UMG/Public/Components/VerticalBoxSlot.h
@@ -11,6 +11,8 @@ UCLASS()
 class UMG_API UVerticalBoxSlot : public UPanelSlot
 {
 	GENERATED_UCLASS_BODY()
+
+public:
 	
 	/** The padding area between the slot and the content it contains. */
 	UPROPERTY(EditDefaultsOnly, Category=Layout)
@@ -28,6 +30,8 @@ class UMG_API UVerticalBoxSlot : public UPanelSlot
 	UPROPERTY(EditDefaultsOnly, Category=Layout)
 	TEnumAsByte<EVerticalAlignment> VerticalAlignment;
 
+public:
+
 	UFUNCTION(BlueprintCallable, Category="Appearance")
 	void SetPadding(FMargin InPadding);
 
diff --git a/Engine/Source/Runtime/UMG/Public/Components/WidgetSwitcher.h b/Engine/Source/Runtime/UMG/Public/Components/WidgetSwitcher.h
--- a/Engine/Source/Runtime/UMG/Public/Components/WidgetSwitcher.h
+++ b/Engine/Source/Runtime/UMG/Public/Components/WidgetSwitcher.h
@@ -13,10 +13,13 @@ class UMG_API UWidgetSwitcher : public UPanelWidget
 {
 	GENERATED_UCLASS_BODY()
 
-	/** Image to draw */
+public:
+	/** The slot index to display */
 	UPROPERTY(EditDefaultsOnly, Category="Switcher", meta=( UIMin=0, ClampMin=0 ))
 	int32 ActiveWidgetIndex;
 
+public:
+
 	/** Gets the number of widgets that this switcher manages. */
 	UFUNCTION(BlueprintCallable, Category="Switcher")
 	int32 GetNumWidgets() const;
@@ -28,6 +31,10 @@ class UMG_API UWidgetSwitcher : public UPanelWidget
 	/** Activates the widget at the specified index. */
 	UFUNCTION(BlueprintCallable, Category="Switcher")
 	void SetActiveWidgetIndex( int32 Index );
+
+	/** Activates the widget and makes it the active index. */
+	UFUNCTION(BlueprintCallable, Category="Switcher")
+	void SetActiveWidget(UWidget* Widget);
 	
 	virtual void SyncronizeProperties() override;
 
