UMG - Fixing a lot of issues with how widgets are initialized.  Fixing issues with sub widgets being nulled out, as well as delay not working on sub widgets.

diff --git a/Engine/Source/Editor/UMGEditor/Private/WidgetBlueprintEditor.cpp b/Engine/Source/Editor/UMGEditor/Private/WidgetBlueprintEditor.cpp
--- a/Engine/Source/Editor/UMGEditor/Private/WidgetBlueprintEditor.cpp
+++ b/Engine/Source/Editor/UMGEditor/Private/WidgetBlueprintEditor.cpp
@@ -511,7 +511,7 @@ void FWidgetBlueprintEditor::UpdatePreview(UBlueprint* InBlueprint, bool bInForc
 		// without needing to do a full recompile.
 		Cast<UWidgetBlueprintGeneratedClass>(PreviewBlueprint->GeneratedClass)->WidgetTree = DuplicateObject<UWidgetTree>(PreviewBlueprint->WidgetTree, PreviewBlueprint->GeneratedClass);
 
-		PreviewActor = ConstructObject<UUserWidget>(PreviewBlueprint->GeneratedClass, PreviewScene.GetWorld()->GetCurrentLevel());
+		PreviewActor = CreateWidget<UUserWidget>(PreviewScene.GetWorld(), PreviewBlueprint->GeneratedClass);
 		PreviewActor->SetFlags(RF_Transactional);
 		
 		// Configure all the widgets to be set to design time.
diff --git a/Engine/Source/Runtime/UMG/Private/UserWidget.cpp b/Engine/Source/Runtime/UMG/Private/UserWidget.cpp
--- a/Engine/Source/Runtime/UMG/Private/UserWidget.cpp
+++ b/Engine/Source/Runtime/UMG/Private/UserWidget.cpp
@@ -138,6 +138,34 @@ UUserWidget::UUserWidget(const FPostConstructInitializeProperties& PCIP)
 	VerticalAlignment = VAlign_Fill;
 
 	Visiblity = ESlateVisibility::SelfHitTestInvisible;
+
+	bInitialized = false;
+	CachedWorld = NULL;
+}
+
+void UUserWidget::Initialize()
+{
+	if ( !bInitialized )
+	{
+		bInitialized = true;
+
+		// Only do this if this widget is of a blueprint class
+		UWidgetBlueprintGeneratedClass* BGClass = Cast<UWidgetBlueprintGeneratedClass>(GetClass());
+		if ( BGClass != NULL )
+		{
+			BGClass->InitializeWidget(this);
+		}
+	}
+}
+
+void UUserWidget::PostEditImport()
+{
+	Initialize();
+}
+
+void UUserWidget::PostDuplicate(bool bDuplicateForPIE)
+{
+	Initialize();
 }
 
 void UUserWidget::ReleaseNativeWidget()
@@ -156,14 +184,6 @@ void UUserWidget::PostInitProperties()
 	Super::PostInitProperties();
 
 	Components.Reset();
-
-	// Only do this if this widget is of a blueprint class
-	UWidgetBlueprintGeneratedClass* BGClass = Cast<UWidgetBlueprintGeneratedClass>(GetClass());
-	if ( BGClass != NULL )
-	{
-		BGClass->InitializeWidget(this);
-	}
-
 	//TODO UMG For non-BP versions how do we generate the Components list?
 }
 
@@ -174,7 +194,22 @@ UWorld* UUserWidget::GetWorld() const
 		return PlayerContext.GetWorld();
 	}
 
-	return NULL;
+	if ( CachedWorld == NULL )
+	{
+		UObject* Outer = GetOuter();
+		while ( Outer != NULL )
+		{
+			CachedWorld = Outer->GetWorld();
+			if ( CachedWorld )
+			{
+				return CachedWorld;
+			}
+
+			Outer = Outer->GetOuter();
+		}
+	}
+	
+	return CachedWorld;
 }
 
 void UUserWidget::PlayAnimation(FName AnimationName)
@@ -239,6 +274,8 @@ TSharedRef<SWidget> UUserWidget::RebuildWidget()
 {
 	TSharedPtr<SWidget> UserRootWidget;
 
+	check(bInitialized);
+
 	// Add the first component to the root of the widget surface.
 	if ( Components.Num() > 0 && Components[0] != NULL )
 	{
@@ -296,7 +333,8 @@ void UUserWidget::NativeTick(const FGeometry& MyGeometry, float InDeltaTime )
 		if ( World )
 		{
 			// Update any latent actions we have for this actor
-			World->GetLatentActionManager().ProcessLatentActions(this, InDeltaTime);
+			FLatentActionManager& LatentActionManager = World->GetLatentActionManager();
+			LatentActionManager.ProcessLatentActions(this, InDeltaTime);
 		}
 	}
 
@@ -472,4 +510,4 @@ const FSlateBrush* UUserWidget::GetEditorIcon()
 	return FUMGStyle::Get().GetBrush("Widget.UserWidget");
 }
 
-#endif
\ No newline at end of file
+#endif
diff --git a/Engine/Source/Runtime/UMG/Private/WidgetBlueprintGeneratedClass.cpp b/Engine/Source/Runtime/UMG/Private/WidgetBlueprintGeneratedClass.cpp
--- a/Engine/Source/Runtime/UMG/Private/WidgetBlueprintGeneratedClass.cpp
+++ b/Engine/Source/Runtime/UMG/Private/WidgetBlueprintGeneratedClass.cpp
@@ -83,6 +83,8 @@ void UWidgetBlueprintGeneratedClass::InitializeWidget(UUserWidget* UserWidget) c
 			if ( Prop )
 			{
 				Prop->SetObjectPropertyValue_InContainer(UserWidget, Widget);
+				UObject* Value = Prop->GetObjectPropertyValue_InContainer(UserWidget);
+				check(Value == Widget);
 			}
 
 			// Perform binding
diff --git a/Engine/Source/Runtime/UMG/Private/WidgetBlueprintLibrary.cpp b/Engine/Source/Runtime/UMG/Private/WidgetBlueprintLibrary.cpp
--- a/Engine/Source/Runtime/UMG/Private/WidgetBlueprintLibrary.cpp
+++ b/Engine/Source/Runtime/UMG/Private/WidgetBlueprintLibrary.cpp
@@ -16,78 +16,17 @@ UWidgetBlueprintLibrary::UWidgetBlueprintLibrary(const FPostConstructInitializeP
 
 UUserWidget* UWidgetBlueprintLibrary::Create(UObject* WorldContextObject, TSubclassOf<UUserWidget> WidgetType, APlayerController* OwningPlayer)
 {
-	if ( WidgetType == NULL || WidgetType->HasAnyClassFlags(CLASS_Abstract) )
-	{
-		// TODO UMG script Error?
-		return NULL;
-	}
-
-	UUserWidget* NewWidget = NULL;
-
 	if ( OwningPlayer == NULL )
 	{
 		UWorld* World = GEngine->GetWorldFromContextObject(WorldContextObject);
-		ULocalPlayer* Player = GEngine->GetGamePlayer(World, 0);
-
-		NewWidget = ConstructObject<UUserWidget>(WidgetType, Player);
-		NewWidget->SetPlayerContext(FLocalPlayerContext(Player));
+		return CreateWidget<UUserWidget>(World, WidgetType);
 	}
 	else
 	{
-		NewWidget = ConstructObject<UUserWidget>(WidgetType, OwningPlayer);
-		NewWidget->SetPlayerContext(FLocalPlayerContext(OwningPlayer));
+		return CreateWidget<UUserWidget>(OwningPlayer, WidgetType);
 	}
-
-	return NewWidget;
 }
 
-//void UWidgetBlueprintLibrary::Popup(UObject* WorldContextObject, UWidget* PopupWidget, int32 ZIndex)
-//{
-//	if ( !(WorldContextObject || PopupWidget) )
-//	{
-//		return;
-//	}
-//
-//	UWorld* World = GEngine->GetWorldFromContextObject(WorldContextObject);
-//	if ( World && World->IsGameWorld() )
-//	{
-//		TSharedRef<SWidget> SlatePopupWidget = PopupWidget->GetWidget();
-//
-//		UGameViewportClient* Viewport = World->GetGameViewport();
-//		Viewport->AddViewportWidgetContent(SlatePopupWidget, ZIndex);
-//
-//		TWeakPtr<SViewport> GameViewportWidget = Viewport->GetGameViewport()->GetViewportWidget();
-//		if ( GameViewportWidget.IsValid() )
-//		{
-//			GameViewportWidget.Pin()->SetWidgetToFocusOnActivate(SlatePopupWidget);
-//			FSlateApplication::Get().SetKeyboardFocus(SlatePopupWidget);
-//		}
-//
-//		
-//
-//		//TSharedPtr<SWeakWidget> PopupWrapperPtr;
-//		//PopupLayerSlot = &PopupWindow->AddPopupLayerSlot()
-//		//	.DesktopPosition(NewPosition)
-//		//	.WidthOverride(NewWindowSize.X)
-//		//	[
-//		//		SAssignNew(PopupWrapperPtr, SWeakWidget)
-//		//		.PossiblyNullContent
-//		//		(
-//		//		MenuContentRef
-//		//		)
-//		//	];
-//
-//		// We want to support dismissing the popup widget when the user clicks outside it.
-//		//FSlateApplication::Get().GetPopupSupport().RegisterClickNotification(PopupWrapperPtr.ToSharedRef(), FOnClickedOutside::CreateSP(this, &SMenuAnchor::OnClickedOutsidePopup));
-//		//}
-//
-//		//if ( bFocusMenu )
-//		//{
-//		//	FSlateApplication::Get().SetKeyboardFocus(MenuContentRef, EKeyboardFocusCause::SetDirectly);
-//		//}
-//	}
-//}
-
 void UWidgetBlueprintLibrary::DrawBox(UPARAM(ref) FPaintContext& Context, FVector2D Position, FVector2D Size, USlateBrushAsset* Brush, FLinearColor Tint)
 {
 	Context.MaxLayer++;
diff --git a/Engine/Source/Runtime/UMG/Public/Blueprint/UserWidget.h b/Engine/Source/Runtime/UMG/Public/Blueprint/UserWidget.h
--- a/Engine/Source/Runtime/UMG/Public/Blueprint/UserWidget.h
+++ b/Engine/Source/Runtime/UMG/Public/Blueprint/UserWidget.h
@@ -102,8 +102,12 @@ public:
 	//UObject interface
 	virtual void PostInitProperties() override;
 	virtual class UWorld* GetWorld() const override;
+	virtual void PostEditImport() override;
+	virtual void PostDuplicate(bool bDuplicateForPIE) override;
 	// End of UObject interface
 
+	void Initialize();
+
 	//UVisual interface
 	virtual void ReleaseNativeWidget() override;
 	// End of UVisual interface
@@ -339,4 +343,43 @@ private:
 	TWeakPtr<SWidget> FullScreenWidget;
 
 	FLocalPlayerContext PlayerContext;
+
+	bool bInitialized;
+
+	mutable UWorld* CachedWorld;
 };
+
+template< class T >
+T* CreateWidget(UWorld* World, UClass* UserWidgetClass)
+{
+	if ( !UserWidgetClass->IsChildOf(UUserWidget::StaticClass()) )
+	{
+		// TODO UMG Error?
+		return NULL;
+	}
+
+	ULocalPlayer* Player = World->GetFirstLocalPlayerFromController();
+
+	UObject* Outer = ( Player == NULL ) ? StaticCast<UObject*>(World) : StaticCast<UObject*>(Player);
+	UUserWidget* NewWidget = ConstructObject<UUserWidget>(UserWidgetClass, Outer);
+	NewWidget->SetPlayerContext(FLocalPlayerContext(Player));
+	NewWidget->Initialize();
+
+	return Cast<T>(NewWidget);
+}
+
+template< class T >
+T* CreateWidget(APlayerController* OwningPlayer, UClass* UserWidgetClass)
+{
+	if ( !UserWidgetClass->IsChildOf(UUserWidget::StaticClass()) )
+	{
+		// TODO UMG Error?
+		return NULL;
+	}
+
+	UUserWidget* NewWidget = ConstructObject<UUserWidget>(UserWidgetClass, OwningPlayer);
+	NewWidget->SetPlayerContext(FLocalPlayerContext(OwningPlayer));
+	NewWidget->Initialize();
+
+	return Cast<T>(NewWidget);
+}
diff --git a/Engine/Source/Runtime/UMG/Public/Blueprint/WidgetBlueprintLibrary.h b/Engine/Source/Runtime/UMG/Public/Blueprint/WidgetBlueprintLibrary.h
--- a/Engine/Source/Runtime/UMG/Public/Blueprint/WidgetBlueprintLibrary.h
+++ b/Engine/Source/Runtime/UMG/Public/Blueprint/WidgetBlueprintLibrary.h
@@ -12,10 +12,6 @@ class UWidgetBlueprintLibrary : public UBlueprintFunctionLibrary
 	UFUNCTION(BlueprintCallable, meta=( HidePin="WorldContextObject", DefaultToSelf="WorldContextObject", FriendlyName = "Create Widget" ), Category="User Interface|Widget")
 	static class UUserWidget* Create(UObject* WorldContextObject, TSubclassOf<class UUserWidget> WidgetType, APlayerController* OwningPlayer);
 
-	///** Show a widget as a popup */
-	//UFUNCTION(BlueprintCallable, meta=( HidePin="WorldContextObject", DefaultToSelf="WorldContextObject", UnsafeDuringActorConstruction = "true" ), Category="User Interface" )
-	//static void Popup(UObject* WorldContextObject, UWidget* PopupWidget, int32 ZIndex = 0);
-
 	/** Draws a box */
 	UFUNCTION(BlueprintCallable, Category="Painting")
 	static void DrawBox(UPARAM(ref) FPaintContext& Context, FVector2D Position, FVector2D Size, USlateBrushAsset* Brush, FLinearColor Tint = FLinearColor::White);
