UMG - Now correctly initializing UUserWidgets created in a widget tree during design time.

diff --git a/Engine/Source/Runtime/UMG/Public/Blueprint/WidgetTree.h b/Engine/Source/Runtime/UMG/Public/Blueprint/WidgetTree.h
--- a/Engine/Source/Runtime/UMG/Public/Blueprint/WidgetTree.h
+++ b/Engine/Source/Runtime/UMG/Public/Blueprint/WidgetTree.h
@@ -39,10 +39,19 @@ public:
 
 		// TODO UMG Don't have the widget tree responsible for construction and adding to the tree.
 
-		UWidget* Widget = (UWidget*)ConstructObject<UWidget>(WidgetType, this);
-		Widget->SetFlags(RF_Transactional);
-
-		return (T*)Widget;
+		if ( WidgetType->IsChildOf(UUserWidget::StaticClass()) )
+		{
+			UUserWidget* Widget = ConstructObject<UUserWidget>(WidgetType, this);
+			Widget->Initialize();
+			Widget->SetFlags(RF_Transactional);
+			return (T*)Widget;
+		}
+		else
+		{
+			UWidget* Widget = (UWidget*)ConstructObject<UWidget>(WidgetType, this);
+			Widget->SetFlags(RF_Transactional);
+			return (T*)Widget;
+		}
 	}
 
 	virtual void PreSave() override
