UMG - Fixing non-unity build.

diff --git a/Engine/Source/Runtime/UMG/Public/Blueprint/WidgetTree.h b/Engine/Source/Runtime/UMG/Public/Blueprint/WidgetTree.h
--- a/Engine/Source/Runtime/UMG/Public/Blueprint/WidgetTree.h
+++ b/Engine/Source/Runtime/UMG/Public/Blueprint/WidgetTree.h
@@ -1,7 +1,9 @@
 // Copyright 1998-2014 Epic Games, Inc. All Rights Reserved.
 
 #pragma once
-#include "Blueprint/UserWidget.h"
+
+#include "UserWidget.h"
+
 #include "WidgetTree.generated.h"
 
 /** The widget tree manages the collection of widgets in a blueprint widget. */
