UMG - Starting to prototype a solution for drag/drop in UMG.

diff --git a/Engine/Source/Runtime/UMG/Private/Slate/SDragPanel.cpp b/Engine/Source/Runtime/UMG/Private/Slate/SDragPanel.cpp
new file mode 100644
--- /dev/null
+++ b/Engine/Source/Runtime/UMG/Private/Slate/SDragPanel.cpp
@@ -0,0 +1,54 @@
+// Copyright 1998-2014 Epic Games, Inc. All Rights Reserved.
+
+#include "UMGPrivatePCH.h"
+
+#include "SDragPanel.h"
+#include "UMGDragDropOp.h"
+
+
+/* SDragPanel interface
+ *****************************************************************************/
+
+void SDragPanel::Construct( const SDragPanel::FArguments& InArgs )
+{
+	OnDragOperationRequested = InArgs._OnDragDetected;
+
+	ChildSlot
+	[
+		InArgs._Content.Widget
+	];
+}
+
+void SDragPanel::SetContent(TSharedRef<SWidget> InContent)
+{
+	ChildSlot
+	[
+		InContent
+	];
+}
+
+FReply SDragPanel::OnMouseButtonDown(const FGeometry& MyGeometry, const FPointerEvent& PointerEvent)
+{
+	// TODO UMG The button should be configurable
+	if ( PointerEvent.GetEffectingButton() == EKeys::LeftMouseButton )
+	{
+		return FReply::Handled().DetectDrag(SharedThis(this), EKeys::LeftMouseButton);
+	}
+	
+	return FReply::Unhandled();
+}
+
+FReply SDragPanel::OnDragDetected(const FGeometry& MyGeometry, const FPointerEvent& PointerEvent)
+{
+	if ( OnDragOperationRequested.IsBound() )
+	{
+		TSharedPtr<FDragDropOperation> DragOperation = OnDragOperationRequested.Execute(MyGeometry, PointerEvent);
+
+		if ( DragOperation.IsValid() )
+		{
+			return FReply::Handled().BeginDragDrop(DragOperation.ToSharedRef());
+		}
+	}
+
+	return FReply::Unhandled();
+}
\ No newline at end of file
diff --git a/Engine/Source/Runtime/UMG/Private/Slate/SDropPanel.cpp b/Engine/Source/Runtime/UMG/Private/Slate/SDropPanel.cpp
new file mode 100644
--- /dev/null
+++ b/Engine/Source/Runtime/UMG/Private/Slate/SDropPanel.cpp
@@ -0,0 +1,46 @@
+// Copyright 1998-2014 Epic Games, Inc. All Rights Reserved.
+
+#include "UMGPrivatePCH.h"
+
+#include "SDropPanel.h"
+#include "UMGDragDropOp.h"
+
+
+/* SDropPanel interface
+ *****************************************************************************/
+
+void SDropPanel::Construct( const SDropPanel::FArguments& InArgs )
+{
+	ChildSlot
+	[
+		InArgs._Content.Widget
+	];
+}
+
+void SDropPanel::SetContent(TSharedRef<SWidget> InContent)
+{
+	ChildSlot
+	[
+		InContent
+	];
+}
+
+void SDropPanel::OnDragEnter(const FGeometry& MyGeometry, const FDragDropEvent& DragDropEvent)
+{
+
+}
+
+void SDropPanel::OnDragLeave(const FDragDropEvent& DragDropEvent)
+{
+
+}
+
+FReply SDropPanel::OnDragOver(const FGeometry& MyGeometry, const FDragDropEvent& DragDropEvent)
+{
+	return FReply::Unhandled();
+}
+
+FReply SDropPanel::OnDrop(const FGeometry& MyGeometry, const FDragDropEvent& DragDropEvent)
+{
+	return FReply::Unhandled();
+}
diff --git a/Engine/Source/Runtime/UMG/Private/Slate/SObjectWidget.cpp b/Engine/Source/Runtime/UMG/Private/Slate/SObjectWidget.cpp
--- a/Engine/Source/Runtime/UMG/Private/Slate/SObjectWidget.cpp
+++ b/Engine/Source/Runtime/UMG/Private/Slate/SObjectWidget.cpp
@@ -132,15 +132,7 @@ FReply SObjectWidget::OnMouseButtonDown(const FGeometry& MyGeometry, const FPoin
 	TSharedPtr<SWidget> SlateWidget = WidgetObject->GetCachedWidget();
 	if ( SlateWidget.IsValid() )
 	{
-		FReply Reply = WidgetObject->OnMouseButtonDown(MyGeometry, MouseEvent).ToReply(SlateWidget.ToSharedRef());
-
-		//TODO UMG Figure out how to let the user more easily manage when drags start.
-		//if ( MouseEvent.GetEffectingButton() == EKeys::LeftMouseButton )
-		//{
-		//	return Reply.DetectDrag(SlateWidget.ToSharedRef(), EKeys::LeftMouseButton);
-		//}
-
-		return Reply;
+		return WidgetObject->OnMouseButtonDown(MyGeometry, MouseEvent).ToReply(SlateWidget.ToSharedRef());
 	}
 
 	return FReply::Unhandled();
diff --git a/Engine/Source/Runtime/UMG/Private/Slate/UMGDragDropOp.cpp b/Engine/Source/Runtime/UMG/Private/Slate/UMGDragDropOp.cpp
new file mode 100644
--- /dev/null
+++ b/Engine/Source/Runtime/UMG/Private/Slate/UMGDragDropOp.cpp
@@ -0,0 +1,54 @@
+// Copyright 1998-2014 Epic Games, Inc. All Rights Reserved.
+
+#include "UMGPrivatePCH.h"
+
+#include "UMGDragDropOp.h"
+
+//////////////////////////////////////////////////////////////////////////
+// FUMGDragDropOp
+void FUMGDragDropOp::OnDrop( bool bDropWasHandled, const FPointerEvent& MouseEvent )
+{
+	//if ( bDropWasHandled == false )
+	//{
+	//	if (OriginalTrackNode.IsValid())
+	//	{
+	//		//OriginalTrackNode.Pin()->OnDropCancelled(MouseEvent);
+	//	}
+	//}
+	
+	FDragDropOperation::OnDrop(bDropWasHandled, MouseEvent);
+}
+
+void FUMGDragDropOp::OnDragged( const class FDragDropEvent& DragDropEvent )
+{
+	//TSharedPtr<STrackNode> Node = OriginalTrackNode.Pin();
+	//if (Node.IsValid())
+	//{
+	//	Node->OnDragged(DragDropEvent);
+	//}
+
+	FVector2D pos;
+	pos.X = (DragDropEvent.GetScreenSpacePosition() + Offset).X;
+	pos.Y = StartingScreenPos.Y;
+
+	CursorDecoratorWindow->MoveWindowTo(pos);
+}
+
+TSharedPtr<SWidget> FUMGDragDropOp::GetDefaultDecorator() const
+{
+	return DecoratorWidget;
+}
+
+TSharedRef<FUMGDragDropOp> FUMGDragDropOp::New(UObject* Payload, UUserWidget* Decorator, const FVector2D &CursorPosition, const FVector2D &ScreenPositionOfNode)
+{
+	TSharedRef<FUMGDragDropOp> Operation = MakeShareable(new FUMGDragDropOp);
+	Operation->Offset = ScreenPositionOfNode - CursorPosition;
+	Operation->StartingScreenPos = ScreenPositionOfNode;
+
+	Operation->Payload = Payload;
+	Operation->DecoratorWidget = Decorator->TakeWidget();
+
+	Operation->Construct();
+
+	return Operation;
+}
\ No newline at end of file
diff --git a/Engine/Source/Runtime/UMG/Private/UserWidget.cpp b/Engine/Source/Runtime/UMG/Private/UserWidget.cpp
--- a/Engine/Source/Runtime/UMG/Private/UserWidget.cpp
+++ b/Engine/Source/Runtime/UMG/Private/UserWidget.cpp
@@ -5,29 +5,6 @@
 #include "UMGSequencePlayer.h"
 #include "SceneViewport.h"
 
-TSharedRef<FUMGDragDropOp> FUMGDragDropOp::New()
-{
-	TSharedRef<FUMGDragDropOp> Operation = MakeShareable(new FUMGDragDropOp);
-	Operation->Construct();
-
-	return Operation;
-}
-
-void FUMGDragDropOp::OnDrop(bool bDropWasHandled, const FPointerEvent& MouseEvent)
-{
-	FDragDropOperation::OnDrop(bDropWasHandled, MouseEvent);
-}
-
-void FUMGDragDropOp::OnDragged(const class FDragDropEvent& DragDropEvent)
-{
-	FDragDropOperation::OnDragged(DragDropEvent);
-}
-
-TSharedPtr<SWidget> FUMGDragDropOp::GetDefaultDecorator() const
-{
-	return DecoratorWidget;
-}
-
 /**
  * This class holds onto the widget when it's placed into the viewport.
  */
diff --git a/Engine/Source/Runtime/UMG/Public/Blueprint/UserWidget.h b/Engine/Source/Runtime/UMG/Public/Blueprint/UserWidget.h
--- a/Engine/Source/Runtime/UMG/Public/Blueprint/UserWidget.h
+++ b/Engine/Source/Runtime/UMG/Public/Blueprint/UserWidget.h
@@ -11,23 +11,6 @@ static FSlateRect NullRect;
 static FSlateWindowElementList NullElementList;
 static FWidgetStyle NullStyle;
 
-class FUMGDragDropOp : public FDragDropOperation
-{
-public:
-	DRAG_DROP_OPERATOR_TYPE(FUMGDragDropOp, FDragDropOperation)
-
-	static TSharedRef<FUMGDragDropOp> New();
-
-	virtual void OnDrop(bool bDropWasHandled, const FPointerEvent& MouseEvent) override;
-
-	virtual void OnDragged(const class FDragDropEvent& DragDropEvent) override;
-
-	virtual TSharedPtr<SWidget> GetDefaultDecorator() const override;
-
-private:
-	TSharedPtr<SWidget> DecoratorWidget;
-};
-
 /**
  * The state passed into OnPaint that we can expose as a single painting structure to blueprints to
  * allow script code to override OnPaint behavior.
diff --git a/Engine/Source/Runtime/UMG/Public/Slate/SDragPanel.h b/Engine/Source/Runtime/UMG/Public/Slate/SDragPanel.h
new file mode 100644
--- /dev/null
+++ b/Engine/Source/Runtime/UMG/Public/Slate/SDragPanel.h
@@ -0,0 +1,46 @@
+// Copyright 1998-2014 Epic Games, Inc. All Rights Reserved.
+
+#pragma once
+
+/**
+ * A delegate that is invoked when widgets want to notify a user that they have been clicked.
+ * Intended for use by buttons and other button-like widgets.
+ */
+DECLARE_DELEGATE_RetVal_TwoParams( 
+	TSharedPtr<FDragDropOperation>,
+	FOnDragOperationRequested,
+	FGeometry,
+	FPointerEvent);
+
+/**
+ * 
+ */
+class UMG_API SDragPanel : public SCompoundWidget
+{
+public:
+	SLATE_BEGIN_ARGS(SDragPanel)
+	: _Content()
+	{
+		_Visibility = EVisibility::Visible;
+	}
+		/** Slot for this designers content (optional) */
+		SLATE_DEFAULT_SLOT(FArguments, Content)
+
+		/**  */
+		SLATE_EVENT(FOnDragOperationRequested, OnDragDetected)
+
+	SLATE_END_ARGS()
+
+	void Construct(const FArguments& InArgs);
+
+	virtual FReply OnMouseButtonDown(const FGeometry& MyGeometry, const FPointerEvent& MouseEvent) override;
+	virtual FReply OnDragDetected(const FGeometry& MyGeometry, const FPointerEvent& MouseEvent) override;
+	
+	/** See Content slot */
+	void SetContent(TSharedRef<SWidget> InContent);
+
+private:
+
+	/** The delegate to execute when the drag drop operation is requested */
+	FOnDragOperationRequested OnDragOperationRequested;
+};
diff --git a/Engine/Source/Runtime/UMG/Public/Slate/SDropPanel.h b/Engine/Source/Runtime/UMG/Public/Slate/SDropPanel.h
new file mode 100644
--- /dev/null
+++ b/Engine/Source/Runtime/UMG/Public/Slate/SDropPanel.h
@@ -0,0 +1,30 @@
+// Copyright 1998-2014 Epic Games, Inc. All Rights Reserved.
+
+#pragma once
+
+/**
+ * 
+ */
+class UMG_API SDropPanel : public SCompoundWidget
+{
+public:
+	SLATE_BEGIN_ARGS(SDropPanel)
+	: _Content()
+	{
+		_Visibility = EVisibility::Visible;
+	}
+		/** Slot for this designers content (optional) */
+		SLATE_DEFAULT_SLOT(FArguments, Content)
+
+	SLATE_END_ARGS()
+
+	void Construct(const FArguments& InArgs);
+
+	virtual void OnDragEnter(const FGeometry& MyGeometry, const FDragDropEvent& DragDropEvent) override;
+	virtual void OnDragLeave(const FDragDropEvent& DragDropEvent) override;
+	virtual FReply OnDragOver(const FGeometry& MyGeometry, const FDragDropEvent& DragDropEvent) override;
+	virtual FReply OnDrop(const FGeometry& MyGeometry, const FDragDropEvent& DragDropEvent) override;
+	
+	/** See Content slot */
+	void SetContent(TSharedRef<SWidget> InContent);
+};
diff --git a/Engine/Source/Runtime/UMG/Public/Slate/UMGDragDropOp.h b/Engine/Source/Runtime/UMG/Public/Slate/UMGDragDropOp.h
new file mode 100644
--- /dev/null
+++ b/Engine/Source/Runtime/UMG/Public/Slate/UMGDragDropOp.h
@@ -0,0 +1,37 @@
+// Copyright 1998-2014 Epic Games, Inc. All Rights Reserved.
+
+#pragma once
+
+#include "DragAndDrop.h"
+
+/**
+ * This is the drag/drop class used for UMG, all UMG drag drop operations utilize this operation.
+ * It supports moving a UObject payload and using a UWidget decorator.
+ */
+class FUMGDragDropOp : public FDragDropOperation
+{
+public:
+	DRAG_DROP_OPERATOR_TYPE(FUMGDragDropOp, FDragDropOperation)
+	
+	static TSharedRef<FUMGDragDropOp> New(UObject* Payload, UUserWidget* Decorator, const FVector2D &CursorPosition, const FVector2D &ScreenPositionOfNode);
+
+	virtual void OnDrop( bool bDropWasHandled, const FPointerEvent& MouseEvent ) override;
+
+	virtual void OnDragged( const class FDragDropEvent& DragDropEvent ) override;
+
+	virtual TSharedPtr<SWidget> GetDefaultDecorator() const override;
+
+private:
+
+	/** The data object associated with the drag/drop operation. */
+	TWeakObjectPtr<UObject> Payload;
+
+	/** The widget used during the drag/drop action to show something being dragged. */
+	TSharedPtr<SWidget> DecoratorWidget;
+
+	/** The offset to use when dragging the object so that it says the same distance away from the mouse. */
+	FVector2D Offset;
+
+	/** The starting screen location where the drag operation started. */
+	FVector2D StartingScreenPos;
+};
