UMG - Fixing a crash bug in a few widgets accessing null slate widgets.

diff --git a/Engine/Source/Runtime/UMG/Private/Components/ProcessBar.cpp b/Engine/Source/Runtime/UMG/Private/Components/ProcessBar.cpp
--- a/Engine/Source/Runtime/UMG/Private/Components/ProcessBar.cpp
+++ b/Engine/Source/Runtime/UMG/Private/Components/ProcessBar.cpp
@@ -47,7 +47,11 @@ void UProgressBar::SyncronizeProperties()
 
 void UProgressBar::SetPercent(float InPercent)
 {
-	return MyProgressBar->SetPercent(InPercent);
+	Percent = InPercent;
+	if ( MyProgressBar.IsValid() )
+	{
+		MyProgressBar->SetPercent(InPercent);
+	}
 }
 
 #if WITH_EDITOR
diff --git a/Engine/Source/Runtime/UMG/Private/Components/Slider.cpp b/Engine/Source/Runtime/UMG/Private/Components/Slider.cpp
--- a/Engine/Source/Runtime/UMG/Private/Components/Slider.cpp
+++ b/Engine/Source/Runtime/UMG/Private/Components/Slider.cpp
@@ -54,13 +54,21 @@ void USlider::HandleOnMouseCaptureEnd()
 
 float USlider::GetValue()
 {
-	return MySlider->GetValue();
+	if ( MySlider.IsValid() )
+	{
+		return MySlider->GetValue();
+	}
+
+	return Value;
 }
 
 void USlider::SetValue(float InValue)
 {
 	Value = InValue;
-	return MySlider->SetValue(InValue);
+	if ( MySlider.IsValid() )
+	{
+		MySlider->SetValue(InValue);
+	}
 }
 
 #if WITH_EDITOR
diff --git a/Engine/Source/Runtime/UMG/Private/Components/Spacer.cpp b/Engine/Source/Runtime/UMG/Private/Components/Spacer.cpp
--- a/Engine/Source/Runtime/UMG/Private/Components/Spacer.cpp
+++ b/Engine/Source/Runtime/UMG/Private/Components/Spacer.cpp
@@ -28,7 +28,7 @@ TSharedRef<SWidget> USpacer::RebuildWidget()
 {
 	MySpacer = SNew(SSpacer);
 
-	//TODO UMG COnsider using a design time wrapper for spacer to show expandy arrows or some other
+	//TODO UMG Consider using a design time wrapper for spacer to show expandy arrows or some other
 	// indicator that there's a widget at work here.
 	
 	return MySpacer.ToSharedRef();
