UMG - Image widgets now only show the checkboard at design time.

diff --git a/Engine/Source/Runtime/UMG/Private/Components/Image.cpp b/Engine/Source/Runtime/UMG/Private/Components/Image.cpp
--- a/Engine/Source/Runtime/UMG/Private/Components/Image.cpp
+++ b/Engine/Source/Runtime/UMG/Private/Components/Image.cpp
@@ -70,13 +70,13 @@ const FSlateBrush* UImage::ConvertImage(TAttribute<USlateBrushAsset*> InImageAss
 		return &DynamicBrush.GetValue();
 	}
 
-	if ( Image == NULL )
+	if ( Image == NULL && IsDesignTime() )
 	{
 		SImage::FArguments ImageDefaults;
 		return ImageDefaults._Image.Get();
 	}
 
-	return &Image->Brush;
+	return Image ? &Image->Brush : NULL;
 }
 
 const FSlateBrush* UImage::GetImageBrush() const
