UMG - First prototype of Drag/Drop for UMG.  Usable in two ways, either using the DropPanel and DragPanel to fascilitate moving data around.  Or to implement the Drag/Drop functions on the widget blueprint.  There's a default DragDropOperation users can use, or they can extend it with their own for more complex scenarios.

diff --git a/Engine/Source/Runtime/UMG/Private/Components/DragPanel.cpp b/Engine/Source/Runtime/UMG/Private/Components/DragPanel.cpp
new file mode 100644
--- /dev/null
+++ b/Engine/Source/Runtime/UMG/Private/Components/DragPanel.cpp
@@ -0,0 +1,88 @@
+// Copyright 1998-2014 Epic Games, Inc. All Rights Reserved.
+
+#include "UMGPrivatePCH.h"
+
+#include "UMGDragDropOp.h"
+
+#define LOCTEXT_NAMESPACE "UMG"
+
+/////////////////////////////////////////////////////
+// UDragPanel
+
+UDragPanel::UDragPanel(const FPostConstructInitializeProperties& PCIP)
+	: Super(PCIP)
+{
+}
+
+void UDragPanel::ReleaseNativeWidget()
+{
+	Super::ReleaseNativeWidget();
+
+	MyDragPanel.Reset();
+}
+
+TSharedRef<SWidget> UDragPanel::RebuildWidget()
+{
+	MyDragPanel = SNew(SDragPanel)
+		.OnDragDetected(BIND_UOBJECT_DELEGATE(FOnDragOperationRequested, HandleDragDetected));
+
+	if ( GetChildrenCount() > 0 )
+	{
+		MyDragPanel->SetContent(GetContentSlot()->Content ? GetContentSlot()->Content->TakeWidget() : SNullWidget::NullWidget);
+	}
+	
+	return MyDragPanel.ToSharedRef();
+}
+
+void UDragPanel::SyncronizeProperties()
+{
+	Super::SyncronizeProperties();
+}
+
+void UDragPanel::OnSlotAdded(UPanelSlot* Slot)
+{
+	// Add the child to the live slot if it already exists
+	if ( MyDragPanel.IsValid() )
+	{
+		MyDragPanel->SetContent(Slot->Content ? Slot->Content->TakeWidget() : SNullWidget::NullWidget);
+	}
+}
+
+void UDragPanel::OnSlotRemoved(UPanelSlot* Slot)
+{
+	// Remove the widget from the live slot if it exists.
+	if ( MyDragPanel.IsValid() )
+	{
+		MyDragPanel->SetContent(SNullWidget::NullWidget);
+	}
+}
+
+TSharedPtr<FDragDropOperation> UDragPanel::HandleDragDetected(FGeometry Geometry, FPointerEvent PointerEvent)
+{
+	if ( OnDragEvent.IsBound() )
+	{
+		UDragDropOperation* Operation = NULL;
+
+		OnDragEvent.Execute(Geometry, PointerEvent, Operation);
+
+		FVector2D ScreenCursorPos = PointerEvent.GetScreenSpacePosition();
+		FVector2D ScreenDrageePosition = Geometry.AbsolutePosition;
+
+		return FUMGDragDropOp::New(Operation, ScreenCursorPos, ScreenDrageePosition);
+	}
+
+	return TSharedPtr<FDragDropOperation>();
+}
+
+#if WITH_EDITOR
+
+const FSlateBrush* UDragPanel::GetEditorIcon()
+{
+	return FUMGStyle::Get().GetBrush("Widget.DragPanel");
+}
+
+#endif
+
+/////////////////////////////////////////////////////
+
+#undef LOCTEXT_NAMESPACE
diff --git a/Engine/Source/Runtime/UMG/Private/Components/DropPanel.cpp b/Engine/Source/Runtime/UMG/Private/Components/DropPanel.cpp
new file mode 100644
--- /dev/null
+++ b/Engine/Source/Runtime/UMG/Private/Components/DropPanel.cpp
@@ -0,0 +1,132 @@
+// Copyright 1998-2014 Epic Games, Inc. All Rights Reserved.
+
+#include "UMGPrivatePCH.h"
+
+#include "UMGDragDropOp.h"
+
+#define LOCTEXT_NAMESPACE "UMG"
+
+/////////////////////////////////////////////////////
+// UDropPanel
+
+UDropPanel::UDropPanel(const FPostConstructInitializeProperties& PCIP)
+	: Super(PCIP)
+{
+}
+
+void UDropPanel::ReleaseNativeWidget()
+{
+	Super::ReleaseNativeWidget();
+
+	MyDropPanel.Reset();
+}
+
+TSharedRef<SWidget> UDropPanel::RebuildWidget()
+{
+	MyDropPanel = SNew(SDropPanel)
+		.OnDragEnter(BIND_UOBJECT_DELEGATE(FOnDragEnterEvent, HandleDragEnter))
+		.OnDragLeave(BIND_UOBJECT_DELEGATE(FOnDragLeaveEvent, HandleDragLeave))
+		.OnDragOver(BIND_UOBJECT_DELEGATE(FOnDragDropEvent, HandleDragOver))
+		.OnDrop(BIND_UOBJECT_DELEGATE(FOnDragDropEvent, HandleDrop));
+
+	if ( GetChildrenCount() > 0 )
+	{
+		MyDropPanel->SetContent(GetContentSlot()->Content ? GetContentSlot()->Content->TakeWidget() : SNullWidget::NullWidget);
+	}
+	
+	return MyDropPanel.ToSharedRef();
+}
+
+void UDropPanel::SyncronizeProperties()
+{
+	Super::SyncronizeProperties();
+}
+
+void UDropPanel::OnSlotAdded(UPanelSlot* Slot)
+{
+	// Add the child to the live slot if it already exists
+	if ( MyDropPanel.IsValid() )
+	{
+		MyDropPanel->SetContent(Slot->Content ? Slot->Content->TakeWidget() : SNullWidget::NullWidget);
+	}
+}
+
+void UDropPanel::OnSlotRemoved(UPanelSlot* Slot)
+{
+	// Remove the widget from the live slot if it exists.
+	if ( MyDropPanel.IsValid() )
+	{
+		MyDropPanel->SetContent(SNullWidget::NullWidget);
+	}
+}
+
+void UDropPanel::HandleDragEnter(FGeometry MyGeometry, const FDragDropEvent& DragDropEvent)
+{
+	if ( OnDragEnterEvent.IsBound() )
+	{
+		TSharedPtr<FUMGDragDropOp> NativeOp = DragDropEvent.GetOperationAs<FUMGDragDropOp>();
+		if ( NativeOp.IsValid() )
+		{
+			OnDragEnterEvent.Execute(MyGeometry, DragDropEvent, NativeOp->GetOperation());
+		}
+	}
+}
+
+void UDropPanel::HandleDragLeave(const FDragDropEvent& DragDropEvent)
+{
+	if ( OnDragLeaveEvent.IsBound() )
+	{
+		TSharedPtr<FUMGDragDropOp> NativeOp = DragDropEvent.GetOperationAs<FUMGDragDropOp>();
+		if ( NativeOp.IsValid() )
+		{
+			OnDragLeaveEvent.Execute(DragDropEvent, NativeOp->GetOperation());
+		}
+	}
+}
+
+FReply UDropPanel::HandleDragOver(FGeometry MyGeometry, const FDragDropEvent& DragDropEvent)
+{
+	if ( OnDragOverEvent.IsBound() )
+	{
+		TSharedPtr<FUMGDragDropOp> NativeOp = DragDropEvent.GetOperationAs<FUMGDragDropOp>();
+		if ( NativeOp.IsValid() )
+		{
+			if ( OnDragOverEvent.Execute(MyGeometry, DragDropEvent, NativeOp->GetOperation()) )
+			{
+				return FReply::Handled();
+			}
+		}
+	}
+
+	return FReply::Unhandled();
+}
+
+FReply UDropPanel::HandleDrop(FGeometry MyGeometry, const FDragDropEvent& DragDropEvent)
+{
+	if ( OnDropEvent.IsBound() )
+	{
+		TSharedPtr<FUMGDragDropOp> NativeOp = DragDropEvent.GetOperationAs<FUMGDragDropOp>();
+		if ( NativeOp.IsValid() )
+		{
+			if ( OnDropEvent.Execute(MyGeometry, DragDropEvent, NativeOp->GetOperation()) )
+			{
+				return FReply::Handled();
+			}
+		}
+	}
+
+	return FReply::Unhandled();
+}
+
+#if WITH_EDITOR
+
+const FSlateBrush* UDropPanel::GetEditorIcon()
+{
+	return FUMGStyle::Get().GetBrush("Widget.DropPanel");
+}
+
+#endif
+
+/////////////////////////////////////////////////////
+
+#undef LOCTEXT_NAMESPACE
diff --git a/Engine/Source/Runtime/UMG/Private/Components/PanelWidget.cpp b/Engine/Source/Runtime/UMG/Private/Components/PanelWidget.cpp
--- a/Engine/Source/Runtime/UMG/Private/Components/PanelWidget.cpp
+++ b/Engine/Source/Runtime/UMG/Private/Components/PanelWidget.cpp
@@ -73,11 +73,18 @@ bool UPanelWidget::RemoveChildAt(int32 Index)
 
 UPanelSlot* UPanelWidget::AddChild(UWidget* Content)
 {
+	if ( Content == nullptr )
+	{
+		return NULL;
+	}
+
 	if ( !bCanHaveMultipleChildren && GetChildrenCount() > 0 )
 	{
 		return NULL;
 	}
 
+	Content->RemoveFromParent();
+
 	UPanelSlot* Slot = ConstructObject<UPanelSlot>(GetSlotClass(), this);
 	Slot->SetFlags(RF_Transactional);
 	Slot->Content = Content;
diff --git a/Engine/Source/Runtime/UMG/Private/Components/Widget.cpp b/Engine/Source/Runtime/UMG/Private/Components/Widget.cpp
--- a/Engine/Source/Runtime/UMG/Private/Components/Widget.cpp
+++ b/Engine/Source/Runtime/UMG/Private/Components/Widget.cpp
@@ -86,6 +86,15 @@ UPanelWidget* UWidget::GetParent() const
 	return NULL;
 }
 
+void UWidget::RemoveFromParent()
+{
+	UPanelWidget* CurrentParent = GetParent();
+	if ( CurrentParent )
+	{
+		CurrentParent->RemoveChild(this);
+	}
+}
+
 TSharedRef<SWidget> UWidget::TakeWidget()
 {
 	TSharedPtr<SWidget> SafeWidget;
diff --git a/Engine/Source/Runtime/UMG/Private/DragDropOperation.cpp b/Engine/Source/Runtime/UMG/Private/DragDropOperation.cpp
new file mode 100644
--- /dev/null
+++ b/Engine/Source/Runtime/UMG/Private/DragDropOperation.cpp
@@ -0,0 +1,26 @@
+// Copyright 1998-2014 Epic Games, Inc. All Rights Reserved.
+
+#include "UMGPrivatePCH.h"
+
+/////////////////////////////////////////////////////
+// UDragDropOperation
+
+UDragDropOperation::UDragDropOperation(const FPostConstructInitializeProperties& PCIP)
+	: Super(PCIP)
+{
+}
+
+void UDragDropOperation::Drop_Implementation(const FPointerEvent& PointerEvent)
+{
+	OnDrop.Broadcast(this);
+}
+
+void UDragDropOperation::DropCanceled_Implementation(const FPointerEvent& PointerEvent)
+{
+	OnDropCanceled.Broadcast(this);
+}
+
+void UDragDropOperation::Dragged_Implementation(const FPointerEvent& PointerEvent)
+{
+	OnDragged.Broadcast(this);
+}
diff --git a/Engine/Source/Runtime/UMG/Private/Slate/SDropPanel.cpp b/Engine/Source/Runtime/UMG/Private/Slate/SDropPanel.cpp
--- a/Engine/Source/Runtime/UMG/Private/Slate/SDropPanel.cpp
+++ b/Engine/Source/Runtime/UMG/Private/Slate/SDropPanel.cpp
@@ -11,6 +11,11 @@
 
 void SDropPanel::Construct( const SDropPanel::FArguments& InArgs )
 {
+	OnDragEnterDelegate = InArgs._OnDragEnter;
+	OnDragLeaveDelegate = InArgs._OnDragLeave;
+	OnDragOverDelegate = InArgs._OnDragOver;
+	OnDropDelegate = InArgs._OnDrop;
+
 	ChildSlot
 	[
 		InArgs._Content.Widget
@@ -27,20 +32,36 @@ void SDropPanel::SetContent(TSharedRef<SWidget> InContent)
 
 void SDropPanel::OnDragEnter(const FGeometry& MyGeometry, const FDragDropEvent& DragDropEvent)
 {
-
+	if ( OnDragEnterDelegate.IsBound() )
+	{
+		OnDragEnterDelegate.Execute(MyGeometry, DragDropEvent);
+	}
 }
 
 void SDropPanel::OnDragLeave(const FDragDropEvent& DragDropEvent)
 {
-
+	if ( OnDragLeaveDelegate.IsBound() )
+	{
+		OnDragLeaveDelegate.Execute(DragDropEvent);
+	}
 }
 
 FReply SDropPanel::OnDragOver(const FGeometry& MyGeometry, const FDragDropEvent& DragDropEvent)
 {
+	if ( OnDragOverDelegate.IsBound() )
+	{
+		return OnDragOverDelegate.Execute(MyGeometry, DragDropEvent);
+	}
+
 	return FReply::Unhandled();
 }
 
 FReply SDropPanel::OnDrop(const FGeometry& MyGeometry, const FDragDropEvent& DragDropEvent)
 {
+	if ( OnDropDelegate.IsBound() )
+	{
+		return OnDropDelegate.Execute(MyGeometry, DragDropEvent);
+	}
+
 	return FReply::Unhandled();
 }
diff --git a/Engine/Source/Runtime/UMG/Private/Slate/SObjectWidget.cpp b/Engine/Source/Runtime/UMG/Private/Slate/SObjectWidget.cpp
--- a/Engine/Source/Runtime/UMG/Private/Slate/SObjectWidget.cpp
+++ b/Engine/Source/Runtime/UMG/Private/Slate/SObjectWidget.cpp
@@ -2,6 +2,8 @@
 
 #include "UMGPrivatePCH.h"
 
+#include "UMGDragDropOp.h"
+
 void SObjectWidget::Construct(const FArguments& InArgs, UUserWidget* InWidgetObject)
 {
 	WidgetObject = InWidgetObject;
@@ -209,41 +211,65 @@ FReply SObjectWidget::OnMouseButtonDoubleClick(const FGeometry& MyGeometry, cons
 	return FReply::Unhandled();
 }
 
-FReply SObjectWidget::OnDragDetected(const FGeometry& MyGeometry, const FPointerEvent& MouseEvent)
+FReply SObjectWidget::OnDragDetected(const FGeometry& MyGeometry, const FPointerEvent& PointerEvent)
 {
-	WidgetObject->OnDragDetected(MyGeometry, MouseEvent);
-
-	//TODO UMG Need to support drag drop, instead of returning the standard FSReply, maybe something else?
-	// Maybe like drag data, optional uobject + string data.
+	UDragDropOperation* Operation = NULL;
+	WidgetObject->OnDragDetected(MyGeometry, PointerEvent, Operation);
 
-	//FUMGDragDropOp
-	//.BeginDragDrop(FTrackNodeDragDropOp::New(SharedThis(this), ScreenCursorPos, ScreenNodePosition));
+	if ( Operation )
+	{
+		FVector2D ScreenCursorPos = PointerEvent.GetScreenSpacePosition();
+		FVector2D ScreenDrageePosition = MyGeometry.AbsolutePosition;
 
-	//FTrackNodeDragDropOp
+		return FReply::Handled().BeginDragDrop(FUMGDragDropOp::New(Operation, ScreenCursorPos, ScreenDrageePosition));
+	}
 
-	// TODO UMG
 	return FReply::Unhandled();
 }
 
 void SObjectWidget::OnDragEnter(const FGeometry& MyGeometry, const FDragDropEvent& DragDropEvent)
 {
-	// TODO UMG
+	TSharedPtr<FUMGDragDropOp> NativeOp = DragDropEvent.GetOperationAs<FUMGDragDropOp>();
+	if ( NativeOp.IsValid() )
+	{
+		WidgetObject->OnDragEnter(MyGeometry, DragDropEvent, NativeOp->GetOperation());
+	}
 }
 
 void SObjectWidget::OnDragLeave(const FDragDropEvent& DragDropEvent)
 {
-	// TODO UMG
+	TSharedPtr<FUMGDragDropOp> NativeOp = DragDropEvent.GetOperationAs<FUMGDragDropOp>();
+	if ( NativeOp.IsValid() )
+	{
+		WidgetObject->OnDragLeave(DragDropEvent, NativeOp->GetOperation());
+	}
 }
 
 FReply SObjectWidget::OnDragOver(const FGeometry& MyGeometry, const FDragDropEvent& DragDropEvent)
 {
-	// TODO UMG
+	TSharedPtr<FUMGDragDropOp> NativeOp = DragDropEvent.GetOperationAs<FUMGDragDropOp>();
+	if ( NativeOp.IsValid() )
+	{
+		if ( WidgetObject->OnDragOver(MyGeometry, DragDropEvent, NativeOp->GetOperation()) )
+		{
+			return FReply::Handled();
+		}
+	}
+
 	return FReply::Unhandled();
 }
 
 FReply SObjectWidget::OnDrop(const FGeometry& MyGeometry, const FDragDropEvent& DragDropEvent)
 {
-	// TODO UMG
+	TSharedPtr<FUMGDragDropOp> NativeOp = DragDropEvent.GetOperationAs<FUMGDragDropOp>();
+	if ( NativeOp.IsValid() )
+	{
+		if ( WidgetObject->OnDrop(MyGeometry, DragDropEvent, NativeOp->GetOperation()) )
+		{
+			return FReply::Handled();
+		}
+	}
+
 	return FReply::Unhandled();
 }
 
diff --git a/Engine/Source/Runtime/UMG/Private/Slate/UMGDragDropOp.cpp b/Engine/Source/Runtime/UMG/Private/Slate/UMGDragDropOp.cpp
--- a/Engine/Source/Runtime/UMG/Private/Slate/UMGDragDropOp.cpp
+++ b/Engine/Source/Runtime/UMG/Private/Slate/UMGDragDropOp.cpp
@@ -8,30 +8,33 @@
 // FUMGDragDropOp
 void FUMGDragDropOp::OnDrop( bool bDropWasHandled, const FPointerEvent& MouseEvent )
 {
-	//if ( bDropWasHandled == false )
-	//{
-	//	if (OriginalTrackNode.IsValid())
-	//	{
-	//		//OriginalTrackNode.Pin()->OnDropCancelled(MouseEvent);
-	//	}
-	//}
+	UDragDropOperation* OperationObject = DragOperation.Get();
+	if ( OperationObject )
+	{
+		if ( bDropWasHandled )
+		{
+			OperationObject->Drop(MouseEvent);
+		}
+		else
+		{
+			OperationObject->DropCanceled(MouseEvent);
+		}
+	}
 	
 	FDragDropOperation::OnDrop(bDropWasHandled, MouseEvent);
 }
 
 void FUMGDragDropOp::OnDragged( const class FDragDropEvent& DragDropEvent )
 {
-	//TSharedPtr<STrackNode> Node = OriginalTrackNode.Pin();
-	//if (Node.IsValid())
-	//{
-	//	Node->OnDragged(DragDropEvent);
-	//}
+	UDragDropOperation* OperationObject = DragOperation.Get();
+	if ( OperationObject )
+	{
+		OperationObject->Dragged(DragDropEvent);
+	}
 
-	FVector2D pos;
-	pos.X = (DragDropEvent.GetScreenSpacePosition() + Offset).X;
-	pos.Y = StartingScreenPos.Y;
+	FVector2D Position = DragDropEvent.GetScreenSpacePosition() + Offset;
 
-	CursorDecoratorWindow->MoveWindowTo(pos);
+	CursorDecoratorWindow->MoveWindowTo(Position);
 }
 
 TSharedPtr<SWidget> FUMGDragDropOp::GetDefaultDecorator() const
@@ -39,16 +42,25 @@ TSharedPtr<SWidget> FUMGDragDropOp::GetDefaultDecorator() const
 	return DecoratorWidget;
 }
 
-TSharedRef<FUMGDragDropOp> FUMGDragDropOp::New(UObject* Payload, UUserWidget* Decorator, const FVector2D &CursorPosition, const FVector2D &ScreenPositionOfNode)
+TSharedRef<FUMGDragDropOp> FUMGDragDropOp::New(UDragDropOperation* InOperation, const FVector2D &CursorPosition, const FVector2D &ScreenPositionOfDragee)
 {
 	TSharedRef<FUMGDragDropOp> Operation = MakeShareable(new FUMGDragDropOp);
-	Operation->Offset = ScreenPositionOfNode - CursorPosition;
-	Operation->StartingScreenPos = ScreenPositionOfNode;
+	Operation->Offset = ScreenPositionOfDragee - CursorPosition;
+	Operation->StartingScreenPos = ScreenPositionOfDragee;
 
-	Operation->Payload = Payload;
-	Operation->DecoratorWidget = Decorator->TakeWidget();
+	Operation->DragOperation = InOperation;
+
+	if ( InOperation->DefaultDragVisual == nullptr )
+	{
+		Operation->DecoratorWidget = SNew(STextBlock)
+			.Text(FText::FromString(InOperation->Tag));
+	}
+	else
+	{
+		Operation->DecoratorWidget = InOperation->DefaultDragVisual->TakeWidget();
+	}
 
 	Operation->Construct();
 
 	return Operation;
-}
\ No newline at end of file
+}
diff --git a/Engine/Source/Runtime/UMG/Private/UserWidget.cpp b/Engine/Source/Runtime/UMG/Private/UserWidget.cpp
--- a/Engine/Source/Runtime/UMG/Private/UserWidget.cpp
+++ b/Engine/Source/Runtime/UMG/Private/UserWidget.cpp
@@ -513,6 +513,11 @@ int32 UUserWidget::GetFullScreenZOrder() const
 	return ViewportZOrder;
 }
 
+UDragDropOperation* UUserWidget::CreateDragDropOperation(TSubclassOf<UDragDropOperation> Operation)
+{
+	return ConstructObject<UDragDropOperation>(Operation);
+}
+
 #if WITH_EDITOR
 
 const FSlateBrush* UUserWidget::GetEditorIcon()
diff --git a/Engine/Source/Runtime/UMG/Public/Blueprint/DragDropOperation.h b/Engine/Source/Runtime/UMG/Public/Blueprint/DragDropOperation.h
new file mode 100644
--- /dev/null
+++ b/Engine/Source/Runtime/UMG/Public/Blueprint/DragDropOperation.h
@@ -0,0 +1,54 @@
+// Copyright 1998-2014 Epic Games, Inc. All Rights Reserved.
+
+#pragma once
+
+#include "DragDropOperation.generated.h"
+
+DECLARE_DYNAMIC_MULTICAST_DELEGATE_OneParam(FOnDragDropMulticast, UDragDropOperation*, Operation);
+
+
+/**
+ * This class is the base drag drop operation for UMG, extend it to add additional data and add new functionality.
+ */
+UCLASS()
+class UMG_API UDragDropOperation : public UObject
+{
+	GENERATED_UCLASS_BODY()
+	
+public:
+	/**  */
+	UPROPERTY(EditAnywhere, BlueprintReadWrite, Category="Drag and Drop", meta=( ExposeOnSpawn="true" ))
+	FString Tag;
+	
+	/**  */
+	UPROPERTY(EditAnywhere, BlueprintReadWrite, Category="Drag and Drop", meta=( ExposeOnSpawn="true" ))
+	UObject* Payload;
+	
+	/**  */
+	UPROPERTY(EditAnywhere, BlueprintReadWrite, Category="Drag and Drop", meta=( ExposeOnSpawn="true" ))
+	class UWidget* DefaultDragVisual;
+
+	/**  */
+	UPROPERTY(BlueprintAssignable)
+	FOnDragDropMulticast OnDrop;
+
+	/**  */
+	UPROPERTY(BlueprintAssignable)
+	FOnDragDropMulticast OnDropCanceled;
+
+	/**  */
+	UPROPERTY(BlueprintAssignable)
+	FOnDragDropMulticast OnDragged;
+
+	/**  */
+	UFUNCTION(BlueprintNativeEvent, Category="Drag and Drop")
+	void Drop(const FPointerEvent& PointerEvent);
+
+	/**  */
+	UFUNCTION(BlueprintNativeEvent, Category="Drag and Drop")
+	void DropCanceled(const FPointerEvent& PointerEvent);
+
+	/**  */
+	UFUNCTION(BlueprintNativeEvent, Category="Drag and Drop")
+	void Dragged(const FPointerEvent& PointerEvent);
+};
diff --git a/Engine/Source/Runtime/UMG/Public/Blueprint/UserWidget.h b/Engine/Source/Runtime/UMG/Public/Blueprint/UserWidget.h
--- a/Engine/Source/Runtime/UMG/Public/Blueprint/UserWidget.h
+++ b/Engine/Source/Runtime/UMG/Public/Blueprint/UserWidget.h
@@ -192,15 +192,16 @@ public:
 	FSReply OnMouseButtonDoubleClick(FGeometry InMyGeometry, const FPointerEvent& InMouseEvent);
 
 	UFUNCTION(BlueprintImplementableEvent, Category="User Interface")
-	FSReply OnDragDetected(FGeometry MyGeometry, const FPointerEvent& MouseEvent);
-	//UFUNCTION(BlueprintImplementableEvent, Category="User Interface")
-	//void OnDragEnter(FGeometry MyGeometry, FDragDropEvent DragDropEvent);
-	//UFUNCTION(BlueprintImplementableEvent, Category="User Interface")
-	//void OnDragLeave(FDragDropEvent DragDropEvent);
-	//UFUNCTION(BlueprintImplementableEvent, Category="User Interface")
-	//FSReply OnDragOver(FGeometry MyGeometry, FDragDropEvent DragDropEvent);
-	//UFUNCTION(BlueprintImplementableEvent, Category="User Interface")
-	//FSReply OnDrop(FGeometry MyGeometry, FDragDropEvent DragDropEvent);
+	void OnDragDetected(FGeometry MyGeometry, const FPointerEvent& MouseEvent, UDragDropOperation*& Operation);
+	
+	UFUNCTION(BlueprintImplementableEvent, Category="User Interface")
+	void OnDragEnter(FGeometry MyGeometry, FPointerEvent PointerEvent, UDragDropOperation* Operation);
+	UFUNCTION(BlueprintImplementableEvent, Category="User Interface")
+	void OnDragLeave(FPointerEvent PointerEvent, UDragDropOperation* Operation);
+	UFUNCTION(BlueprintImplementableEvent, Category="User Interface")
+	bool OnDragOver(FGeometry MyGeometry, FPointerEvent PointerEvent, UDragDropOperation* Operation);
+	UFUNCTION(BlueprintImplementableEvent, Category="User Interface")
+	bool OnDrop(FGeometry MyGeometry, FPointerEvent PointerEvent, UDragDropOperation* Operation);
 
 	UFUNCTION(BlueprintImplementableEvent, Category="User Interface")
 	FSReply OnControllerButtonPressed(FGeometry MyGeometry, FControllerEvent ControllerEvent);
@@ -220,6 +221,9 @@ public:
 	UFUNCTION(BlueprintImplementableEvent, Category="User Interface")
 	FSReply OnMotionDetected(FGeometry MyGeometry, FMotionEvent InMotionEvent);
 
+	UFUNCTION(BlueprintCallable, Category="User Interface|Drag and Drop")
+	UDragDropOperation* CreateDragDropOperation(TSubclassOf<UDragDropOperation> Operation);
+
 	//virtual bool OnVisualizeTooltip(const TSharedPtr<SWidget>& TooltipContent);
 
 	/**
diff --git a/Engine/Source/Runtime/UMG/Public/Components/DragPanel.h b/Engine/Source/Runtime/UMG/Public/Components/DragPanel.h
new file mode 100644
--- /dev/null
+++ b/Engine/Source/Runtime/UMG/Public/Components/DragPanel.h
@@ -0,0 +1,51 @@
+// Copyright 1998-2014 Epic Games, Inc. All Rights Reserved.
+
+#pragma once
+
+#include "SDragPanel.h"
+
+#include "DragPanel.generated.h"
+
+/**
+ *
+ */
+UCLASS(meta=( Category="Primitive" ), ClassGroup=UserInterface)
+class UMG_API UDragPanel : public UContentWidget
+{
+	GENERATED_UCLASS_BODY()
+
+	DECLARE_DYNAMIC_DELEGATE_ThreeParams(FOnDragEvent, FGeometry, Geometry, const FPointerEvent&, PointerEvent, UDragDropOperation*&, Operation);
+
+public:
+	UPROPERTY(EditDefaultsOnly, Category=Events)
+	FOnDragEvent OnDragEvent;
+
+public:
+	
+	// UWidget interface
+	virtual void SyncronizeProperties() override;
+	// End of UWidget interface
+
+	virtual void ReleaseNativeWidget() override;
+
+#if WITH_EDITOR
+	virtual const FSlateBrush* GetEditorIcon() override;
+#endif
+
+protected:
+
+	// UPanelWidget
+	virtual void OnSlotAdded(UPanelSlot* Slot) override;
+	virtual void OnSlotRemoved(UPanelSlot* Slot) override;
+	// End UPanelWidget
+
+protected:
+	// UWidget interface
+	virtual TSharedRef<SWidget> RebuildWidget() override;
+	// End of UWidget interface
+
+	TSharedPtr<FDragDropOperation> HandleDragDetected(FGeometry Geometry, FPointerEvent PointerEvent);
+	
+protected:
+	TSharedPtr<SDragPanel> MyDragPanel;
+};
diff --git a/Engine/Source/Runtime/UMG/Public/Components/DropPanel.h b/Engine/Source/Runtime/UMG/Public/Components/DropPanel.h
new file mode 100644
--- /dev/null
+++ b/Engine/Source/Runtime/UMG/Public/Components/DropPanel.h
@@ -0,0 +1,71 @@
+// Copyright 1998-2014 Epic Games, Inc. All Rights Reserved.
+
+#pragma once
+
+#include "SDropPanel.h"
+
+#include "DropPanel.generated.h"
+
+/**
+ *
+ */
+UCLASS(meta=( Category="Primitive" ), ClassGroup=UserInterface)
+class UMG_API UDropPanel : public UContentWidget
+{
+	GENERATED_UCLASS_BODY()
+
+	DECLARE_DYNAMIC_DELEGATE_ThreeParams(FOnDragEnter, FGeometry, DropGeometry, const FPointerEvent&, PointerEvent, UDragDropOperation*, DragOperation);
+	DECLARE_DYNAMIC_DELEGATE_TwoParams(FOnDragLeave, const FPointerEvent&, PointerEvent, UDragDropOperation*, DragOperation);
+	DECLARE_DYNAMIC_DELEGATE_RetVal_ThreeParams(bool, FOnDragOver, FGeometry, DropGeometry, const FPointerEvent&, PointerEvent, UDragDropOperation*, DragOperation);
+	DECLARE_DYNAMIC_DELEGATE_RetVal_ThreeParams(bool, FOnDrop, FGeometry, DropGeometry, const FPointerEvent&, PointerEvent, UDragDropOperation*, DragOperation);
+
+public:
+
+	UPROPERTY(EditDefaultsOnly, Category=Events)
+	FOnDragEnter OnDragEnterEvent;
+
+	UPROPERTY(EditDefaultsOnly, Category=Events)
+	FOnDragLeave OnDragLeaveEvent;
+
+	UPROPERTY(EditDefaultsOnly, Category=Events)
+	FOnDragOver OnDragOverEvent;
+
+	UPROPERTY(EditDefaultsOnly, Category=Events)
+	FOnDrop OnDropEvent;
+
+public:
+
+	// UWidget interface
+	virtual void SyncronizeProperties() override;
+	// End of UWidget interface
+
+	virtual void ReleaseNativeWidget() override;
+
+#if WITH_EDITOR
+	virtual const FSlateBrush* GetEditorIcon() override;
+#endif
+
+protected:
+
+	// UPanelWidget
+	virtual void OnSlotAdded(UPanelSlot* Slot) override;
+	virtual void OnSlotRemoved(UPanelSlot* Slot) override;
+	// End UPanelWidget
+
+protected:
+	// UWidget interface
+	virtual TSharedRef<SWidget> RebuildWidget() override;
+	// End of UWidget interface
+
+protected:
+	void HandleDragEnter(FGeometry MyGeometry, const FDragDropEvent& DragDropEvent);
+
+	void HandleDragLeave(const FDragDropEvent& DragDropEvent);
+
+	FReply HandleDragOver(FGeometry MyGeometry, const FDragDropEvent& DragDropEvent);
+
+	FReply HandleDrop(FGeometry MyGeometry, const FDragDropEvent& DragDropEvent);
+
+protected:
+	TSharedPtr<SDropPanel> MyDropPanel;
+};
diff --git a/Engine/Source/Runtime/UMG/Public/Components/SlateWrapperTypes.h b/Engine/Source/Runtime/UMG/Public/Components/SlateWrapperTypes.h
--- a/Engine/Source/Runtime/UMG/Public/Components/SlateWrapperTypes.h
+++ b/Engine/Source/Runtime/UMG/Public/Components/SlateWrapperTypes.h
@@ -49,10 +49,10 @@ struct FSReply
 {
 	GENERATED_USTRUCT_BODY()
 
-	UPROPERTY(EditAnywhere, Category=User)
+	UPROPERTY(EditAnywhere, BlueprintReadWrite, Category=User)
 	uint32 bIsHandled : 1;
 
-	UPROPERTY(EditAnywhere, Category=User)
+	UPROPERTY(EditAnywhere, BlueprintReadWrite, Category=User)
 	uint32 bCaptureMouse : 1;
 
 	FSReply()
@@ -87,11 +87,11 @@ struct FSlateChildSize
 	GENERATED_USTRUCT_BODY()
 
 	/** The parameter of the size rule. */
-	UPROPERTY(EditAnywhere, Category=Appearance, meta=(ClampMin="0", ClampMax="1"))
+	UPROPERTY(EditAnywhere, BlueprintReadWrite, Category=Appearance, meta=( ClampMin="0", ClampMax="1" ))
 	float Value;
 
 	/** The sizing rule of the content. */
-	UPROPERTY(EditAnywhere, Category=Appearance)
+	UPROPERTY(EditAnywhere, BlueprintReadWrite, Category=Appearance)
 	TEnumAsByte<ESlateSizeRule::Type> SizeRule;
 
 	FSlateChildSize()
diff --git a/Engine/Source/Runtime/UMG/Public/Components/Widget.h b/Engine/Source/Runtime/UMG/Public/Components/Widget.h
--- a/Engine/Source/Runtime/UMG/Public/Components/Widget.h
+++ b/Engine/Source/Runtime/UMG/Public/Components/Widget.h
@@ -135,6 +135,10 @@ public:
 	UFUNCTION(BlueprintCallable, Category="Widget")
 	class UPanelWidget* GetParent() const;
 
+	/** Removes the widget from it's parent widget */
+	UFUNCTION(BlueprintCallable, Category="Widget")
+	void RemoveFromParent();
+
 	/**
 	 * Gets the underlying slate widget or constructs it if it doesn't exist.  This function is
 	 * virtual however, you should not inherit this function unless you're very aware of what you're
diff --git a/Engine/Source/Runtime/UMG/Public/Slate/SDropPanel.h b/Engine/Source/Runtime/UMG/Public/Slate/SDropPanel.h
--- a/Engine/Source/Runtime/UMG/Public/Slate/SDropPanel.h
+++ b/Engine/Source/Runtime/UMG/Public/Slate/SDropPanel.h
@@ -2,6 +2,30 @@
 
 #pragma once
 
+/**
+ * 
+ */
+DECLARE_DELEGATE_TwoParams(
+	FOnDragEnterEvent,
+	FGeometry,
+	const FDragDropEvent&);
+
+	/**
+ * 
+ */
+DECLARE_DELEGATE_OneParam(
+	FOnDragLeaveEvent,
+	const FDragDropEvent&);
+
+/**
+ * 
+ */
+DECLARE_DELEGATE_RetVal_TwoParams(
+	FReply,
+	FOnDragDropEvent,
+	FGeometry,
+	const FDragDropEvent&);
+
 /**
  * 
  */
@@ -16,6 +40,18 @@ public:
 		/** Slot for this designers content (optional) */
 		SLATE_DEFAULT_SLOT(FArguments, Content)
 
+		/**  */
+		SLATE_EVENT(FOnDragEnterEvent, OnDragEnter)
+
+		/**  */
+		SLATE_EVENT(FOnDragLeaveEvent, OnDragLeave)
+
+		/**  */
+		SLATE_EVENT(FOnDragDropEvent, OnDragOver)
+
+		/**  */
+		SLATE_EVENT(FOnDragDropEvent, OnDrop)
+
 	SLATE_END_ARGS()
 
 	void Construct(const FArguments& InArgs);
@@ -27,4 +63,11 @@ public:
 	
 	/** See Content slot */
 	void SetContent(TSharedRef<SWidget> InContent);
+
+private:
+
+	FOnDragEnterEvent OnDragEnterDelegate;
+	FOnDragLeaveEvent OnDragLeaveDelegate;
+	FOnDragDropEvent OnDragOverDelegate;
+	FOnDragDropEvent OnDropDelegate;
 };
diff --git a/Engine/Source/Runtime/UMG/Public/Slate/UMGDragDropOp.h b/Engine/Source/Runtime/UMG/Public/Slate/UMGDragDropOp.h
--- a/Engine/Source/Runtime/UMG/Public/Slate/UMGDragDropOp.h
+++ b/Engine/Source/Runtime/UMG/Public/Slate/UMGDragDropOp.h
@@ -13,7 +13,7 @@ class FUMGDragDropOp : public FDragDropOperation
 public:
 	DRAG_DROP_OPERATOR_TYPE(FUMGDragDropOp, FDragDropOperation)
 	
-	static TSharedRef<FUMGDragDropOp> New(UObject* Payload, UUserWidget* Decorator, const FVector2D &CursorPosition, const FVector2D &ScreenPositionOfNode);
+	static TSharedRef<FUMGDragDropOp> New(UDragDropOperation* Operation, const FVector2D &CursorPosition, const FVector2D &ScreenPositionOfNode);
 
 	virtual void OnDrop( bool bDropWasHandled, const FPointerEvent& MouseEvent ) override;
 
@@ -21,10 +21,11 @@ public:
 
 	virtual TSharedPtr<SWidget> GetDefaultDecorator() const override;
 
+	UDragDropOperation* GetOperation() const { return DragOperation.Get(); }
+
 private:
 
-	/** The data object associated with the drag/drop operation. */
-	TWeakObjectPtr<UObject> Payload;
+	TWeakObjectPtr<UDragDropOperation> DragOperation;
 
 	/** The widget used during the drag/drop action to show something being dragged. */
 	TSharedPtr<SWidget> DecoratorWidget;
diff --git a/Engine/Source/Runtime/UMG/Public/UMG.h b/Engine/Source/Runtime/UMG/Public/UMG.h
--- a/Engine/Source/Runtime/UMG/Public/UMG.h
+++ b/Engine/Source/Runtime/UMG/Public/UMG.h
@@ -24,6 +24,8 @@
 // Components
 #include "SlateWrapperTypes.h"
 
+#include "DragDropOperation.h"
+
 #include "Visual.h"
 
 #include "PanelSlot.h"
@@ -92,6 +94,9 @@
 #include "ListView.h"
 #include "TileView.h"
 
+#include "DragPanel.h"
+#include "DropPanel.h"
+
 #include "Viewport.h"
 
 #include "NativeWidgetHost.h"
