#UMG

- Sub-UserWidgets were not getting the Local Player's Context. This Context is now passed on when RebuildWidget is called.

diff --git a/Engine/Source/Runtime/UMG/Private/UserWidget.cpp b/Engine/Source/Runtime/UMG/Private/UserWidget.cpp
--- a/Engine/Source/Runtime/UMG/Private/UserWidget.cpp
+++ b/Engine/Source/Runtime/UMG/Private/UserWidget.cpp
@@ -253,6 +253,19 @@ TSharedRef<SWidget> UUserWidget::RebuildWidget()
 
 	check(bInitialized);
 
+	//setup the player context on sub user widgets, if we have a valid context
+	if (PlayerContext.IsValid())
+	{
+		for (UWidget* Widget : Components)
+		{
+			UUserWidget* UserWidget = Cast<UUserWidget>(Widget);
+			if (UserWidget)
+			{
+				UserWidget->SetPlayerContext(PlayerContext);
+			}
+		}
+	}
+
 	// Add the first component to the root of the widget surface.
 	if ( Components.Num() > 0 && Components[0] != NULL )
 	{
