#UMG

- Fixed an issue with visibility binding. UImage wasn't calling Super::SyncronizeProperties();

diff --git a/Engine/Source/Runtime/UMG/Private/Components/Image.cpp b/Engine/Source/Runtime/UMG/Private/Components/Image.cpp
--- a/Engine/Source/Runtime/UMG/Private/Components/Image.cpp
+++ b/Engine/Source/Runtime/UMG/Private/Components/Image.cpp
@@ -30,6 +30,8 @@ TSharedRef<SWidget> UImage::RebuildWidget()
 
 void UImage::SyncronizeProperties()
 {
+	Super::SyncronizeProperties();
+
 	TAttribute<FSlateColor> ColorAndOpacityBinding = OPTIONAL_BINDING(FSlateColor, ColorAndOpacity);
 	TAttribute<const FSlateBrush*> ImageBinding = OPTIONAL_BINDING_CONVERT(USlateBrushAsset*, Image, const FSlateBrush*, ConvertImage);
 
