UMG - User widget used in another blueprint, the visibility and enabled bindings no longer stomp the bindings performed to root widgets of the user widget.  At design time, enabled and visibility flags no longer apply.

diff --git a/Engine/Source/Runtime/UMG/Private/Components/Widget.cpp b/Engine/Source/Runtime/UMG/Private/Components/Widget.cpp
--- a/Engine/Source/Runtime/UMG/Private/Components/Widget.cpp
+++ b/Engine/Source/Runtime/UMG/Private/Components/Widget.cpp
@@ -98,6 +98,7 @@ void UWidget::RemoveFromParent()
 TSharedRef<SWidget> UWidget::TakeWidget()
 {
 	TSharedPtr<SWidget> SafeWidget;
+	bool bNewlyCreated = false;
 
 	if ( !MyWidget.IsValid() )
 	{
@@ -107,7 +108,8 @@ TSharedRef<SWidget> UWidget::TakeWidget()
 		// just blow away the const here.
 		SafeWidget = RebuildWidget();
 		MyWidget = SafeWidget;
-		SyncronizeProperties();
+
+		bNewlyCreated = true;
 	}
 	else
 	{
@@ -130,11 +132,18 @@ TSharedRef<SWidget> UWidget::TakeWidget()
 
 			MyGCWidget = SafeGCWidget;
 
+			SyncronizeProperties();
+
 			return SafeGCWidget.ToSharedRef();
 		}
 	}
 	else
 	{
+		if ( bNewlyCreated )
+		{
+			SyncronizeProperties();
+		}
+
 		return SafeWidget.ToSharedRef();
 	}
 }
@@ -289,13 +298,26 @@ TSharedRef<SWidget> UWidget::RebuildWidget()
 
 void UWidget::SyncronizeProperties()
 {
-	MyWidget.Pin()->SetEnabled(OPTIONAL_BINDING(bool, bIsEnabled));
-	MyWidget.Pin()->SetVisibility(OPTIONAL_BINDING_CONVERT(ESlateVisibility::Type, Visiblity, EVisibility, ConvertVisibility));
-	//MyWidget->SetCursor(OPTIONAL_BINDING(EMouseCursor)
+	// We want to apply the bindings to the cached widget, which could be the SWidget, or the SObjectWidget, 
+	// in the case where it's a user widget.  We always want to prefer the SObjectWidget so that bindings to 
+	// visibility and enabled status are not stomping values setup in the root widget in the User Widget.
+	TSharedPtr<SWidget> SafeWidget = GetCachedWidget();
+
+	// Always use an enabled and visible state in the designer.
+	if ( IsDesignTime() )
+	{
+		SafeWidget->SetEnabled(true);
+		SafeWidget->SetVisibility(EVisibility::Visible);
+	}
+	else
+	{
+		SafeWidget->SetEnabled(OPTIONAL_BINDING(bool, bIsEnabled));
+		SafeWidget->SetVisibility(OPTIONAL_BINDING_CONVERT(ESlateVisibility::Type, Visiblity, EVisibility, ConvertVisibility));
+	}
 
 	if ( !ToolTipText.IsEmpty() )
 	{
-		MyWidget.Pin()->SetToolTipText(OPTIONAL_BINDING(FText, ToolTipText));
+		SafeWidget->SetToolTipText(OPTIONAL_BINDING(FText, ToolTipText));
 	}
 }
 
