UMG - Organization of some less common widgets, and general code cleanup.

diff --git a/Engine/Source/Runtime/UMG/Private/Components/Border.cpp b/Engine/Source/Runtime/UMG/Private/Components/Border.cpp
--- a/Engine/Source/Runtime/UMG/Private/Components/Border.cpp
+++ b/Engine/Source/Runtime/UMG/Private/Components/Border.cpp
@@ -40,7 +40,7 @@ TSharedRef<SWidget> UBorder::RebuildWidget()
 		Cast<UBorderSlot>(GetContentSlot())->BuildSlot(MyBorder.ToSharedRef());
 	}
 
-	return MyBorder.ToSharedRef();
+	return BuildDesignTimeWidget( MyBorder.ToSharedRef() );
 }
 
 void UBorder::SyncronizeProperties()
@@ -109,8 +109,7 @@ const FSlateBrush* UBorder::GetBorderBrush() const
 {
 	if ( Brush == NULL )
 	{
-		SBorder::FArguments BorderDefaults;
-		return BorderDefaults._BorderImage.Get();
+		return NULL;
 	}
 
 	return &Brush->Brush;
diff --git a/Engine/Source/Runtime/UMG/Private/Components/EditableText.cpp b/Engine/Source/Runtime/UMG/Private/Components/EditableText.cpp
--- a/Engine/Source/Runtime/UMG/Private/Components/EditableText.cpp
+++ b/Engine/Source/Runtime/UMG/Private/Components/EditableText.cpp
@@ -75,8 +75,11 @@ void UEditableText::SyncronizeProperties()
 {
 	Super::SyncronizeProperties();
 
-	MyEditableText->SetText(Text);
-	MyEditableText->SetHintText(HintText);
+	TAttribute<FText> TextBinding = OPTIONAL_BINDING(FText, Text);
+	TAttribute<FText> HintTextBinding = OPTIONAL_BINDING(FText, HintText);
+
+	MyEditableText->SetText(TextBinding);
+	MyEditableText->SetHintText(HintTextBinding);
 	MyEditableText->SetIsReadOnly(IsReadOnly);
 	MyEditableText->SetIsPassword(IsPassword);
 	MyEditableText->SetColorAndOpacity(ColorAndOpacity);
diff --git a/Engine/Source/Runtime/UMG/Private/Components/EditableTextBox.cpp b/Engine/Source/Runtime/UMG/Private/Components/EditableTextBox.cpp
--- a/Engine/Source/Runtime/UMG/Private/Components/EditableTextBox.cpp
+++ b/Engine/Source/Runtime/UMG/Private/Components/EditableTextBox.cpp
@@ -75,8 +75,11 @@ void UEditableTextBox::SyncronizeProperties()
 		MyEditableTextBlock->SetStyle(StylePtr);
 	}
 
-	MyEditableTextBlock->SetText(Text);
-	MyEditableTextBlock->SetHintText(HintText);
+	TAttribute<FText> TextBinding = OPTIONAL_BINDING(FText, Text);
+	TAttribute<FText> HintTextBinding = OPTIONAL_BINDING(FText, HintText);
+
+	MyEditableTextBlock->SetText(TextBinding);
+	MyEditableTextBlock->SetHintText(HintTextBinding);
 	MyEditableTextBlock->SetIsReadOnly(IsReadOnly);
 	MyEditableTextBlock->SetIsPassword(IsPassword);
 //	MyEditableTextBlock->SetColorAndOpacity(ColorAndOpacity);
diff --git a/Engine/Source/Runtime/UMG/Private/Components/PanelWidget.cpp b/Engine/Source/Runtime/UMG/Private/Components/PanelWidget.cpp
--- a/Engine/Source/Runtime/UMG/Private/Components/PanelWidget.cpp
+++ b/Engine/Source/Runtime/UMG/Private/Components/PanelWidget.cpp
@@ -127,6 +127,9 @@ void UPanelWidget::InsertChildAt(int32 Index, UWidget* Content)
 		Content->Slot = Slot;
 	}
 
+	// Only allow inserting within the valid range of slots (and one more than the size).
+	Index = FMath::Clamp(Index, 0, FMath::Max(Slots.Num(), 1));
+
 	Slots.Insert(Slot, Index);
 
 	OnSlotAdded(Slot);
@@ -143,6 +146,15 @@ bool UPanelWidget::RemoveChild(UWidget* Content)
 	return false;
 }
 
+void UPanelWidget::ClearChildren()
+{
+	int32 Children = GetChildrenCount();
+	for ( int32 ChildIndex = 0; ChildIndex < Children; ChildIndex++ )
+	{
+		RemoveChildAt(0);
+	}
+}
+
 void UPanelWidget::PostLoad()
 {
 	Super::PostLoad();
diff --git a/Engine/Source/Runtime/UMG/Private/Components/ScrollBox.cpp b/Engine/Source/Runtime/UMG/Private/Components/ScrollBox.cpp
--- a/Engine/Source/Runtime/UMG/Private/Components/ScrollBox.cpp
+++ b/Engine/Source/Runtime/UMG/Private/Components/ScrollBox.cpp
@@ -90,16 +90,6 @@ void UScrollBox::SyncronizeProperties()
 	MyScrollBox->SetOrientation(Orientation);
 }
 
-void UScrollBox::ClearChildren()
-{
-	Slots.Reset();
-
-	if ( MyScrollBox.IsValid() )
-	{
-		MyScrollBox->ClearChildren();
-	}
-}
-
 void UScrollBox::SetScrollOffset(float NewScrollOffset)
 {
 	DesiredScrollOffset = NewScrollOffset;
diff --git a/Engine/Source/Runtime/UMG/Public/Components/BorderSlot.h b/Engine/Source/Runtime/UMG/Public/Components/BorderSlot.h
--- a/Engine/Source/Runtime/UMG/Public/Components/BorderSlot.h
+++ b/Engine/Source/Runtime/UMG/Public/Components/BorderSlot.h
@@ -11,6 +11,8 @@ UCLASS()
 class UMG_API UBorderSlot : public UPanelSlot
 {
 	GENERATED_UCLASS_BODY()
+
+public:
 	
 	/** The padding area between the slot and the content it contains. */
 	UPROPERTY(EditDefaultsOnly, Category=Layout)
diff --git a/Engine/Source/Runtime/UMG/Public/Components/CanvasPanel.h b/Engine/Source/Runtime/UMG/Public/Components/CanvasPanel.h
--- a/Engine/Source/Runtime/UMG/Public/Components/CanvasPanel.h
+++ b/Engine/Source/Runtime/UMG/Public/Components/CanvasPanel.h
@@ -9,6 +9,8 @@ class UMG_API UCanvasPanel : public UPanelWidget
 {
 	GENERATED_UCLASS_BODY()
 
+public:
+
 	/** Gets the underlying native canvas widget if it has been constructed */
 	TSharedPtr<class SConstraintCanvas> GetCanvasWidget() const;
 
diff --git a/Engine/Source/Runtime/UMG/Public/Components/CanvasPanelSlot.h b/Engine/Source/Runtime/UMG/Public/Components/CanvasPanelSlot.h
--- a/Engine/Source/Runtime/UMG/Public/Components/CanvasPanelSlot.h
+++ b/Engine/Source/Runtime/UMG/Public/Components/CanvasPanelSlot.h
@@ -12,6 +12,8 @@ struct FAnchorData
 public:
 	GENERATED_USTRUCT_BODY()
 
+public:
+
 	/** Offset. */
 	UPROPERTY(EditAnywhere, Category=AnchorData)
 	FMargin Offsets;
diff --git a/Engine/Source/Runtime/UMG/Public/Components/CheckBox.h b/Engine/Source/Runtime/UMG/Public/Components/CheckBox.h
--- a/Engine/Source/Runtime/UMG/Public/Components/CheckBox.h
+++ b/Engine/Source/Runtime/UMG/Public/Components/CheckBox.h
@@ -12,7 +12,7 @@ class UMG_API UCheckBox : public UContentWidget
 {
 	GENERATED_UCLASS_BODY()
 
-protected:
+public:
 	/** Style of the check box */
 	UPROPERTY(EditDefaultsOnly, Category=Style, meta=( DisplayThumbnail = "true" ))
 	USlateWidgetStyleAsset* Style;
diff --git a/Engine/Source/Runtime/UMG/Public/Components/EditableText.h b/Engine/Source/Runtime/UMG/Public/Components/EditableText.h
--- a/Engine/Source/Runtime/UMG/Public/Components/EditableText.h
+++ b/Engine/Source/Runtime/UMG/Public/Components/EditableText.h
@@ -5,7 +5,7 @@
 #include "EditableText.generated.h"
 
 /** Editable text box widget */
-UCLASS(meta=( Category="Common" ), ClassGroup=UserInterface)
+UCLASS(meta=( Category="Primitive" ), ClassGroup=UserInterface)
 class UMG_API UEditableText : public UWidget
 {
 	GENERATED_UCLASS_BODY()
@@ -21,10 +21,18 @@ public:
 	UPROPERTY(EditDefaultsOnly, Category=Content)
 	FText Text;
 
+	/** A bindable delegate to allow logic to drive the text of the widget */
+	UPROPERTY()
+	FGetText TextDelegate;
+
 	/** Hint text that appears when there is no text in the text box */
 	UPROPERTY(EditDefaultsOnly, Category=Content)
 	FText HintText;
 
+	/** A bindable delegate to allow logic to drive the hint text of the widget */
+	UPROPERTY()
+	FGetText HintTextDelegate;
+
 	/** Text style */
 	UPROPERTY(EditDefaultsOnly, Category=Style, meta=( DisplayThumbnail = "true" ))
 	USlateWidgetStyleAsset* Style;
diff --git a/Engine/Source/Runtime/UMG/Public/Components/EditableTextBox.h b/Engine/Source/Runtime/UMG/Public/Components/EditableTextBox.h
--- a/Engine/Source/Runtime/UMG/Public/Components/EditableTextBox.h
+++ b/Engine/Source/Runtime/UMG/Public/Components/EditableTextBox.h
@@ -25,10 +25,18 @@ public:
 	UPROPERTY(EditDefaultsOnly, Category=Content)
 	FText Text;
 
+	/** A bindable delegate to allow logic to drive the text of the widget */
+	UPROPERTY()
+	FGetText TextDelegate;
+
 	/** Hint text that appears when there is no text in the text box */
 	UPROPERTY(EditDefaultsOnly, Category=Content)
 	FText HintText;
 
+	/** A bindable delegate to allow logic to drive the hint text of the widget */
+	UPROPERTY()
+	FGetText HintTextDelegate;
+
 	/** Font color and opacity (overrides Style) */
 	UPROPERTY(EditDefaultsOnly, Category=Appearance)
 	FSlateFontInfo Font;
diff --git a/Engine/Source/Runtime/UMG/Public/Components/HorizontalBoxSlot.h b/Engine/Source/Runtime/UMG/Public/Components/HorizontalBoxSlot.h
--- a/Engine/Source/Runtime/UMG/Public/Components/HorizontalBoxSlot.h
+++ b/Engine/Source/Runtime/UMG/Public/Components/HorizontalBoxSlot.h
@@ -10,6 +10,8 @@ UCLASS()
 class UMG_API UHorizontalBoxSlot : public UPanelSlot
 {
 	GENERATED_UCLASS_BODY()
+
+public:
 	
 	/** The amount of padding between the slots parent and the content. */
 	UPROPERTY(EditDefaultsOnly, Category=Appearance)
diff --git a/Engine/Source/Runtime/UMG/Public/Components/MultiLineEditableText.h b/Engine/Source/Runtime/UMG/Public/Components/MultiLineEditableText.h
--- a/Engine/Source/Runtime/UMG/Public/Components/MultiLineEditableText.h
+++ b/Engine/Source/Runtime/UMG/Public/Components/MultiLineEditableText.h
@@ -5,7 +5,7 @@
 #include "MultiLineEditableText.generated.h"
 
 /** Editable text box widget */
-UCLASS(meta=( Category="Common" ), ClassGroup=UserInterface)
+UCLASS(meta=( Category="Primitive" ), ClassGroup=UserInterface)
 class UMG_API UMultiLineEditableText : public UWidget
 {
 	GENERATED_UCLASS_BODY()
diff --git a/Engine/Source/Runtime/UMG/Public/Components/PanelWidget.h b/Engine/Source/Runtime/UMG/Public/Components/PanelWidget.h
--- a/Engine/Source/Runtime/UMG/Public/Components/PanelWidget.h
+++ b/Engine/Source/Runtime/UMG/Public/Components/PanelWidget.h
@@ -18,8 +18,6 @@ public:
 
 public:
 
-	//TODO UMG Add ways to make adding slots callable by blueprints.
-
 	/**  */
 	UFUNCTION(BlueprintCallable, Category="Appearance")
 	int32 GetChildrenCount() const;
@@ -40,19 +38,26 @@ public:
 	UFUNCTION(BlueprintCallable, Category="Appearance")
 	UPanelSlot* AddChild(UWidget* Content);
 
-	/**  */
+	/** Swaps the widget out of the slot at the given index, replacing it with a different widget. */
 	UFUNCTION(BlueprintCallable, Category="Appearance")
 	void ReplaceChildAt(int32 Index, UWidget* Content);
 
-	/**  */
+	/** Inserts a widget at a specific index */
 	UFUNCTION(BlueprintCallable, Category="Appearance")
 	void InsertChildAt(int32 Index, UWidget* Content);
 
-	/**  */
+	/**
+	 * Removes a specific widget from the container.
+	 * @return true if the widget was found and removed.
+	 */
 	UFUNCTION(BlueprintCallable, Category="Appearance")
 	bool RemoveChild(UWidget* Content);
 
-	/**  */
+	/** Remove all child widgets from the panel widget. */
+	UFUNCTION(BlueprintCallable, Category="Appearance")
+	void ClearChildren();
+
+	/** @returns true if the panel supports more than one child. */
 	bool CanHaveMultipleChildren() const
 	{
 		return bCanHaveMultipleChildren;
diff --git a/Engine/Source/Runtime/UMG/Public/Components/ScrollBox.h b/Engine/Source/Runtime/UMG/Public/Components/ScrollBox.h
--- a/Engine/Source/Runtime/UMG/Public/Components/ScrollBox.h
+++ b/Engine/Source/Runtime/UMG/Public/Components/ScrollBox.h
@@ -24,10 +24,6 @@ class UMG_API UScrollBox : public UPanelWidget
 
 	virtual void ReleaseNativeWidget() override;
 
-	/** Removes all child widget from the scrollbox */
-	UFUNCTION(BlueprintCallable, Category="Widget")
-	void ClearChildren();
-
 	/** Updates the scroll offset of the scrollbox */
 	UFUNCTION(BlueprintCallable, Category="Widget")
 	void SetScrollOffset(float NewScrollOffset);
diff --git a/Engine/Source/Runtime/UMG/Public/Components/ScrollBoxSlot.h b/Engine/Source/Runtime/UMG/Public/Components/ScrollBoxSlot.h
--- a/Engine/Source/Runtime/UMG/Public/Components/ScrollBoxSlot.h
+++ b/Engine/Source/Runtime/UMG/Public/Components/ScrollBoxSlot.h
@@ -11,6 +11,8 @@ UCLASS()
 class UMG_API UScrollBoxSlot : public UPanelSlot
 {
 	GENERATED_UCLASS_BODY()
+
+public:
 	
 	/** The padding area between the slot and the content it contains. */
 	UPROPERTY(EditDefaultsOnly, Category=Layout)
diff --git a/Engine/Source/Runtime/UMG/Public/Components/Slider.h b/Engine/Source/Runtime/UMG/Public/Components/Slider.h
--- a/Engine/Source/Runtime/UMG/Public/Components/Slider.h
+++ b/Engine/Source/Runtime/UMG/Public/Components/Slider.h
@@ -13,6 +13,8 @@ UCLASS(meta=( Category="Common" ), ClassGroup=UserInterface)
 class UMG_API USlider : public UWidget
 {
 	GENERATED_UCLASS_BODY()
+
+public:
 	
 	///** The style used to draw the slider. */
 	//SLATE_STYLE_ARGUMENT(FSliderStyle, Style)
diff --git a/Engine/Source/Runtime/UMG/Public/Components/Viewport.h b/Engine/Source/Runtime/UMG/Public/Components/Viewport.h
--- a/Engine/Source/Runtime/UMG/Public/Components/Viewport.h
+++ b/Engine/Source/Runtime/UMG/Public/Components/Viewport.h
@@ -238,7 +238,7 @@ protected:
 /**
  * 
  */
-UCLASS(meta=( Category="Common" ), ClassGroup=UserInterface)
+UCLASS(meta=( Category="Advanced" ), ClassGroup=UserInterface)
 class UMG_API UViewport : public UContentWidget
 {
 	GENERATED_UCLASS_BODY()
