UMG - The base UWidget now implements HasKeyboardFocus, HasFocusedDescendants, and HasMouseCapture.

diff --git a/Engine/Source/Runtime/UMG/Private/Components/Widget.cpp b/Engine/Source/Runtime/UMG/Private/Components/Widget.cpp
--- a/Engine/Source/Runtime/UMG/Private/Components/Widget.cpp
+++ b/Engine/Source/Runtime/UMG/Private/Components/Widget.cpp
@@ -68,6 +68,36 @@ bool UWidget::IsHovered() const
 	return false;
 }
 
+bool UWidget::HasKeyboardFocus() const
+{
+	if ( MyWidget.IsValid() )
+	{
+		return MyWidget.Pin()->HasKeyboardFocus();
+	}
+
+	return false;
+}
+
+bool UWidget::HasFocusedDescendants() const
+{
+	if ( MyWidget.IsValid() )
+	{
+		return MyWidget.Pin()->HasFocusedDescendants();
+	}
+
+	return false;
+}
+
+bool UWidget::HasMouseCapture() const
+{
+	if ( MyWidget.IsValid() )
+	{
+		return MyWidget.Pin()->HasMouseCapture();
+	}
+
+	return false;
+}
+
 void UWidget::ForceLayoutPrepass()
 {
 	if ( MyWidget.IsValid() )
diff --git a/Engine/Source/Runtime/UMG/Public/Components/Widget.h b/Engine/Source/Runtime/UMG/Public/Components/Widget.h
--- a/Engine/Source/Runtime/UMG/Public/Components/Widget.h
+++ b/Engine/Source/Runtime/UMG/Public/Components/Widget.h
@@ -127,6 +127,28 @@ public:
 	UFUNCTION(BlueprintCallable, Category="Widget")
 	bool IsHovered() const;
 
+	/**
+	 * Checks to see if this widget currently has the keyboard focus
+	 *
+	 * @return  True if this widget has keyboard focus
+	 */
+	UFUNCTION(BlueprintCallable, Category="Widget")
+	bool HasKeyboardFocus() const;
+
+	/**
+	 * @return Whether this widget has any descendants with keyboard focus
+	 */
+	UFUNCTION(BlueprintCallable, Category="Widget")
+	bool HasFocusedDescendants() const;
+
+	/**
+	 * Checks to see if this widget is the current mouse captor
+	 *
+	 * @return  True if this widget has captured the mouse
+	 */
+	UFUNCTION(BlueprintCallable, Category="Widget")
+	bool HasMouseCapture() const;
+
 	/** Forces the underlying slate system to perform a pre-pass on the layout of the widget.  This is for advanced users. */
 	UFUNCTION(BlueprintCallable, Category="Widget")
 	void ForceLayoutPrepass();
