UMG - Exposing SBox to UMG.  Users can now employ the SizeBox widget to force particular sizing rules for thier widgets.

diff --git a/Engine/Source/Runtime/Slate/Private/Widgets/Layout/SBox.cpp b/Engine/Source/Runtime/Slate/Private/Widgets/Layout/SBox.cpp
--- a/Engine/Source/Runtime/Slate/Private/Widgets/Layout/SBox.cpp
+++ b/Engine/Source/Runtime/Slate/Private/Widgets/Layout/SBox.cpp
@@ -13,6 +13,11 @@ void SBox::Construct( const FArguments& InArgs )
 	WidthOverride = InArgs._WidthOverride;
 	HeightOverride = InArgs._HeightOverride;
 
+	MinDesiredWidth = InArgs._MinDesiredWidth;
+	MinDesiredHeight = InArgs._MinDesiredHeight;
+	MaxDesiredWidth = InArgs._MaxDesiredWidth;
+	MaxDesiredHeight = InArgs._MaxDesiredHeight;
+
 	ChildSlot
 		.HAlign( InArgs._HAlign )
 		.VAlign( InArgs._VAlign )
@@ -30,6 +35,51 @@ void SBox::SetContent(const TSharedRef< SWidget >& InContent)
 	];
 }
 
+void SBox::SetHAlign(EHorizontalAlignment HAlign)
+{
+	ChildSlot.HAlignment = HAlign;
+}
+
+void SBox::SetVAlign(EVerticalAlignment VAlign)
+{
+	ChildSlot.VAlignment = VAlign;
+}
+
+void SBox::SetPadding(const TAttribute<FMargin>& InPadding)
+{
+	ChildSlot.SlotPadding = InPadding;
+}
+
+void SBox::SetWidthOverride(TAttribute<FOptionalSize> InWidthOverride)
+{
+	WidthOverride = InWidthOverride;
+}
+
+void SBox::SetHeightOverride(TAttribute<FOptionalSize> InHeightOverride)
+{
+	HeightOverride = InHeightOverride;
+}
+
+void SBox::SetMinDesiredWidth(TAttribute<FOptionalSize> InMinDesiredWidth)
+{
+	MinDesiredWidth = InMinDesiredWidth;
+}
+
+void SBox::SetMinDesiredHeight(TAttribute<FOptionalSize> InMinDesiredHeight)
+{
+	MinDesiredHeight = InMinDesiredHeight;
+}
+
+void SBox::SetMaxDesiredWidth(TAttribute<FOptionalSize> InMaxDesiredWidth)
+{
+	MaxDesiredWidth = InMaxDesiredWidth;
+}
+
+void SBox::SetMaxDesiredHeight(TAttribute<FOptionalSize> InMaxDesiredHeight)
+{
+	MaxDesiredHeight = InMaxDesiredHeight;
+}
+
 FVector2D SBox::ComputeDesiredSize() const
 {
 	EVisibility ChildVisibility = ChildSlot.GetWidget()->GetVisibility();
@@ -40,9 +90,38 @@ FVector2D SBox::ComputeDesiredSize() const
 		const FVector2D& UnmodifiedChildDesiredSize = ChildSlot.GetWidget()->GetDesiredSize() + ChildSlot.SlotPadding.Get().GetDesiredSize();
 		const FOptionalSize CurrentWidthOverride = WidthOverride.Get();
 		const FOptionalSize CurrentHeightOverride = HeightOverride.Get();
+		const FOptionalSize CurrentMinDesiredWidth = MinDesiredWidth.Get();
+		const FOptionalSize CurrentMinDesiredHeight = MinDesiredHeight.Get();
+		const FOptionalSize CurrentMaxDesiredWidth = MaxDesiredWidth.Get();
+		const FOptionalSize CurrentMaxDesiredHeight = MaxDesiredHeight.Get();
+
+		float CurrentWidth = UnmodifiedChildDesiredSize.X;
+
+		if ( CurrentMinDesiredWidth.IsSet() )
+		{
+			CurrentWidth = FMath::Max(CurrentWidth, CurrentMinDesiredWidth.Get());
+		}
+
+		if ( CurrentMaxDesiredWidth.IsSet() )
+		{
+			CurrentWidth = FMath::Min(CurrentWidth, CurrentMaxDesiredWidth.Get());
+		}
+
+		float CurrentHeight = UnmodifiedChildDesiredSize.Y;
+
+		if ( CurrentMinDesiredHeight.IsSet() )
+		{
+			CurrentHeight = FMath::Max(CurrentHeight, CurrentMinDesiredHeight.Get());
+		}
+
+		if ( CurrentMaxDesiredHeight.IsSet() )
+		{
+			CurrentHeight = FMath::Min(CurrentHeight, CurrentMaxDesiredHeight.Get());
+		}
+
 		return FVector2D(
-			(CurrentWidthOverride.IsSet()) ? CurrentWidthOverride.Get() : UnmodifiedChildDesiredSize.X,
-			(CurrentHeightOverride.IsSet()) ? CurrentHeightOverride.Get() : UnmodifiedChildDesiredSize.Y
+			( CurrentWidthOverride.IsSet() ) ? CurrentWidthOverride.Get() : CurrentWidth,
+			( CurrentHeightOverride.IsSet() ) ? CurrentHeightOverride.Get() : CurrentHeight
 		);
 	}
 	
diff --git a/Engine/Source/Runtime/Slate/Public/Widgets/Layout/SBox.h b/Engine/Source/Runtime/Slate/Public/Widgets/Layout/SBox.h
--- a/Engine/Source/Runtime/Slate/Public/Widgets/Layout/SBox.h
+++ b/Engine/Source/Runtime/Slate/Public/Widgets/Layout/SBox.h
@@ -26,6 +26,10 @@ class SLATE_API SBox : public SPanel
 			, _Content()
 			, _WidthOverride(FOptionalSize())
 			, _HeightOverride(FOptionalSize())
+			, _MinDesiredWidth(FOptionalSize())
+			, _MinDesiredHeight(FOptionalSize())
+			, _MaxDesiredWidth(FOptionalSize())
+			, _MaxDesiredHeight(FOptionalSize())
 			{
 				_Visibility = EVisibility::SelfHitTestInvisible;
 			}
@@ -48,6 +52,18 @@ class SLATE_API SBox : public SPanel
 			/** When specified, ignore the content's desired size and report the HeightOverride as the Box's desired height. */
 			SLATE_ATTRIBUTE( FOptionalSize, HeightOverride)
 
+			/** When specified, will report the MinDesiredWidth if larger than the content's desired width. */
+			SLATE_ATTRIBUTE(FOptionalSize, MinDesiredWidth)
+
+			/** When specified, will report the MinDesiredHeight if larger than the content's desired height. */
+			SLATE_ATTRIBUTE(FOptionalSize, MinDesiredHeight)
+
+			/** When specified, will report the MaxDesiredWidth if smaller than the content's desired width. */
+			SLATE_ATTRIBUTE(FOptionalSize, MaxDesiredWidth)
+
+			/** When specified, will report the MaxDesiredHeight if smaller than the content's desired height. */
+			SLATE_ATTRIBUTE(FOptionalSize, MaxDesiredHeight)
+
 		SLATE_END_ARGS()
 
 		SBox();
@@ -64,6 +80,34 @@ class SLATE_API SBox : public SPanel
 		 * See the Content slot.
 		 */
 		void SetContent(const TSharedRef< SWidget >& InContent);
+
+		/** See HAlign argument */
+		void SetHAlign(EHorizontalAlignment HAlign);
+
+		/** See VAlign argument */
+		void SetVAlign(EVerticalAlignment VAlign);
+
+		/** See Padding attribute */
+		void SetPadding(const TAttribute<FMargin>& InPadding);
+
+		/** See WidthOverride attribute */
+		void SetWidthOverride(TAttribute<FOptionalSize> InWidthOverride);
+
+		/** See HeightOverride attribute */
+		void SetHeightOverride(TAttribute<FOptionalSize> InHeightOverride);
+
+		/** See MinDesiredWidth attribute */
+		void SetMinDesiredWidth(TAttribute<FOptionalSize> InMinDesiredWidth);
+
+		/** See MinDesiredHeight attribute */
+		void SetMinDesiredHeight(TAttribute<FOptionalSize> InMinDesiredHeight);
+
+		/** See MaxDesiredWidth attribute */
+		void SetMaxDesiredWidth(TAttribute<FOptionalSize> InMaxDesiredWidth);
+
+		/** See MaxDesiredHeight attribute */
+		void SetMaxDesiredHeight(TAttribute<FOptionalSize> InMaxDesiredHeight);
+
 private:
 
 		FBoxSlot ChildSlot;
@@ -73,7 +117,18 @@ private:
 
 		/** When specified, ignore the content's desired size and report the.HeightOverride as the Box's desired height. */
 		TAttribute<FOptionalSize> HeightOverride;
-	
+		
+		/** When specified, will report the MinDesiredWidth if larger than the content's desired width. */
+		TAttribute<FOptionalSize> MinDesiredWidth;
+
+		/** When specified, will report the MinDesiredHeight if larger than the content's desired height. */
+		TAttribute<FOptionalSize> MinDesiredHeight;
+
+		/** When specified, will report the MaxDesiredWidth if smaller than the content's desired width. */
+		TAttribute<FOptionalSize> MaxDesiredWidth;
+
+		/** When specified, will report the MaxDesiredHeight if smaller than the content's desired height. */
+		TAttribute<FOptionalSize> MaxDesiredHeight;
 };
 
 
diff --git a/Engine/Source/Runtime/UMG/Private/Components/SizeBox.cpp b/Engine/Source/Runtime/UMG/Private/Components/SizeBox.cpp
new file mode 100644
--- /dev/null
+++ b/Engine/Source/Runtime/UMG/Private/Components/SizeBox.cpp
@@ -0,0 +1,243 @@
+// Copyright 1998-2014 Epic Games, Inc. All Rights Reserved.
+
+#include "UMGPrivatePCH.h"
+#include "Slate/SlateBrushAsset.h"
+
+#define LOCTEXT_NAMESPACE "UMG"
+
+/////////////////////////////////////////////////////
+// USizeBox
+
+USizeBox::USizeBox(const FPostConstructInitializeProperties& PCIP)
+	: Super(PCIP)
+{
+	bIsVariable = false;
+}
+
+void USizeBox::ReleaseNativeWidget()
+{
+	Super::ReleaseNativeWidget();
+
+	MySizeBox.Reset();
+}
+
+TSharedRef<SWidget> USizeBox::RebuildWidget()
+{
+	MySizeBox = SNew(SBox);
+	
+	if ( GetChildrenCount() > 0 )
+	{
+		Cast<USizeBoxSlot>(GetContentSlot())->BuildSlot(MySizeBox.ToSharedRef());
+	}
+
+	return BuildDesignTimeWidget( MySizeBox.ToSharedRef() );
+}
+
+void USizeBox::SetWidthOverride(float InWidthOverride)
+{
+	bOverride_WidthOverride = true;
+	WidthOverride = InWidthOverride;
+	if ( MySizeBox.IsValid() )
+	{
+		MySizeBox->SetWidthOverride(InWidthOverride);
+	}
+}
+
+void USizeBox::ClearWidthOverride()
+{
+	bOverride_WidthOverride = false;
+	if ( MySizeBox.IsValid() )
+	{
+		MySizeBox->SetWidthOverride(FOptionalSize());
+	}
+}
+
+void USizeBox::SetHeightOverride(float InHeightOverride)
+{
+	bOverride_HeightOverride = true;
+	HeightOverride = InHeightOverride;
+	if ( MySizeBox.IsValid() )
+	{
+		MySizeBox->SetHeightOverride(InHeightOverride);
+	}
+}
+
+void USizeBox::ClearHeightOverride()
+{
+	bOverride_HeightOverride = false;
+	if ( MySizeBox.IsValid() )
+	{
+		MySizeBox->SetHeightOverride(FOptionalSize());
+	}
+}
+
+void USizeBox::SetMinDesiredWidth(float InMinDesiredWidth)
+{
+	bOverride_MinDesiredWidth = true;
+	MinDesiredWidth = InMinDesiredWidth;
+	if ( MySizeBox.IsValid() )
+	{
+		MySizeBox->SetMinDesiredWidth(InMinDesiredWidth);
+	}
+}
+
+void USizeBox::ClearMinDesiredWidth()
+{
+	bOverride_MinDesiredWidth = false;
+	if ( MySizeBox.IsValid() )
+	{
+		MySizeBox->SetMinDesiredWidth(FOptionalSize());
+	}
+}
+
+void USizeBox::SetMinDesiredHeight(float InMinDesiredHeight)
+{
+	bOverride_MinDesiredHeight = true;
+	MinDesiredHeight = InMinDesiredHeight;
+	if ( MySizeBox.IsValid() )
+	{
+		MySizeBox->SetMinDesiredHeight(InMinDesiredHeight);
+	}
+}
+
+void USizeBox::ClearMinDesiredHeight()
+{
+	bOverride_MinDesiredHeight = false;
+	if ( MySizeBox.IsValid() )
+	{
+		MySizeBox->SetMinDesiredHeight(FOptionalSize());
+	}
+}
+
+void USizeBox::SetMaxDesiredWidth(float InMaxDesiredWidth)
+{
+	bOverride_MaxDesiredWidth = true;
+	MaxDesiredWidth = InMaxDesiredWidth;
+	if ( MySizeBox.IsValid() )
+	{
+		MySizeBox->SetMaxDesiredWidth(InMaxDesiredWidth);
+	}
+}
+
+void USizeBox::ClearMaxDesiredWidth()
+{
+	bOverride_MaxDesiredWidth = false;
+	if ( MySizeBox.IsValid() )
+	{
+		MySizeBox->SetMaxDesiredWidth(FOptionalSize());
+	}
+}
+
+void USizeBox::SetMaxDesiredHeight(float InMaxDesiredHeight)
+{
+	bOverride_MaxDesiredHeight = true;
+	MaxDesiredHeight = InMaxDesiredHeight;
+	if ( MySizeBox.IsValid() )
+	{
+		MySizeBox->SetMaxDesiredHeight(InMaxDesiredHeight);
+	}
+}
+
+void USizeBox::ClearMaxDesiredHeight()
+{
+	bOverride_MaxDesiredHeight = false;
+	if ( MySizeBox.IsValid() )
+	{
+		MySizeBox->SetMaxDesiredHeight(FOptionalSize());
+	}
+}
+
+void USizeBox::SyncronizeProperties()
+{
+	Super::SyncronizeProperties();
+
+	if ( bOverride_WidthOverride )
+	{
+		SetWidthOverride(WidthOverride);
+	}
+	else
+	{
+		ClearWidthOverride();
+	}
+
+	if ( bOverride_HeightOverride )
+	{
+		SetHeightOverride(HeightOverride);
+	}
+	else
+	{
+		ClearHeightOverride();
+	}
+
+	if ( bOverride_MinDesiredWidth )
+	{
+		SetMinDesiredWidth(MinDesiredWidth);
+	}
+	else
+	{
+		ClearMinDesiredWidth();
+	}
+
+	if ( bOverride_MinDesiredHeight )
+	{
+		SetMinDesiredHeight(MinDesiredHeight);
+	}
+	else
+	{
+		ClearMinDesiredHeight();
+	}
+
+	if ( bOverride_MaxDesiredWidth )
+	{
+		SetMaxDesiredWidth(MaxDesiredWidth);
+	}
+	else
+	{
+		ClearMaxDesiredWidth();
+	}
+
+	if ( bOverride_MaxDesiredHeight )
+	{
+		SetMaxDesiredHeight(MaxDesiredHeight);
+	}
+	else
+	{
+		ClearMaxDesiredHeight();
+	}
+}
+
+UClass* USizeBox::GetSlotClass() const
+{
+	return USizeBoxSlot::StaticClass();
+}
+
+void USizeBox::OnSlotAdded(UPanelSlot* Slot)
+{
+	// Add the child to the live slot if it already exists
+	if ( MySizeBox.IsValid() )
+	{
+		Cast<USizeBoxSlot>(Slot)->BuildSlot(MySizeBox.ToSharedRef());
+	}
+}
+
+void USizeBox::OnSlotRemoved(UPanelSlot* Slot)
+{
+	// Remove the widget from the live slot if it exists.
+	if ( MySizeBox.IsValid() )
+	{
+		MySizeBox->SetContent(SNullWidget::NullWidget);
+	}
+}
+
+#if WITH_EDITOR
+
+const FSlateBrush* USizeBox::GetEditorIcon()
+{
+	return FUMGStyle::Get().GetBrush("Widget.SizeBox");
+}
+
+#endif
+
+/////////////////////////////////////////////////////
+
+#undef LOCTEXT_NAMESPACE
diff --git a/Engine/Source/Runtime/UMG/Private/Components/SizeBoxSlot.cpp b/Engine/Source/Runtime/UMG/Private/Components/SizeBoxSlot.cpp
new file mode 100644
--- /dev/null
+++ b/Engine/Source/Runtime/UMG/Private/Components/SizeBoxSlot.cpp
@@ -0,0 +1,67 @@
+// Copyright 1998-2014 Epic Games, Inc. All Rights Reserved.
+
+#include "UMGPrivatePCH.h"
+
+/////////////////////////////////////////////////////
+// USizeBoxSlot
+
+USizeBoxSlot::USizeBoxSlot(const FPostConstructInitializeProperties& PCIP)
+	: Super(PCIP)
+{
+	Padding = FMargin(0, 0);
+
+	HorizontalAlignment = HAlign_Fill;
+	VerticalAlignment = VAlign_Fill;
+
+	SBox::FArguments SizeBoxDefaults;
+}
+
+void USizeBoxSlot::ReleaseNativeWidget()
+{
+	Super::ReleaseNativeWidget();
+
+	SizeBox.Reset();
+}
+
+void USizeBoxSlot::BuildSlot(TSharedRef<SBox> InSizeBox)
+{
+	SizeBox = InSizeBox;
+
+	SyncronizeProperties();
+
+	SizeBox->SetContent(Content ? Content->TakeWidget() : SNullWidget::NullWidget);
+}
+
+void USizeBoxSlot::SetPadding(FMargin InPadding)
+{
+	Padding = InPadding;
+	if ( SizeBox.IsValid() )
+	{
+		SizeBox->SetPadding(InPadding);
+	}
+}
+
+void USizeBoxSlot::SetHorizontalAlignment(EHorizontalAlignment InHorizontalAlignment)
+{
+	HorizontalAlignment = InHorizontalAlignment;
+	if ( SizeBox.IsValid() )
+	{
+		SizeBox->SetHAlign(InHorizontalAlignment);
+	}
+}
+
+void USizeBoxSlot::SetVerticalAlignment(EVerticalAlignment InVerticalAlignment)
+{
+	VerticalAlignment = InVerticalAlignment;
+	if ( SizeBox.IsValid() )
+	{
+		SizeBox->SetVAlign(InVerticalAlignment);
+	}
+}
+
+void USizeBoxSlot::SyncronizeProperties()
+{
+	SetPadding(Padding);
+	SetHorizontalAlignment(HorizontalAlignment);
+	SetVerticalAlignment(VerticalAlignment);
+}
diff --git a/Engine/Source/Runtime/UMG/Public/Components/SizeBox.h b/Engine/Source/Runtime/UMG/Public/Components/SizeBox.h
new file mode 100644
--- /dev/null
+++ b/Engine/Source/Runtime/UMG/Public/Components/SizeBox.h
@@ -0,0 +1,142 @@
+// Copyright 1998-2014 Epic Games, Inc. All Rights Reserved.
+
+#pragma once
+
+#include "SizeBox.generated.h"
+
+class USlateBrushAsset;
+
+/**
+ * A SizeBox is a container widget that can contain one child widget, providing an opportunity 
+ * to surround it with a SizeBox image and adjustable padding.
+ */
+UCLASS(meta=( Category="Common" ), ClassGroup=UserInterface)
+class UMG_API USizeBox : public UContentWidget
+{
+	GENERATED_UCLASS_BODY()
+
+public:
+		
+	/** When specified, ignore the content's desired size and report the WidthOverride as the Box's desired width. */
+	UFUNCTION(BlueprintCallable, Category="Layout")
+	void SetWidthOverride(float InWidthOverride);
+
+	UFUNCTION(BlueprintCallable, Category="Layout")
+	void ClearWidthOverride();
+
+	/** When specified, ignore the content's desired size and report the HeightOverride as the Box's desired height. */
+	UFUNCTION(BlueprintCallable, Category="Layout")
+	void SetHeightOverride(float InHeightOverride);
+
+	UFUNCTION(BlueprintCallable, Category="Layout")
+	void ClearHeightOverride();
+
+	/** When specified, will report the MinDesiredWidth if larger than the content's desired width. */
+	UFUNCTION(BlueprintCallable, Category="Layout")
+	void SetMinDesiredWidth(float InMinDesiredWidth);
+
+	UFUNCTION(BlueprintCallable, Category="Layout")
+	void ClearMinDesiredWidth();
+
+	/** When specified, will report the MinDesiredHeight if larger than the content's desired height. */
+	UFUNCTION(BlueprintCallable, Category="Layout")
+	void SetMinDesiredHeight(float InMinDesiredHeight);
+
+	UFUNCTION(BlueprintCallable, Category="Layout")
+	void ClearMinDesiredHeight();
+
+	/** When specified, will report the MaxDesiredWidth if smaller than the content's desired width. */
+	UFUNCTION(BlueprintCallable, Category="Layout")
+	void SetMaxDesiredWidth(float InMaxDesiredWidth);
+
+	UFUNCTION(BlueprintCallable, Category="Layout")
+	void ClearMaxDesiredWidth();
+
+	/** When specified, will report the MaxDesiredHeight if smaller than the content's desired height. */
+	UFUNCTION(BlueprintCallable, Category="Layout")
+	void SetMaxDesiredHeight(float InMaxDesiredHeight);
+
+	UFUNCTION(BlueprintCallable, Category="Layout")
+	void ClearMaxDesiredHeight();
+
+public:
+
+	// UWidget interface
+	virtual void SyncronizeProperties() override;
+	// End of UWidget interface
+
+	// UVisual interface
+	virtual void ReleaseNativeWidget() override;
+	// End of UVisual interface
+
+#if WITH_EDITOR
+	virtual const FSlateBrush* GetEditorIcon() override;
+#endif
+
+protected:
+
+	/**  */
+	UPROPERTY()
+	uint32 bOverride_WidthOverride : 1;
+
+	/**  */
+	UPROPERTY()
+	uint32 bOverride_HeightOverride : 1;
+
+	/**  */
+	UPROPERTY()
+	uint32 bOverride_MinDesiredWidth : 1;
+
+	/**  */
+	UPROPERTY()
+	uint32 bOverride_MinDesiredHeight : 1;
+
+	/**  */
+	UPROPERTY()
+	uint32 bOverride_MaxDesiredWidth : 1;
+
+	/**  */
+	UPROPERTY()
+	uint32 bOverride_MaxDesiredHeight : 1;
+
+
+	/** When specified, ignore the content's desired size and report the WidthOverride as the Box's desired width. */
+	UPROPERTY(EditDefaultsOnly, Category=Layout, meta=( editcondition="bOverride_WidthOverride" ))
+	float WidthOverride;
+
+	/** When specified, ignore the content's desired size and report the HeightOverride as the Box's desired height. */
+	UPROPERTY(EditDefaultsOnly, Category=Layout, meta=( editcondition="bOverride_HeightOverride" ))
+	float HeightOverride;
+
+	/** When specified, will report the MinDesiredWidth if larger than the content's desired width. */
+	UPROPERTY(EditDefaultsOnly, Category=Layout, meta=( editcondition="bOverride_MinDesiredWidth" ))
+	float MinDesiredWidth;
+
+	/** When specified, will report the MinDesiredHeight if larger than the content's desired height. */
+	UPROPERTY(EditDefaultsOnly, Category=Layout, meta=( editcondition="bOverride_MinDesiredHeight" ))
+	float MinDesiredHeight;
+
+	/** When specified, will report the MaxDesiredWidth if smaller than the content's desired width. */
+	UPROPERTY(EditDefaultsOnly, Category=Layout, meta=( editcondition="bOverride_MaxDesiredWidth" ))
+	float MaxDesiredWidth;
+
+	/** When specified, will report the MaxDesiredHeight if smaller than the content's desired height. */
+	UPROPERTY(EditDefaultsOnly, Category=Layout, meta=( editcondition="bOverride_MaxDesiredHeight" ))
+	float MaxDesiredHeight;
+
+protected:
+
+	// UPanelWidget
+	virtual UClass* GetSlotClass() const override;
+	virtual void OnSlotAdded(UPanelSlot* Slot) override;
+	virtual void OnSlotRemoved(UPanelSlot* Slot) override;
+	// End UPanelWidget
+
+protected:
+	TSharedPtr<SBox> MySizeBox;
+
+protected:
+	// UWidget interface
+	virtual TSharedRef<SWidget> RebuildWidget() override;
+	// End of UWidget interface
+};
diff --git a/Engine/Source/Runtime/UMG/Public/Components/SizeBoxSlot.h b/Engine/Source/Runtime/UMG/Public/Components/SizeBoxSlot.h
new file mode 100644
--- /dev/null
+++ b/Engine/Source/Runtime/UMG/Public/Components/SizeBoxSlot.h
@@ -0,0 +1,54 @@
+// Copyright 1998-2014 Epic Games, Inc. All Rights Reserved.
+
+#pragma once
+
+#include "SlateWrapperTypes.h"
+
+#include "SizeBoxSlot.generated.h"
+
+/** The Slot for the USizeBoxSlot, contains the widget displayed in a button's single slot */
+UCLASS()
+class UMG_API USizeBoxSlot : public UPanelSlot
+{
+	GENERATED_UCLASS_BODY()
+
+public:
+
+	UFUNCTION(BlueprintCallable, Category="Layout")
+	void SetPadding(FMargin InPadding);
+
+	UFUNCTION(BlueprintCallable, Category="Layout")
+	void SetHorizontalAlignment(EHorizontalAlignment InHorizontalAlignment);
+
+	UFUNCTION(BlueprintCallable, Category="Layout")
+	void SetVerticalAlignment(EVerticalAlignment InVerticalAlignment);
+
+public:
+
+	// UPanelSlot interface
+	virtual void SyncronizeProperties() override;
+	// End of UPanelSlot interface
+
+	/** Builds the underlying slot for the slate SizeBox. */
+	void BuildSlot(TSharedRef<SBox> InSizeBox);
+
+	virtual void ReleaseNativeWidget() override;
+
+protected:
+	/** The padding area between the slot and the content it contains. */
+	UPROPERTY(EditDefaultsOnly, Category=Layout)
+	FMargin Padding;
+
+	/** The alignment of the object horizontally. */
+	UPROPERTY(EditDefaultsOnly, Category=Layout)
+	TEnumAsByte<EHorizontalAlignment> HorizontalAlignment;
+
+	/** The alignment of the object vertically. */
+	UPROPERTY(EditDefaultsOnly, Category=Layout)
+	TEnumAsByte<EVerticalAlignment> VerticalAlignment;
+
+private:
+
+	/** A pointer to the button to allow us to adjust the size, padding...etc at runtime. */
+	TSharedPtr<SBox> SizeBox;
+};
diff --git a/Engine/Source/Runtime/UMG/Public/UMG.h b/Engine/Source/Runtime/UMG/Public/UMG.h
--- a/Engine/Source/Runtime/UMG/Public/UMG.h
+++ b/Engine/Source/Runtime/UMG/Public/UMG.h
@@ -65,6 +65,9 @@
 #include "ButtonSlot.h"
 #include "Button.h"
 
+#include "SizeBoxSlot.h"
+#include "SizeBox.h"
+
 #include "TextBlock.h"
 #include "CheckBox.h"
 #include "Image.h"
