UMG - General cleanup of the horizontal box.

diff --git a/Engine/Source/Runtime/UMG/Private/Components/HorizontalBox.cpp b/Engine/Source/Runtime/UMG/Private/Components/HorizontalBox.cpp
--- a/Engine/Source/Runtime/UMG/Private/Components/HorizontalBox.cpp
+++ b/Engine/Source/Runtime/UMG/Private/Components/HorizontalBox.cpp
@@ -48,11 +48,16 @@ void UHorizontalBox::OnSlotRemoved(UPanelSlot* Slot)
 	}
 }
 
-UHorizontalBoxSlot* UHorizontalBox::Add(UWidget* Content)
+UHorizontalBoxSlot* UHorizontalBox::AddSlot(UWidget* Content)
 {
 	return Cast<UHorizontalBoxSlot>( Super::AddChild(Content) );
 }
 
+UHorizontalBoxSlot* UHorizontalBox::HorizontalBoxSlot(UWidget* Child)
+{
+	return Cast<UHorizontalBoxSlot>(Child->Slot);
+}
+
 TSharedRef<SWidget> UHorizontalBox::RebuildWidget()
 {
 	MyHorizontalBox = SNew(SHorizontalBox);
diff --git a/Engine/Source/Runtime/UMG/Public/Components/HorizontalBox.h b/Engine/Source/Runtime/UMG/Public/Components/HorizontalBox.h
--- a/Engine/Source/Runtime/UMG/Public/Components/HorizontalBox.h
+++ b/Engine/Source/Runtime/UMG/Public/Components/HorizontalBox.h
@@ -11,7 +11,11 @@ class UMG_API UHorizontalBox : public UPanelWidget
 
 	/**  */
 	UFUNCTION(BlueprintCallable, Category="Widget")
-	UHorizontalBoxSlot* Add(UWidget* Content);
+	UHorizontalBoxSlot* AddSlot(UWidget* Content);
+
+	/**  */
+	UFUNCTION(BlueprintCallable, Category="Slot")
+	static UHorizontalBoxSlot* HorizontalBoxSlot(UWidget* Child);
 
 #if WITH_EDITOR
 	// UWidget interface
