UMG - Adding additional functions for changing focus to widgets and changing focus to the viewport.

diff --git a/Engine/Source/Runtime/UMG/Private/Components/Widget.cpp b/Engine/Source/Runtime/UMG/Private/Components/Widget.cpp
--- a/Engine/Source/Runtime/UMG/Private/Components/Widget.cpp
+++ b/Engine/Source/Runtime/UMG/Private/Components/Widget.cpp
@@ -98,6 +98,14 @@ bool UWidget::HasMouseCapture() const
 	return false;
 }
 
+void UWidget::SetKeyboardFocus() const
+{
+	if ( MyWidget.IsValid() )
+	{
+		FSlateApplication::Get().SetKeyboardFocus(MyWidget.Pin());
+	}
+}
+
 void UWidget::ForceLayoutPrepass()
 {
 	if ( MyWidget.IsValid() )
diff --git a/Engine/Source/Runtime/UMG/Private/WidgetBlueprintLibrary.cpp b/Engine/Source/Runtime/UMG/Private/WidgetBlueprintLibrary.cpp
--- a/Engine/Source/Runtime/UMG/Private/WidgetBlueprintLibrary.cpp
+++ b/Engine/Source/Runtime/UMG/Private/WidgetBlueprintLibrary.cpp
@@ -27,6 +27,11 @@ UUserWidget* UWidgetBlueprintLibrary::Create(UObject* WorldContextObject, TSubcl
 	}
 }
 
+void UWidgetBlueprintLibrary::SetFocusToGameViewport()
+{
+	FSlateApplication::Get().SetFocusToGameViewport();
+}
+
 void UWidgetBlueprintLibrary::DrawBox(UPARAM(ref) FPaintContext& Context, FVector2D Position, FVector2D Size, USlateBrushAsset* Brush, FLinearColor Tint)
 {
 	Context.MaxLayer++;
diff --git a/Engine/Source/Runtime/UMG/Public/Blueprint/WidgetBlueprintLibrary.h b/Engine/Source/Runtime/UMG/Public/Blueprint/WidgetBlueprintLibrary.h
--- a/Engine/Source/Runtime/UMG/Public/Blueprint/WidgetBlueprintLibrary.h
+++ b/Engine/Source/Runtime/UMG/Public/Blueprint/WidgetBlueprintLibrary.h
@@ -12,6 +12,9 @@ class UWidgetBlueprintLibrary : public UBlueprintFunctionLibrary
 	UFUNCTION(BlueprintCallable, meta=( HidePin="WorldContextObject", DefaultToSelf="WorldContextObject", FriendlyName = "Create Widget" ), Category="User Interface|Widget")
 	static class UUserWidget* Create(UObject* WorldContextObject, TSubclassOf<class UUserWidget> WidgetType, APlayerController* OwningPlayer);
 
+	UFUNCTION(BlueprintCallable, Category="Focus")
+	static void SetFocusToGameViewport();
+
 	/** Draws a box */
 	UFUNCTION(BlueprintCallable, Category="Painting")
 	static void DrawBox(UPARAM(ref) FPaintContext& Context, FVector2D Position, FVector2D Size, USlateBrushAsset* Brush, FLinearColor Tint = FLinearColor::White);
diff --git a/Engine/Source/Runtime/UMG/Public/Components/Widget.h b/Engine/Source/Runtime/UMG/Public/Components/Widget.h
--- a/Engine/Source/Runtime/UMG/Public/Components/Widget.h
+++ b/Engine/Source/Runtime/UMG/Public/Components/Widget.h
@@ -149,6 +149,10 @@ public:
 	UFUNCTION(BlueprintCallable, Category="Widget")
 	bool HasMouseCapture() const;
 
+	/**  */
+	UFUNCTION(BlueprintCallable, Category="Widget")
+	void SetKeyboardFocus() const;
+
 	/** Forces the underlying slate system to perform a pre-pass on the layout of the widget.  This is for advanced users. */
 	UFUNCTION(BlueprintCallable, Category="Widget")
 	void ForceLayoutPrepass();
