UMG - Fixing adding to the viewport to not replace the viewport activation widget if the widget is not moda.  This fixes an issue with mouse locking never working with UMG widgets on the screenl.  Also making the modal flag work to block input even when there's no background border.

diff --git a/Engine/Source/Runtime/UMG/Private/UserWidget.cpp b/Engine/Source/Runtime/UMG/Private/UserWidget.cpp
--- a/Engine/Source/Runtime/UMG/Private/UserWidget.cpp
+++ b/Engine/Source/Runtime/UMG/Private/UserWidget.cpp
@@ -24,6 +24,8 @@ class SViewportWidgetHost : public SCompoundWidget
 	{
 		bModal = bInModal;
 
+		SetVisibility(bModal ? EVisibility::Visible : EVisibility::SelfHitTestInvisible);
+
 		ChildSlot
 		[
 			InArgs._Content.Widget
@@ -101,6 +103,8 @@ class SViewportWidgetHost : public SCompoundWidget
 		return bModal ? FReply::Handled() : SCompoundWidget::OnMouseWheel(MyGeometry, MouseEvent);
 	}
 
+	bool IsModal() const { return bModal; }
+
 protected:
 	bool bModal;
 };
@@ -396,11 +400,16 @@ void UUserWidget::AddToViewport(bool bAbsoluteLayout, bool bModal, bool bShowCur
 			UGameViewportClient* Viewport = World->GetGameViewport();
 			Viewport->AddViewportWidgetContent(WidgetHost);
 
-			TWeakPtr<SViewport> GameViewportWidget = Viewport->GetGameViewport()->GetViewportWidget();
-			if ( GameViewportWidget.IsValid() )
+			//TODO UMG this isn't what should manage focus, a higher level window controller, probably the viewport needs to understand
+			// the Widget stack, and the dialog stack.
+			if ( bModal )
 			{
-				GameViewportWidget.Pin()->SetWidgetToFocusOnActivate(RootWidget);
-				FSlateApplication::Get().SetKeyboardFocus(RootWidget);
+				TWeakPtr<SViewport> GameViewportWidget = Viewport->GetGameViewport()->GetViewportWidget();
+				if ( GameViewportWidget.IsValid() )
+				{
+					GameViewportWidget.Pin()->SetWidgetToFocusOnActivate(RootWidget);
+					FSlateApplication::Get().SetKeyboardFocus(RootWidget);
+				}
 			}
 		}
 	}
@@ -410,7 +419,7 @@ void UUserWidget::RemoveFromViewport()
 {
 	if ( FullScreenWidget.IsValid() )
 	{
-		TSharedPtr<SWidget> RootWidget = FullScreenWidget.Pin();
+		TSharedPtr<SViewportWidgetHost> RootWidget = FullScreenWidget.Pin();
 
 		// If this is a game world add the widget to the current worlds viewport.
 		UWorld* World = GetWorld();
@@ -419,13 +428,19 @@ void UUserWidget::RemoveFromViewport()
 			UGameViewportClient* Viewport = World->GetGameViewport();
 			Viewport->RemoveViewportWidgetContent(RootWidget.ToSharedRef());
 
-			TWeakPtr<SViewport> GameViewportWidget = Viewport->GetGameViewport()->GetViewportWidget();
-			if ( GameViewportWidget.IsValid() )
+			if ( RootWidget->IsModal() )
 			{
-				//TODO UMG this isn't what should manage focus, a higher level window controller, probably the viewport needs to understand
-				// the Widget stack, and the dialog stack.
-				GameViewportWidget.Pin()->ClearWidgetToFocusOnActivate();
-				FSlateApplication::Get().SetKeyboardFocus(TSharedPtr<SWidget>());
+				TWeakPtr<SViewport> GameViewportWidget = Viewport->GetGameViewport()->GetViewportWidget();
+				if ( GameViewportWidget.IsValid() )
+				{
+					//TODO UMG this isn't what should manage focus, a higher level window controller, probably the viewport needs to understand
+					// the Widget stack, and the dialog stack.
+					GameViewportWidget.Pin()->ClearWidgetToFocusOnActivate();
+					FSlateApplication::Get().ClearKeyboardFocus(EKeyboardFocusCause::SetDirectly);
+
+					FSlateApplication::Get().SetFocusToGameViewport();
+					FSlateApplication::Get().SetJoystickCaptorToGameViewport();
+				}
 			}
 		}
 	}
diff --git a/Engine/Source/Runtime/UMG/Public/Blueprint/UserWidget.h b/Engine/Source/Runtime/UMG/Public/Blueprint/UserWidget.h
--- a/Engine/Source/Runtime/UMG/Public/Blueprint/UserWidget.h
+++ b/Engine/Source/Runtime/UMG/Public/Blueprint/UserWidget.h
@@ -342,7 +342,7 @@ private:
 	FVector2D ViewportAlignment;
 	int32 ViewportZOrder;
 
-	TWeakPtr<SWidget> FullScreenWidget;
+	TWeakPtr<class SViewportWidgetHost> FullScreenWidget;
 
 	FLocalPlayerContext PlayerContext;
 
