UMG - No longer connecting bindable properties at design time.

diff --git a/Engine/Source/Runtime/UMG/Public/Components/Widget.h b/Engine/Source/Runtime/UMG/Public/Components/Widget.h
--- a/Engine/Source/Runtime/UMG/Public/Components/Widget.h
+++ b/Engine/Source/Runtime/UMG/Public/Components/Widget.h
@@ -9,13 +9,13 @@
 /**
  * Helper macro for binding to a delegate or using the constant value when constructing the underlying SWidget
  */
-#define OPTIONAL_BINDING(ReturnType, MemberName) MemberName ## Delegate.IsBound() ? TAttribute< ReturnType >::Create(MemberName ## Delegate.GetUObject(), MemberName ## Delegate.GetFunctionName()) : TAttribute< ReturnType >(MemberName)
+#define OPTIONAL_BINDING(ReturnType, MemberName) ( MemberName ## Delegate.IsBound() && !IsDesignTime() ) ? TAttribute< ReturnType >::Create(MemberName ## Delegate.GetUObject(), MemberName ## Delegate.GetFunctionName()) : TAttribute< ReturnType >(MemberName)
 
 /**
  * Helper macro for binding to a delegate or using the constant value when constructing the underlying SWidget,
  * also allows a conversion function to be provided to convert between the SWidget value and the value exposed to UMG.
  */
-#define OPTIONAL_BINDING_CONVERT(ReturnType, MemberName, ConvertedType, ConversionFunction) MemberName ## Delegate.IsBound() ? TAttribute< ConvertedType >::Create(TAttribute< ConvertedType >::FGetter::CreateUObject(this, &ThisClass::ConversionFunction, TAttribute< ReturnType >::Create(MemberName ## Delegate.GetUObject(), MemberName ## Delegate.GetFunctionName()))) : ConversionFunction(TAttribute< ReturnType >(MemberName))
+#define OPTIONAL_BINDING_CONVERT(ReturnType, MemberName, ConvertedType, ConversionFunction) ( MemberName ## Delegate.IsBound() && !IsDesignTime() ) ? TAttribute< ConvertedType >::Create(TAttribute< ConvertedType >::FGetter::CreateUObject(this, &ThisClass::ConversionFunction, TAttribute< ReturnType >::Create(MemberName ## Delegate.GetUObject(), MemberName ## Delegate.GetFunctionName()))) : ConversionFunction(TAttribute< ReturnType >(MemberName))
 
 /**
  * This is the base class for all wrapped Slate controls that are exposed to UMG.
