UMG/Slate - Breaking Change - Introducing a DPI scaling mechanism to both UMG and native slate widgets.  The container widget of the game's viewport is now a DPI scaler that is controlled through a curve in the rendering settings for your project.  By default the curve assumes a DPI scale of 1 at 1080, and is a linear scale up or down from that point.  The math is performed based on the short side of the viewport.

If you are currentlying performing your own DPI Scaling of slate you'll need to remove that code or adjust the DPI Curve of this new system to 1 and keep your current system for scaling the UI.

diff --git a/Engine/Config/BaseEngine.ini b/Engine/Config/BaseEngine.ini
--- a/Engine/Config/BaseEngine.ini
+++ b/Engine/Config/BaseEngine.ini
@@ -1309,6 +1309,9 @@ InGameAdManagerClassName=
 RenderPriorityPS3=1001
 SuspendGameIO=True
 
+[/Script/Engine.RendererSettings]
+UIScaleCurve=(EditorCurveData=(Keys=((Time=480,Value=0.444),(Time=720,Value=0.666),(Time=1080,Value=1.0),(Time=8640,Value=8.0))),ExternalCurve=None)
+
 [/Script/Engine.GameEngine]
 MaxDeltaTime=0
 
diff --git a/Engine/Source/Editor/DetailCustomizations/Private/CurveStructCustomization.cpp b/Engine/Source/Editor/DetailCustomizations/Private/CurveStructCustomization.cpp
--- a/Engine/Source/Editor/DetailCustomizations/Private/CurveStructCustomization.cpp
+++ b/Engine/Source/Editor/DetailCustomizations/Private/CurveStructCustomization.cpp
@@ -55,25 +55,20 @@ void FCurveStructCustomization::CustomizeHeader( TSharedRef<class IPropertyHandl
 				InStructPropertyHandle->CreatePropertyNameWidget( TEXT( "" ), false )
 			]
 			.ValueContent()
-			.MinDesiredWidth(0.f)
-			.MaxDesiredWidth(0.f)
+			.HAlign(HAlign_Fill)
+			.MinDesiredWidth(200)
 			[
-				SNew(SVerticalBox)
-				+SVerticalBox::Slot()
-				.AutoHeight()
+				SNew(SBorder)
+				.VAlign(VAlign_Fill)
+				.OnMouseDoubleClick(this, &FCurveStructCustomization::OnCurvePreviewDoubleClick)
 				[
-					SNew(SBorder)
-					.VAlign(VAlign_Fill)
-					.OnMouseDoubleClick(this, &FCurveStructCustomization::OnCurvePreviewDoubleClick)
-					[
-						SAssignNew(CurveWidget, SCurveEditor)
-						.ViewMinInput(this, &FCurveStructCustomization::GetViewMinInput)
-						.ViewMaxInput(this, &FCurveStructCustomization::GetViewMaxInput)
-						.TimelineLength(this, &FCurveStructCustomization::GetTimelineLength)
-						.OnSetInputViewRange(this, &FCurveStructCustomization::SetInputViewRange)
-						.HideUI(false)
-						.DesiredSize(FVector2D(128, 128))
-					]
+					SAssignNew(CurveWidget, SCurveEditor)
+					.ViewMinInput(this, &FCurveStructCustomization::GetViewMinInput)
+					.ViewMaxInput(this, &FCurveStructCustomization::GetViewMaxInput)
+					.TimelineLength(this, &FCurveStructCustomization::GetTimelineLength)
+					.OnSetInputViewRange(this, &FCurveStructCustomization::SetInputViewRange)
+					.HideUI(false)
+					.DesiredSize(FVector2D(300, 150))
 				]
 			];
 
@@ -119,7 +114,7 @@ void FCurveStructCustomization::CustomizeChildren( TSharedRef<class IPropertyHan
 		{
 			ExternalCurveHandle = Child;
 
-			FSimpleDelegate OnCurveChangedDelegate = FSimpleDelegate::CreateSP( this, &FCurveStructCustomization::OnExternalCurveChanged );
+			FSimpleDelegate OnCurveChangedDelegate = FSimpleDelegate::CreateSP( this, &FCurveStructCustomization::OnExternalCurveChanged, InStructPropertyHandle );
 			Child->SetOnPropertyValueChanged(OnCurveChangedDelegate);
 
 			StructBuilder.AddChildContent(TEXT("ExternalCurve"))
@@ -223,6 +218,11 @@ void FCurveStructCustomization::MakeTransactional()
 	}
 }
 
+void FCurveStructCustomization::OnCurveChanged()
+{
+	StructPropertyHandle->NotifyPostChange();
+}
+
 float FCurveStructCustomization::GetTimelineLength() const
 {
 	return 0.f;
@@ -234,7 +234,7 @@ void FCurveStructCustomization::SetInputViewRange(float InViewMinInput, float In
 	ViewMinInput = InViewMinInput;
 }
 
-void FCurveStructCustomization::OnExternalCurveChanged()
+void FCurveStructCustomization::OnExternalCurveChanged(TSharedRef<class IPropertyHandle> CurvePropertyHandle)
 {
 	if (RuntimeCurve)
 	{
@@ -246,6 +246,8 @@ void FCurveStructCustomization::OnExternalCurveChanged()
 		{
 			CurveWidget->SetCurveOwner(this);
 		}
+
+		CurvePropertyHandle->NotifyPostChange();
 	}
 }
 
diff --git a/Engine/Source/Editor/DetailCustomizations/Private/CurveStructCustomization.h b/Engine/Source/Editor/DetailCustomizations/Private/CurveStructCustomization.h
--- a/Engine/Source/Editor/DetailCustomizations/Private/CurveStructCustomization.h
+++ b/Engine/Source/Editor/DetailCustomizations/Private/CurveStructCustomization.h
@@ -27,6 +27,7 @@ public:
 	virtual UObject* GetOwner() override;
 	virtual void ModifyOwner() override;
 	virtual void MakeTransactional() override;
+	virtual void OnCurveChanged() override;
 
 private:
 	/**
@@ -57,7 +58,7 @@ private:
 	/**
 	 * Called when RuntimeFloatCurve's External Curve is changed
 	 */
-	void OnExternalCurveChanged();
+	void OnExternalCurveChanged(TSharedRef<class IPropertyHandle> CurvePropertyHandle);
 
 	/**
 	 * Called when button clicked to create an External Curve
diff --git a/Engine/Source/Editor/DetailCustomizations/Private/WheeledVehicleMovementComponent4WDetails.h b/Engine/Source/Editor/DetailCustomizations/Private/WheeledVehicleMovementComponent4WDetails.h
--- a/Engine/Source/Editor/DetailCustomizations/Private/WheeledVehicleMovementComponent4WDetails.h
+++ b/Engine/Source/Editor/DetailCustomizations/Private/WheeledVehicleMovementComponent4WDetails.h
@@ -24,6 +24,7 @@ private:
 		virtual UObject* GetOwner() override;
 		virtual void ModifyOwner() override;
 		virtual void MakeTransactional() override;
+		virtual void OnCurveChanged() override { }
 
 	private:
 		UWheeledVehicleMovementComponent4W * VehicleComponent;
@@ -40,6 +41,7 @@ private:
 		virtual UObject* GetOwner() override;
 		virtual void ModifyOwner() override;
 		virtual void MakeTransactional() override;
+		virtual void OnCurveChanged() override { }
 
 	private:
 		UWheeledVehicleMovementComponent4W * VehicleComponent;
diff --git a/Engine/Source/Editor/LevelEditor/Private/SLevelViewport.cpp b/Engine/Source/Editor/LevelEditor/Private/SLevelViewport.cpp
--- a/Engine/Source/Editor/LevelEditor/Private/SLevelViewport.cpp
+++ b/Engine/Source/Editor/LevelEditor/Private/SLevelViewport.cpp
@@ -158,7 +158,11 @@ void SLevelViewport::ConstructViewportOverlayContent()
 #if ALLOW_PLAY_IN_VIEWPORT_GAMEUI
 		ViewportOverlay->AddSlot( SlotIndex )
 		[
-			PIEViewportOverlayWidget.ToSharedRef()
+			SNew(SDPIScaler)
+			.DPIScale(this, &SLevelViewport::GetGameViewportDPIScale)
+			[
+				PIEViewportOverlayWidget.ToSharedRef()
+			]
 		];
 
 		++SlotIndex;
@@ -348,6 +352,16 @@ void SLevelViewport::ConstructLevelEditorViewportClient( const FArguments& InArg
 	LevelViewportClient->SetViewModes(ViewportInstanceSettings.PerspViewModeIndex, ViewportInstanceSettings.OrthoViewModeIndex );
 }
 
+float SLevelViewport::GetGameViewportDPIScale() const
+{
+	if ( HasPlayInEditorViewport() || LevelViewportClient->IsSimulateInEditorViewport() )
+	{
+		return GetDefault<URendererSettings>(URendererSettings::StaticClass())->GetDPIScaleBasedOnSize( ActiveViewport->GetSize() );
+	}
+
+	return 1.0f;
+}
+
 FReply SLevelViewport::OnKeyDown( const FGeometry& MyGeometry, const FKeyboardEvent& InKeyboardEvent )
 {
 	FReply Reply = FReply::Unhandled();
diff --git a/Engine/Source/Editor/LevelEditor/Public/SLevelViewport.h b/Engine/Source/Editor/LevelEditor/Public/SLevelViewport.h
--- a/Engine/Source/Editor/LevelEditor/Public/SLevelViewport.h
+++ b/Engine/Source/Editor/LevelEditor/Public/SLevelViewport.h
@@ -642,6 +642,10 @@ private:
 	/** Hide the mouse capture label */
 	void HideMouseCaptureLabel();
 
+private:
+	/** Returns the DPI scaler that should be used for game UI in the viewport */
+	float GetGameViewportDPIScale() const;
+
 private:
 	/** Tab which this viewport is located in */
 	TWeakPtr<class FLevelViewportLayout> ParentLayout;
diff --git a/Engine/Source/Editor/Persona/Private/SAnimCurvePanel.cpp b/Engine/Source/Editor/Persona/Private/SAnimCurvePanel.cpp
--- a/Engine/Source/Editor/Persona/Private/SAnimCurvePanel.cpp
+++ b/Engine/Source/Editor/Persona/Private/SAnimCurvePanel.cpp
@@ -97,6 +97,10 @@ public:
 		}
 		return FText::GetEmpty();
 	}
+
+	virtual void OnCurveChanged() override
+	{
+	}
 };
 
 //////////////////////////////////////////////////////////////////////////
diff --git a/Engine/Source/Editor/UMGEditor/Private/Designer/SDesignerView.cpp b/Engine/Source/Editor/UMGEditor/Private/Designer/SDesignerView.cpp
--- a/Engine/Source/Editor/UMGEditor/Private/Designer/SDesignerView.cpp
+++ b/Engine/Source/Editor/UMGEditor/Private/Designer/SDesignerView.cpp
@@ -150,12 +150,17 @@ void SDesignerView::Construct(const FArguments& InArgs, TSharedPtr<FWidgetBluepr
 				.ZoomAmount(this, &SDesignerView::GetZoomAmount)
 				.ViewOffset(this, &SDesignerView::GetViewOffset)
 				[
-					SAssignNew(PreviewSurface, SBox)
+					SNew(SBox)
 					.WidthOverride(this, &SDesignerView::GetPreviewWidth)
 					.HeightOverride(this, &SDesignerView::GetPreviewHeight)
 					.HAlign(HAlign_Fill)
 					.VAlign(VAlign_Fill)
 					.Visibility(EVisibility::SelfHitTestInvisible)
+					[
+						SAssignNew(PreviewSurface, SDPIScaler)
+						.DPIScale(this, &SDesignerView::GetPreviewDPIScale)
+						.Visibility(EVisibility::SelfHitTestInvisible)
+					]
 				]
 			]
 
@@ -247,6 +252,11 @@ FOptionalSize SDesignerView::GetPreviewHeight() const
 	return (float)PreviewHeight;
 }
 
+float SDesignerView::GetPreviewDPIScale() const
+{
+	return GetDefault<URendererSettings>(URendererSettings::StaticClass())->GetDPIScaleBasedOnSize(FIntPoint(PreviewWidth, PreviewHeight));
+}
+
 FSlateRect SDesignerView::ComputeAreaBounds() const
 {
 	return FSlateRect(0, 0, PreviewWidth, PreviewHeight);
@@ -329,7 +339,7 @@ void SDesignerView::CreateExtensionWidgetsForSelection()
 FVector2D SDesignerView::GetExtensionPosition(TSharedRef<FDesignerSurfaceElement> ExtensionElement) const
 {
 	const FVector2D TopLeft = CachedDesignerWidgetLocation;
-	const FVector2D Size = CachedDesignerWidgetSize * GetZoomAmount();
+	const FVector2D Size = CachedDesignerWidgetSize * GetZoomAmount() * GetPreviewDPIScale();
 
 	FVector2D FinalPosition(0, 0);
 
@@ -851,8 +861,8 @@ SDesignerView::DragHandle SDesignerView::HitTestDragHandles(const FGeometry& All
 		FVector2D WidgetLocalPosition = AllottedGeometry.AbsoluteToLocal(ArrangedWidget.Geometry.AbsolutePosition);
 		float X = WidgetLocalPosition.X;
 		float Y = WidgetLocalPosition.Y;
-		float Width = ArrangedWidget.Geometry.Size.X * GetZoomAmount();
-		float Height = ArrangedWidget.Geometry.Size.Y * GetZoomAmount();
+		float Width = ArrangedWidget.Geometry.Size.X * GetZoomAmount() * GetPreviewDPIScale();
+		float Height = ArrangedWidget.Geometry.Size.Y * GetZoomAmount() * GetPreviewDPIScale();
 
 		TArray<FVector2D> Handles;
 		Handles.Add(FVector2D(X, Y));					// Top - Left
diff --git a/Engine/Source/Editor/UMGEditor/Private/Designer/SDesignerView.h b/Engine/Source/Editor/UMGEditor/Private/Designer/SDesignerView.h
--- a/Engine/Source/Editor/UMGEditor/Private/Designer/SDesignerView.h
+++ b/Engine/Source/Editor/UMGEditor/Private/Designer/SDesignerView.h
@@ -49,10 +49,15 @@ public:
 	void Register(TSharedRef<FDesignerExtension> Extension);
 
 private:
+	/** The width of the preview screen for the UI */
 	FOptionalSize GetPreviewWidth() const;
 
+	/** The height of the preview screen for the UI */
 	FOptionalSize GetPreviewHeight() const;
 
+	/** Gets the DPI scale that would be applied given the current preview width and height */
+	float GetPreviewDPIScale() const;
+
 	virtual FSlateRect ComputeAreaBounds() const override;
 
 	/** Updates the designer to display the latest preview widget */
@@ -141,7 +146,7 @@ private:
 	UPanelWidget* DropPreviewParent;
 
 	TSharedPtr<class SZoomPan> PreviewHitTestRoot;
-	TSharedPtr<SBox> PreviewSurface;
+	TSharedPtr<SDPIScaler> PreviewSurface;
 	TSharedPtr<SCanvas> ExtensionWidgetCanvas;
 
 	DragHandle CurrentHandle;
diff --git a/Engine/Source/Editor/UMGEditor/Private/Extensions/CanvasSlotExtension.cpp b/Engine/Source/Editor/UMGEditor/Private/Extensions/CanvasSlotExtension.cpp
--- a/Engine/Source/Editor/UMGEditor/Private/Extensions/CanvasSlotExtension.cpp
+++ b/Engine/Source/Editor/UMGEditor/Private/Extensions/CanvasSlotExtension.cpp
@@ -133,9 +133,11 @@ FReply FCanvasSlotExtension::HandleDragging(const FGeometry& Geometry, const FPo
 {
 	if ( bDragging )
 	{
+		float InverseScale = ( 1.0f / Geometry.Scale );
+
 		for ( FWidgetReference& Selection : SelectionCache )
 		{
-			MoveByAmount(Selection, Event.GetCursorDelta());
+			MoveByAmount(Selection, Event.GetCursorDelta() * InverseScale);
 		}
 
 		return FReply::Handled();
diff --git a/Engine/Source/Editor/UnrealEd/Classes/Editor/EditorEngine.h b/Engine/Source/Editor/UnrealEd/Classes/Editor/EditorEngine.h
--- a/Engine/Source/Editor/UnrealEd/Classes/Editor/EditorEngine.h
+++ b/Engine/Source/Editor/UnrealEd/Classes/Editor/EditorEngine.h
@@ -541,7 +541,7 @@ public:
 	FOnEndTransformObject& OnEndObjectMovement() { return OnEndObjectTransformEvent; }
 
 	/** Delegate broadcast by the engine every tick when PIE/SIE is active, to check to see whether we need to
-	    be able to capture state for simulating actor (for Sequencer recording features).  The single bool parameter
+		be able to capture state for simulating actor (for Sequencer recording features).  The single bool parameter
 		should be set to true if recording features are needed. */
 	DECLARE_EVENT_OneParam( UEditorEngine, FGetActorRecordingState, bool& /* bIsRecordingActive */ );
 	FGetActorRecordingState& GetActorRecordingState() { return GetActorRecordingStateEvent; }
@@ -2214,7 +2214,7 @@ public:
 	/** @return true if the editor is able to launch PIE with online platform support */
 	bool SupportsOnlinePIE() const;
 
-    /** @return true if there are active PIE instances logged into an online platform */
+	/** @return true if there are active PIE instances logged into an online platform */
 	bool IsPlayingWithOnlinePIE() const { return NumOnlinePIEInstances > 0; }
 
 	/**
@@ -2333,6 +2333,9 @@ public:
 	void CreatePIEWorldFromLogin(FWorldContext& PieWorldContext, EPlayNetMode PlayNetMode, FPieLoginStruct& DataStruct);
 
 private:
+	/** Gets the DPI Scale for the game viewport in the editor. */
+	float GetGameViewportDPIScale() const;
+
 	/**
 	 * Non Online PIE creation flow, creates all instances of PIE at once when online isn't requested/required
 	 *
@@ -2506,7 +2509,7 @@ private:
 	FOnEndTransformObject OnEndObjectTransformEvent;
 
 	/** Delegate broadcast by the engine every tick when PIE/SIE is active, to check to see whether we need to
-	    be able to capture state for simulating actor (for Sequencer recording features) */
+		be able to capture state for simulating actor (for Sequencer recording features) */
 	FGetActorRecordingState GetActorRecordingStateEvent;
 
 	/** Reference to owner of the current popup */
diff --git a/Engine/Source/Editor/UnrealEd/Private/PlayLevel.cpp b/Engine/Source/Editor/UnrealEd/Private/PlayLevel.cpp
--- a/Engine/Source/Editor/UnrealEd/Private/PlayLevel.cpp
+++ b/Engine/Source/Editor/UnrealEd/Private/PlayLevel.cpp
@@ -2403,7 +2403,11 @@ UGameInstance* UEditorEngine::CreatePIEGameInstance(int32 PIEInstance, bool bInS
 						.IsEnabled(FSlateApplication::Get().GetNormalExecutionAttribute())
 						.EnableGammaCorrection( false )// Gamma correction in the game is handled in post processing in the scene renderer
 						[
-							ViewportOverlayWidgetRef
+							SNew(SDPIScaler)
+							.DPIScale(TAttribute<float>::Create(TAttribute<float>::FGetter::CreateUObject(this, &UEditorEngine::GetGameViewportDPIScale)))
+							[
+								ViewportOverlayWidgetRef
+							]
 						];
 
 				// Create a viewport widget for the game to render in.
@@ -2524,6 +2528,12 @@ UGameInstance* UEditorEngine::CreatePIEGameInstance(int32 PIEInstance, bool bInS
 	return GameInstance;
 }
 
+float UEditorEngine::GetGameViewportDPIScale() const
+{
+	UGameUserSettings* UserSettings = GEngine->GetGameUserSettings();
+	return GetDefault<URendererSettings>(URendererSettings::StaticClass())->GetDPIScaleBasedOnSize( UserSettings->GetScreenResolution() );
+}
+
 FViewport* UEditorEngine::GetActiveViewport()
 {
 	// Get the Level editor module and request the Active Viewport.
diff --git a/Engine/Source/Editor/UnrealEd/Private/SCurveEditor.cpp b/Engine/Source/Editor/UnrealEd/Private/SCurveEditor.cpp
--- a/Engine/Source/Editor/UnrealEd/Private/SCurveEditor.cpp
+++ b/Engine/Source/Editor/UnrealEd/Private/SCurveEditor.cpp
@@ -891,6 +891,8 @@ void SCurveEditor::DeleteSelectedKeys()
 			Key.Curve->DeleteKey(Key.KeyHandle);
 		}
 	}
+
+	CurveOwner->OnCurveChanged();
 }
 
 
@@ -1035,6 +1037,8 @@ void SCurveEditor::AddNewKey(FGeometry InMyGeometry, FVector2D ScreenPosition)
 			AddToSelection(FSelectedCurveKey(SelectedCurve, NewKeyHandle));
 		}
 	}
+
+	CurveOwner->OnCurveChanged();
 }
 
 void SCurveEditor::ResetDragValues()
@@ -1230,6 +1234,7 @@ void SCurveEditor::OnTimeComitted(float NewTime, ETextCommit::Type CommitType)
 				CurveOwner->ModifyOwner();
 				Key.Curve->SetKeyTime(Key.KeyHandle, NewTime);
 			}
+			CurveOwner->OnCurveChanged();
 		}
 
 		FSlateApplication::Get().DismissAllMenus();
@@ -1248,6 +1253,7 @@ void SCurveEditor::OnTimeChanged(float NewTime)
 				CurveOwner->ModifyOwner();
 				Key.Curve->SetKeyTime(Key.KeyHandle, NewTime);
 			}
+			CurveOwner->OnCurveChanged();
 		}
 	}
 }
@@ -1293,6 +1299,8 @@ void SCurveEditor::OnValueComitted(float NewValue, ETextCommit::Type CommitType)
 			}
 		}
 
+		CurveOwner->OnCurveChanged();
+
 		FSlateApplication::Get().DismissAllMenus();
 	}
 }
@@ -1315,6 +1323,8 @@ void SCurveEditor::OnValueChanged(float NewValue)
 				Key.Curve->SetKeyValue(Key.KeyHandle, NewValue);
 			}
 		}
+
+		CurveOwner->OnCurveChanged();
 	}
 }
 
@@ -1948,6 +1958,7 @@ void SCurveEditor::OnSelectInterpolationMode(ERichCurveInterpMode InterpMode, ER
 		}
 	}
 
+	CurveOwner->OnCurveChanged();
 }
 
 /* Given a tangent value for a key, calculates the 2D delta vector from that key in curve space */
@@ -2040,8 +2051,10 @@ void SCurveEditor::EndDragTransaction()
 {
 	if ( TransactionIndex >= 0 )
 	{
-		GEditor->EndTransaction( );
+		GEditor->EndTransaction();
 		TransactionIndex = -1;
+
+		CurveOwner->OnCurveChanged();
 	}
 }
 
diff --git a/Engine/Source/Runtime/Engine/Classes/Curves/CurveBase.h b/Engine/Source/Runtime/Engine/Classes/Curves/CurveBase.h
--- a/Engine/Source/Runtime/Engine/Classes/Curves/CurveBase.h
+++ b/Engine/Source/Runtime/Engine/Classes/Curves/CurveBase.h
@@ -419,6 +419,9 @@ public:
 
 	/** Called to make curve owner transactional */
 	virtual void MakeTransactional() = 0;
+
+	/** Called when the curve has been changed */
+	virtual void OnCurveChanged() = 0;
 };
 
 
@@ -465,6 +468,7 @@ public:
 
 	virtual void ModifyOwner() override;
 	virtual void MakeTransactional() override;
+	virtual void OnCurveChanged() override;
 	// End FCurveOwnerInterface
 
 	// Begin UCurveBase interface
diff --git a/Engine/Source/Runtime/Engine/Classes/Engine/GameEngine.h b/Engine/Source/Runtime/Engine/Classes/Engine/GameEngine.h
--- a/Engine/Source/Runtime/Engine/Classes/Engine/GameEngine.h
+++ b/Engine/Source/Runtime/Engine/Classes/Engine/GameEngine.h
@@ -135,5 +135,8 @@ public:
 	 * in order to properly support multiple concurrent UWorlds.
 	 */
 	UWorld* GetGameWorld();
+
+protected:
+	float GetGameViewportDPIScale() const;
 };
 
diff --git a/Engine/Source/Runtime/Engine/Classes/Engine/RendererSettings.h b/Engine/Source/Runtime/Engine/Classes/Engine/RendererSettings.h
--- a/Engine/Source/Runtime/Engine/Classes/Engine/RendererSettings.h
+++ b/Engine/Source/Runtime/Engine/Classes/Engine/RendererSettings.h
@@ -177,6 +177,11 @@ class ENGINE_API URendererSettings
 		ToolTip="Screen radius at which wireframe objects are culled. Larger values can improve performance when viewing a scene in wireframe."))
 	float WireframeCullThreshold;
 
+	UPROPERTY(config, EditAnywhere, Category=UI, meta=(
+		DisplayName="DPI Curve",
+		ToolTip="Controls how the UI is scaled at different resolutions based on the shortest side" ))
+	FRuntimeFloatCurve UIScaleCurve;
+
 public:
 
 	// Begin UObject interface
@@ -187,4 +192,7 @@ public:
 #endif
 
 	// End UObject interface
+
+	/** Gets the current scale of the UI based on the size */
+	float GetDPIScaleBasedOnSize(FIntPoint Size) const;
 };
diff --git a/Engine/Source/Runtime/Engine/Private/CurveBase.cpp b/Engine/Source/Runtime/Engine/Private/CurveBase.cpp
--- a/Engine/Source/Runtime/Engine/Private/CurveBase.cpp
+++ b/Engine/Source/Runtime/Engine/Private/CurveBase.cpp
@@ -843,6 +843,11 @@ void UCurveBase::MakeTransactional()
 	SetFlags(GetFlags() | RF_Transactional);
 }
 
+void UCurveBase::OnCurveChanged()
+{
+
+}
+
 void UCurveBase::ResetCurve()
 {
 	TArray<FRichCurveEditInfo> Curves = GetCurves();
diff --git a/Engine/Source/Runtime/Engine/Private/GameEngine.cpp b/Engine/Source/Runtime/Engine/Private/GameEngine.cpp
--- a/Engine/Source/Runtime/Engine/Private/GameEngine.cpp
+++ b/Engine/Source/Runtime/Engine/Private/GameEngine.cpp
@@ -113,7 +113,11 @@ void UGameEngine::CreateGameViewportWidget( UGameViewportClient* GameViewportCli
 			// @todo TEMP
 			.RenderDirectlyToWindow( !GEngine->bStartWithMatineeCapture && GIsDumpingMovie == 0 )
 			[
-				ViewportOverlayWidgetRef
+				SNew(SDPIScaler)
+				.DPIScale(TAttribute<float>::Create(TAttribute<float>::FGetter::CreateUObject(this, &UGameEngine::GetGameViewportDPIScale)) )
+				[
+					ViewportOverlayWidgetRef
+				]
 			];
 
 	GameViewportWidget = GameViewportWidgetRef;
@@ -164,6 +168,13 @@ void UGameEngine::CreateGameViewport( UGameViewportClient* GameViewportClient )
 	GameViewport->SetViewportFrame(ViewportFrame);
 }
 
+float UGameEngine::GetGameViewportDPIScale() const
+{
+	UGameUserSettings* UserSettings = GEngine->GetGameUserSettings();
+	FIntPoint ScreenRes = UserSettings->GetScreenResolution();
+	return GetDefault<URendererSettings>(URendererSettings::StaticClass())->GetDPIScaleBasedOnSize( ScreenRes );
+}
+
 void UGameEngine::ConditionallyOverrideSettings(int32& ResolutionX, int32& ResolutionY, EWindowMode::Type& WindowMode)
 {
 	if (FParse::Param(FCommandLine::Get(),TEXT("Windowed")) || FParse::Param(FCommandLine::Get(), TEXT("SimMobile")))
@@ -435,7 +446,7 @@ void UGameEngine::Init(IEngineLoop* InEngineLoop)
 		ViewportClient = ConstructObject<UGameViewportClient>(GameViewportClientClass,this);
 		ViewportClient->Init(*GameInstance->GetWorldContext(), GameInstance);
 		GameViewport = ViewportClient;
- 		GameInstance->GetWorldContext()->GameViewport = ViewportClient;
+		GameInstance->GetWorldContext()->GameViewport = ViewportClient;
 	}
 
 	bCheckForMovieCapture = true;
@@ -1062,8 +1073,8 @@ UWorld* UGameEngine::GetGameWorld()
 	for (auto It = WorldList.CreateConstIterator(); It; ++It)
 	{
 		const FWorldContext& Context = *It;
-        // Explicitly not checking for PIE worlds here, this should only 
-        // be called outside of editor (and thus is in UGameEngine
+		// Explicitly not checking for PIE worlds here, this should only 
+		// be called outside of editor (and thus is in UGameEngine
 		if (Context.WorldType == EWorldType::Game && Context.World())
 		{
 			return Context.World();
diff --git a/Engine/Source/Runtime/Engine/Private/RenderingSettings.cpp b/Engine/Source/Runtime/Engine/Private/RenderingSettings.cpp
--- a/Engine/Source/Runtime/Engine/Private/RenderingSettings.cpp
+++ b/Engine/Source/Runtime/Engine/Private/RenderingSettings.cpp
@@ -78,3 +78,11 @@ void URendererSettings::PostEditChangeProperty(FPropertyChangedEvent& PropertyCh
 	}
 }
 #endif // #if WITH_EDITOR
+
+float URendererSettings::GetDPIScaleBasedOnSize(FIntPoint Size) const
+{
+	int32 ShortestSide = FMath::Min(Size.X, Size.Y);
+	const FRichCurve* DPICurve = UIScaleCurve.GetRichCurveConst();
+	float Scale = DPICurve->Eval((float)ShortestSide, 1.0f);
+	return FMath::Max(Scale, 0.01f);
+}
diff --git a/Engine/Source/Runtime/UMG/Private/UserWidget.cpp b/Engine/Source/Runtime/UMG/Private/UserWidget.cpp
--- a/Engine/Source/Runtime/UMG/Private/UserWidget.cpp
+++ b/Engine/Source/Runtime/UMG/Private/UserWidget.cpp
@@ -406,18 +406,20 @@ void UUserWidget::AddToViewport(bool bAbsoluteLayout, bool bModal, bool bShowCur
 		UWorld* World = GetWorld();
 		if ( World && World->IsGameWorld() )
 		{
-			UGameViewportClient* Viewport = World->GetGameViewport();
-			Viewport->AddViewportWidgetContent(WidgetHost);
-
-			//TODO UMG this isn't what should manage focus, a higher level window controller, probably the viewport needs to understand
-			// the Widget stack, and the dialog stack.
-			if ( bModal )
+			if ( UGameViewportClient* Viewport = World->GetGameViewport() )
 			{
-				TWeakPtr<SViewport> GameViewportWidget = Viewport->GetGameViewport()->GetViewportWidget();
-				if ( GameViewportWidget.IsValid() )
+				Viewport->AddViewportWidgetContent(WidgetHost);
+
+				//TODO UMG this isn't what should manage focus, a higher level window controller, probably the viewport needs to understand
+				// the Widget stack, and the dialog stack.
+				if ( bModal )
 				{
-					GameViewportWidget.Pin()->SetWidgetToFocusOnActivate(OutUserSlateWidget);
-					FSlateApplication::Get().SetKeyboardFocus(OutUserSlateWidget);
+					TWeakPtr<SViewport> GameViewportWidget = Viewport->GetGameViewport()->GetViewportWidget();
+					if ( GameViewportWidget.IsValid() )
+					{
+						GameViewportWidget.Pin()->SetWidgetToFocusOnActivate(OutUserSlateWidget);
+						FSlateApplication::Get().SetKeyboardFocus(OutUserSlateWidget);
+					}
 				}
 			}
 		}
@@ -434,21 +436,23 @@ void UUserWidget::RemoveFromViewport()
 		UWorld* World = GetWorld();
 		if ( World && World->IsGameWorld() )
 		{
-			UGameViewportClient* Viewport = World->GetGameViewport();
-			Viewport->RemoveViewportWidgetContent(WidgetHost.ToSharedRef());
-
-			if ( WidgetHost->IsModal() )
+			if ( UGameViewportClient* Viewport = World->GetGameViewport() )
 			{
-				TWeakPtr<SViewport> GameViewportWidget = Viewport->GetGameViewport()->GetViewportWidget();
-				if ( GameViewportWidget.IsValid() )
-				{
-					//TODO UMG this isn't what should manage focus, a higher level window controller, probably the viewport needs to understand
-					// the Widget stack, and the dialog stack.
-					GameViewportWidget.Pin()->ClearWidgetToFocusOnActivate();
-					FSlateApplication::Get().ClearKeyboardFocus(EKeyboardFocusCause::SetDirectly);
+				Viewport->RemoveViewportWidgetContent(WidgetHost.ToSharedRef());
 
-					FSlateApplication::Get().SetFocusToGameViewport();
-					FSlateApplication::Get().SetJoystickCaptorToGameViewport();
+				if ( WidgetHost->IsModal() )
+				{
+					TWeakPtr<SViewport> GameViewportWidget = Viewport->GetGameViewport()->GetViewportWidget();
+					if ( GameViewportWidget.IsValid() )
+					{
+						//TODO UMG this isn't what should manage focus, a higher level window controller, probably the viewport needs to understand
+						// the Widget stack, and the dialog stack.
+						GameViewportWidget.Pin()->ClearWidgetToFocusOnActivate();
+						FSlateApplication::Get().ClearKeyboardFocus(EKeyboardFocusCause::SetDirectly);
+
+						FSlateApplication::Get().SetFocusToGameViewport();
+						FSlateApplication::Get().SetJoystickCaptorToGameViewport();
+					}
 				}
 			}
 		}
