UMG - Capturing joystick when modal widgets get focus.

diff --git a/Engine/Source/Runtime/UMG/Private/UserWidget.cpp b/Engine/Source/Runtime/UMG/Private/UserWidget.cpp
--- a/Engine/Source/Runtime/UMG/Private/UserWidget.cpp
+++ b/Engine/Source/Runtime/UMG/Private/UserWidget.cpp
@@ -81,12 +81,12 @@ class SViewportWidgetHost : public SCompoundWidget
 
 	virtual FReply OnMouseButtonDown(const FGeometry& MyGeometry, const FPointerEvent& MouseEvent) override
 	{
-		return bModal ? FReply::Handled().CaptureMouse(AsShared()) : FReply::Unhandled();
+		return bModal ? FReply::Handled().CaptureMouse(AsShared()).CaptureJoystick(AsShared()) : FReply::Unhandled();
 	}
 
 	virtual FReply OnMouseButtonUp(const FGeometry& MyGeometry, const FPointerEvent& MouseEvent) override
 	{
-		return bModal ? FReply::Handled().ReleaseMouseCapture() : FReply::Unhandled();
+		return bModal ? FReply::Handled().ReleaseMouseCapture().ReleaseJoystickCapture() : FReply::Unhandled();
 	}
 
 	virtual FReply OnMouseButtonDoubleClick(const FGeometry& MyGeometry, const FPointerEvent& MouseEvent) override
@@ -418,7 +418,9 @@ void UUserWidget::AddToViewport(bool bAbsoluteLayout, bool bModal, bool bShowCur
 					if ( GameViewportWidget.IsValid() )
 					{
 						GameViewportWidget.Pin()->SetWidgetToFocusOnActivate(OutUserSlateWidget);
-						FSlateApplication::Get().SetKeyboardFocus(OutUserSlateWidget);
+						
+						FSlateApplication::Get().SetFocusToGameViewport();
+						FSlateApplication::Get().SetJoystickCaptorToGameViewport();
 					}
 				}
 			}
