UMG - Changing FSReply to FEventReply.  Changing the way handling is generated and returned from functions.  New functionality added to mimic more capability in slate.  Enough capability is now exposed for drag/drop to be triggered in user widgets.  Removing the Drag and Drop panels, they don't carry enough data to be useful.  Made a simple prototype to prove-out the current drag/drop workflow.

diff --git a/Engine/Config/BaseEngine.ini b/Engine/Config/BaseEngine.ini
--- a/Engine/Config/BaseEngine.ini
+++ b/Engine/Config/BaseEngine.ini
@@ -571,6 +571,7 @@ bUseStreamingPause=false
 
 ; Widget Blueprint Redirects
 +ActiveClassRedirects=(OldClassName="WidgetBlueprint",NewClassName="/Script/UMGEditor.WidgetBlueprint")
++ActiveClassRedirects=(OldClassName="SReply",NewClassName="EventReply")
 
 ;EAxis rename
 +EnumRedirects=(EnumName="EAxis",OldEnumEntry="AXIS_X",NewEnumEntry="EAxis::X")
diff --git a/Engine/Source/Editor/UMGEditor/Private/Hierarchy/SHierarchyView.cpp b/Engine/Source/Editor/UMGEditor/Private/Hierarchy/SHierarchyView.cpp
--- a/Engine/Source/Editor/UMGEditor/Private/Hierarchy/SHierarchyView.cpp
+++ b/Engine/Source/Editor/UMGEditor/Private/Hierarchy/SHierarchyView.cpp
@@ -227,7 +227,7 @@ TSharedRef< ITableRow > SHierarchyView::WidgetHierarchy_OnGenerateRow(TSharedPtr
 
 void SHierarchyView::WidgetHierarchy_OnSelectionChanged(TSharedPtr<FHierarchyModel> SelectedItem, ESelectInfo::Type SelectInfo)
 {
-	if ( SelectInfo != ESelectInfo::Direct )
+	if ( SelectedItem.IsValid() && SelectInfo != ESelectInfo::Direct )
 	{
 		SelectedItem->OnSelection();
 	}
diff --git a/Engine/Source/Runtime/UMG/Private/Components/DragPanel.cpp b/Engine/Source/Runtime/UMG/Private/Components/DragPanel.cpp
deleted file mode 100644
diff --git a/Engine/Source/Runtime/UMG/Private/Components/DropPanel.cpp b/Engine/Source/Runtime/UMG/Private/Components/DropPanel.cpp
deleted file mode 100644
diff --git a/Engine/Source/Runtime/UMG/Private/Components/PanelWidget.cpp b/Engine/Source/Runtime/UMG/Private/Components/PanelWidget.cpp
--- a/Engine/Source/Runtime/UMG/Private/Components/PanelWidget.cpp
+++ b/Engine/Source/Runtime/UMG/Private/Components/PanelWidget.cpp
@@ -146,6 +146,11 @@ bool UPanelWidget::RemoveChild(UWidget* Content)
 	return false;
 }
 
+bool UPanelWidget::HasAnyChildren() const
+{
+	return GetChildrenCount() > 0;
+}
+
 void UPanelWidget::ClearChildren()
 {
 	int32 Children = GetChildrenCount();
diff --git a/Engine/Source/Runtime/UMG/Private/Slate/SDragPanel.cpp b/Engine/Source/Runtime/UMG/Private/Slate/SDragPanel.cpp
deleted file mode 100644
diff --git a/Engine/Source/Runtime/UMG/Private/Slate/SDropPanel.cpp b/Engine/Source/Runtime/UMG/Private/Slate/SDropPanel.cpp
deleted file mode 100644
diff --git a/Engine/Source/Runtime/UMG/Private/WidgetBlueprintLibrary.cpp b/Engine/Source/Runtime/UMG/Private/WidgetBlueprintLibrary.cpp
--- a/Engine/Source/Runtime/UMG/Private/WidgetBlueprintLibrary.cpp
+++ b/Engine/Source/Runtime/UMG/Private/WidgetBlueprintLibrary.cpp
@@ -86,4 +86,94 @@ void UWidgetBlueprintLibrary::DrawText(UPARAM(ref) FPaintContext& Context, const
 		Tint);
 }
 
+FEventReply UWidgetBlueprintLibrary::Handled()
+{
+	FEventReply Reply;
+	Reply.NativeReply = FReply::Handled();
+
+	return Reply;
+}
+
+FEventReply UWidgetBlueprintLibrary::Unhandled()
+{
+	FEventReply Reply;
+	Reply.NativeReply = FReply::Unhandled();
+
+	return Reply;
+}
+
+FEventReply UWidgetBlueprintLibrary::CaptureMouse(UPARAM(ref) FEventReply& Reply, UWidget* CapturingWidget)
+{
+	if ( CapturingWidget )
+	{
+		TSharedPtr<SWidget> CapturingSlateWidget = CapturingWidget->GetCachedWidget();
+		if ( CapturingSlateWidget.IsValid() )
+		{
+			Reply.NativeReply = Reply.NativeReply.CaptureMouse(CapturingSlateWidget.ToSharedRef());
+		}
+	}
+
+	return Reply;
+}
+
+FEventReply UWidgetBlueprintLibrary::ReleaseMouseCapture(UPARAM(ref) FEventReply& Reply)
+{
+	Reply.NativeReply = Reply.NativeReply.ReleaseMouseCapture();
+
+	return Reply;
+}
+
+FEventReply UWidgetBlueprintLibrary::CaptureJoystick(UPARAM(ref) FEventReply& Reply, UWidget* CapturingWidget, bool bInAllJoysticks/* = false*/)
+{
+	if ( CapturingWidget )
+	{
+		TSharedPtr<SWidget> CapturingSlateWidget = CapturingWidget->GetCachedWidget();
+		if ( CapturingSlateWidget.IsValid() )
+		{
+			Reply.NativeReply = Reply.NativeReply.CaptureJoystick(CapturingSlateWidget.ToSharedRef(), bInAllJoysticks);
+		}
+	}
+
+	return Reply;
+}
+
+FEventReply UWidgetBlueprintLibrary::ReleaseJoystickCapture(UPARAM(ref) FEventReply& Reply, bool bInAllJoysticks /*= false*/)
+{
+	Reply.NativeReply = Reply.NativeReply.ReleaseJoystickCapture(bInAllJoysticks);
+
+	return Reply;
+}
+
+FEventReply UWidgetBlueprintLibrary::DetectDrag(UPARAM(ref) FEventReply& Reply, UWidget* WidgetDetectingDrag, FKey DragKey)
+{
+	if ( WidgetDetectingDrag )
+	{
+		TSharedPtr<SWidget> SlateWidgetDetectingDrag = WidgetDetectingDrag->GetCachedWidget();
+		if ( SlateWidgetDetectingDrag.IsValid() )
+		{
+			Reply.NativeReply = Reply.NativeReply.DetectDrag(SlateWidgetDetectingDrag.ToSharedRef(), DragKey);
+		}
+	}
+
+	return Reply;
+}
+
+FEventReply UWidgetBlueprintLibrary::IfKeyPressedDetectDrag(const FPointerEvent& PointerEvent, UWidget* WidgetDetectingDrag, FKey DragKey)
+{
+	if ( PointerEvent.GetEffectingButton() == DragKey )
+	{
+		FEventReply Reply = UWidgetBlueprintLibrary::Handled();
+		return UWidgetBlueprintLibrary::DetectDrag(Reply, WidgetDetectingDrag, DragKey);
+	}
+
+	return UWidgetBlueprintLibrary::Unhandled();
+}
+
+FEventReply UWidgetBlueprintLibrary::EndDragDrop(UPARAM(ref) FEventReply& Reply)
+{
+	Reply.NativeReply = Reply.NativeReply.EndDragDrop();
+
+	return Reply;
+}
+
 #undef LOCTEXT_NAMESPACE
\ No newline at end of file
diff --git a/Engine/Source/Runtime/UMG/Public/Blueprint/UserWidget.h b/Engine/Source/Runtime/UMG/Public/Blueprint/UserWidget.h
--- a/Engine/Source/Runtime/UMG/Public/Blueprint/UserWidget.h
+++ b/Engine/Source/Runtime/UMG/Public/Blueprint/UserWidget.h
@@ -158,39 +158,39 @@ public:
 	void OnPaint(UPARAM(ref) FPaintContext& Context) const;
 
 	UFUNCTION(BlueprintImplementableEvent, Category="Keyboard")
-	FSReply OnKeyboardFocusReceived(FGeometry MyGeometry, FKeyboardFocusEvent InKeyboardFocusEvent);
+	FEventReply OnKeyboardFocusReceived(FGeometry MyGeometry, FKeyboardFocusEvent InKeyboardFocusEvent);
 	UFUNCTION(BlueprintImplementableEvent, Category="Keyboard")
 	void OnKeyboardFocusLost(FKeyboardFocusEvent InKeyboardFocusEvent);
 	//UFUNCTION(BlueprintImplementableEvent, Category="User Interface")
 	//void OnKeyboardFocusChanging(FWeakWidgetPath PreviousFocusPath, FWidgetPath NewWidgetPath);
 
 	UFUNCTION(BlueprintImplementableEvent, Category="Keyboard")
-	FSReply OnKeyChar(FGeometry MyGeometry, FCharacterEvent InCharacterEvent);
+	FEventReply OnKeyChar(FGeometry MyGeometry, FCharacterEvent InCharacterEvent);
 	UFUNCTION(BlueprintImplementableEvent, Category="Keyboard")
-	FSReply OnPreviewKeyDown(FGeometry MyGeometry, FKeyboardEvent InKeyboardEvent);
+	FEventReply OnPreviewKeyDown(FGeometry MyGeometry, FKeyboardEvent InKeyboardEvent);
 	UFUNCTION(BlueprintImplementableEvent, Category="Keyboard")
-	FSReply OnKeyDown(FGeometry MyGeometry, FKeyboardEvent InKeyboardEvent);
+	FEventReply OnKeyDown(FGeometry MyGeometry, FKeyboardEvent InKeyboardEvent);
 	UFUNCTION(BlueprintImplementableEvent, Category="Keyboard")
-	FSReply OnKeyUp(FGeometry MyGeometry, FKeyboardEvent InKeyboardEvent);
+	FEventReply OnKeyUp(FGeometry MyGeometry, FKeyboardEvent InKeyboardEvent);
 
 	UFUNCTION(BlueprintImplementableEvent, Category="Mouse")
-	FSReply OnMouseButtonDown(FGeometry MyGeometry, const FPointerEvent& MouseEvent);
+	FEventReply OnMouseButtonDown(FGeometry MyGeometry, const FPointerEvent& MouseEvent);
 	UFUNCTION(BlueprintImplementableEvent, Category="Mouse")
-	FSReply OnPreviewMouseButtonDown(FGeometry MyGeometry, const FPointerEvent& MouseEvent);
+	FEventReply OnPreviewMouseButtonDown(FGeometry MyGeometry, const FPointerEvent& MouseEvent);
 	UFUNCTION(BlueprintImplementableEvent, Category="Mouse")
-	FSReply OnMouseButtonUp(FGeometry MyGeometry, const FPointerEvent& MouseEvent);
+	FEventReply OnMouseButtonUp(FGeometry MyGeometry, const FPointerEvent& MouseEvent);
 	UFUNCTION(BlueprintImplementableEvent, Category="Mouse")
-	FSReply OnMouseMove(FGeometry MyGeometry, const FPointerEvent& MouseEvent);
+	FEventReply OnMouseMove(FGeometry MyGeometry, const FPointerEvent& MouseEvent);
 	UFUNCTION(BlueprintImplementableEvent, Category="Mouse")
 	void OnMouseEnter(FGeometry MyGeometry, const FPointerEvent& MouseEvent);
 	UFUNCTION(BlueprintImplementableEvent, Category="Mouse")
 	void OnMouseLeave(const FPointerEvent& MouseEvent);
 	UFUNCTION(BlueprintImplementableEvent, Category="Mouse")
-	FSReply OnMouseWheel(FGeometry MyGeometry, const FPointerEvent& MouseEvent);
+	FEventReply OnMouseWheel(FGeometry MyGeometry, const FPointerEvent& MouseEvent);
 	//UFUNCTION(BlueprintImplementableEvent, Category="Mouse")
 	//FCursorReply OnCursorQuery(FGeometry MyGeometry, const FPointerEvent& CursorEvent) const;
 	UFUNCTION(BlueprintImplementableEvent, Category="Mouse")
-	FSReply OnMouseButtonDoubleClick(FGeometry InMyGeometry, const FPointerEvent& InMouseEvent);
+	FEventReply OnMouseButtonDoubleClick(FGeometry InMyGeometry, const FPointerEvent& InMouseEvent);
 
 	UFUNCTION(BlueprintImplementableEvent, Category="Drag and Drop")
 	void OnDragDetected(FGeometry MyGeometry, const FPointerEvent& MouseEvent, UDragDropOperation*& Operation);
@@ -205,22 +205,22 @@ public:
 	bool OnDrop(FGeometry MyGeometry, FPointerEvent PointerEvent, UDragDropOperation* Operation);
 
 	UFUNCTION(BlueprintImplementableEvent, Category="Gamepad Input")
-	FSReply OnControllerButtonPressed(FGeometry MyGeometry, FControllerEvent ControllerEvent);
+	FEventReply OnControllerButtonPressed(FGeometry MyGeometry, FControllerEvent ControllerEvent);
 	UFUNCTION(BlueprintImplementableEvent, Category="Gamepad Input")
-	FSReply OnControllerButtonReleased(FGeometry MyGeometry, FControllerEvent ControllerEvent);
+	FEventReply OnControllerButtonReleased(FGeometry MyGeometry, FControllerEvent ControllerEvent);
 	UFUNCTION(BlueprintImplementableEvent, Category="Gamepad Input")
-	FSReply OnControllerAnalogValueChanged(FGeometry MyGeometry, FControllerEvent ControllerEvent);
+	FEventReply OnControllerAnalogValueChanged(FGeometry MyGeometry, FControllerEvent ControllerEvent);
 
 	UFUNCTION(BlueprintImplementableEvent, Category="Touch Input")
-	FSReply OnTouchGesture(FGeometry MyGeometry, const FPointerEvent& GestureEvent);
+	FEventReply OnTouchGesture(FGeometry MyGeometry, const FPointerEvent& GestureEvent);
 	UFUNCTION(BlueprintImplementableEvent, Category="Touch Input")
-	FSReply OnTouchStarted(FGeometry MyGeometry, const FPointerEvent& InTouchEvent);
+	FEventReply OnTouchStarted(FGeometry MyGeometry, const FPointerEvent& InTouchEvent);
 	UFUNCTION(BlueprintImplementableEvent, Category="Touch Input")
-	FSReply OnTouchMoved(FGeometry MyGeometry, const FPointerEvent& InTouchEvent);
+	FEventReply OnTouchMoved(FGeometry MyGeometry, const FPointerEvent& InTouchEvent);
 	UFUNCTION(BlueprintImplementableEvent, Category="Touch Input")
-	FSReply OnTouchEnded(FGeometry MyGeometry, const FPointerEvent& InTouchEvent);
+	FEventReply OnTouchEnded(FGeometry MyGeometry, const FPointerEvent& InTouchEvent);
 	UFUNCTION(BlueprintImplementableEvent, Category="Touch Input")
-	FSReply OnMotionDetected(FGeometry MyGeometry, FMotionEvent InMotionEvent);
+	FEventReply OnMotionDetected(FGeometry MyGeometry, FMotionEvent InMotionEvent);
 
 	UFUNCTION(BlueprintCallable, Category="User Interface|Drag and Drop")
 	UDragDropOperation* CreateDragDropOperation(TSubclassOf<UDragDropOperation> Operation);
diff --git a/Engine/Source/Runtime/UMG/Public/Blueprint/WidgetBlueprintLibrary.h b/Engine/Source/Runtime/UMG/Public/Blueprint/WidgetBlueprintLibrary.h
--- a/Engine/Source/Runtime/UMG/Public/Blueprint/WidgetBlueprintLibrary.h
+++ b/Engine/Source/Runtime/UMG/Public/Blueprint/WidgetBlueprintLibrary.h
@@ -41,4 +41,41 @@ class UWidgetBlueprintLibrary : public UBlueprintFunctionLibrary
 	 */
 	UFUNCTION(BlueprintCallable, Category="Painting")
 	static void DrawText(UPARAM(ref) FPaintContext& Context, const FString& InString, FVector2D Position, FLinearColor Tint = FLinearColor::White);
+
+	UFUNCTION(BlueprintPure, Category="Widget|Event Reply")
+	static FEventReply Handled();
+
+	UFUNCTION(BlueprintPure, Category="Widget|Event Reply")
+	static FEventReply Unhandled();
+
+	UFUNCTION(BlueprintPure, meta=( HidePin="CapturingWidget", DefaultToSelf="CapturingWidget" ), Category="Widget|Event Reply")
+	static FEventReply CaptureMouse(UPARAM(ref) FEventReply& Reply, UWidget* CapturingWidget);
+
+	UFUNCTION(BlueprintPure, Category="Widget|Event Reply")
+	static FEventReply ReleaseMouseCapture(UPARAM(ref) FEventReply& Reply);
+
+	UFUNCTION(BlueprintPure, meta=( HidePin="CapturingWidget", DefaultToSelf="CapturingWidget" ), Category="Widget|Event Reply")
+	static FEventReply CaptureJoystick(UPARAM(ref) FEventReply& Reply, UWidget* CapturingWidget, bool bInAllJoysticks = false);
+
+	UFUNCTION(BlueprintPure, Category="Widget|Event Reply")
+	static FEventReply ReleaseJoystickCapture(UPARAM(ref) FEventReply& Reply, bool bInAllJoysticks = false);
+
+	/**
+	 * Ask Slate to detect if a user started dragging in this widget.
+	 * If a drag is detected, Slate will send an OnDragDetected event.
+	 *
+	 * @param WidgetDetectingDrag  Detect dragging in this widget
+	 * @param DragKey		       This button should be pressed to detect the drag
+	 */
+	UFUNCTION(BlueprintPure, meta=( HidePin="WidgetDetectingDrag", DefaultToSelf="WidgetDetectingDrag" ), Category="Widget|Event Reply")
+	static FEventReply DetectDrag(UPARAM(ref) FEventReply& Reply, UWidget* WidgetDetectingDrag, FKey DragKey);
+
+	UFUNCTION(BlueprintCallable, meta=( HidePin="WidgetDetectingDrag", DefaultToSelf="WidgetDetectingDrag" ), Category="Widget|Event Reply")
+	static FEventReply IfKeyPressedDetectDrag(const FPointerEvent& PointerEvent, UWidget* WidgetDetectingDrag, FKey DragKey);
+
+	/**
+	 * An event should return FReply::Handled().EndDragDrop() to request that the current drag/drop operation be terminated.
+	 */
+	UFUNCTION(BlueprintPure, Category="Widget|Event Reply")
+	static FEventReply EndDragDrop(UPARAM(ref) FEventReply& Reply);
 };
diff --git a/Engine/Source/Runtime/UMG/Public/Components/DragPanel.h b/Engine/Source/Runtime/UMG/Public/Components/DragPanel.h
deleted file mode 100644
diff --git a/Engine/Source/Runtime/UMG/Public/Components/DropPanel.h b/Engine/Source/Runtime/UMG/Public/Components/DropPanel.h
deleted file mode 100644
diff --git a/Engine/Source/Runtime/UMG/Public/Components/PanelWidget.h b/Engine/Source/Runtime/UMG/Public/Components/PanelWidget.h
--- a/Engine/Source/Runtime/UMG/Public/Components/PanelWidget.h
+++ b/Engine/Source/Runtime/UMG/Public/Components/PanelWidget.h
@@ -19,42 +19,46 @@ public:
 public:
 
 	/**  */
-	UFUNCTION(BlueprintCallable, Category="Appearance")
+	UFUNCTION(BlueprintCallable, Category="Widget|Panel")
 	int32 GetChildrenCount() const;
 
 	/**  */
-	UFUNCTION(BlueprintCallable, Category="Appearance")
+	UFUNCTION(BlueprintCallable, Category="Widget|Panel")
 	UWidget* GetChildAt(int32 Index) const;
 
 	/**  */
-	UFUNCTION(BlueprintCallable, Category="Appearance")
+	UFUNCTION(BlueprintCallable, Category="Widget|Panel")
 	int32 GetChildIndex(UWidget* Content) const;
 
 	/**  */
-	UFUNCTION(BlueprintCallable, Category="Appearance")
+	UFUNCTION(BlueprintCallable, Category="Widget|Panel")
 	bool RemoveChildAt(int32 Index);
 
 	/**  */
-	UFUNCTION(BlueprintCallable, Category="Appearance")
+	UFUNCTION(BlueprintCallable, Category="Widget|Panel")
 	UPanelSlot* AddChild(UWidget* Content);
 
 	/** Swaps the widget out of the slot at the given index, replacing it with a different widget. */
-	UFUNCTION(BlueprintCallable, Category="Appearance")
+	UFUNCTION(BlueprintCallable, Category="Widget|Panel")
 	void ReplaceChildAt(int32 Index, UWidget* Content);
 
 	/** Inserts a widget at a specific index */
-	UFUNCTION(BlueprintCallable, Category="Appearance")
+	UFUNCTION(BlueprintCallable, Category="Widget|Panel")
 	void InsertChildAt(int32 Index, UWidget* Content);
 
 	/**
 	 * Removes a specific widget from the container.
 	 * @return true if the widget was found and removed.
 	 */
-	UFUNCTION(BlueprintCallable, Category="Appearance")
+	UFUNCTION(BlueprintCallable, Category="Widget|Panel")
 	bool RemoveChild(UWidget* Content);
 
+	/** @return true if there are any child widgets in the panel */
+	UFUNCTION(BlueprintCallable, Category="Widget|Panel")
+	bool HasAnyChildren() const;
+
 	/** Remove all child widgets from the panel widget. */
-	UFUNCTION(BlueprintCallable, Category="Appearance")
+	UFUNCTION(BlueprintCallable, Category="Widget|Panel")
 	void ClearChildren();
 
 	/** @returns true if the panel supports more than one child. */
diff --git a/Engine/Source/Runtime/UMG/Public/Components/SlateWrapperTypes.h b/Engine/Source/Runtime/UMG/Public/Components/SlateWrapperTypes.h
--- a/Engine/Source/Runtime/UMG/Public/Components/SlateWrapperTypes.h
+++ b/Engine/Source/Runtime/UMG/Public/Components/SlateWrapperTypes.h
@@ -43,41 +43,37 @@ namespace ESlateSizeRule
 	};
 }
 
-/** A struct exposing FReply options in a UStruct compatible way. */
-USTRUCT()
-struct FSReply
+/**
+ * Allows users to handle events and return information to the underlying UI layer.
+ */
+USTRUCT(BlueprintType)
+struct FEventReply
 {
 	GENERATED_USTRUCT_BODY()
 
-	UPROPERTY(EditAnywhere, BlueprintReadWrite, Category=User)
-	uint32 bIsHandled : 1;
-
-	UPROPERTY(EditAnywhere, BlueprintReadWrite, Category=User)
-	uint32 bCaptureMouse : 1;
-
-	FSReply()
-		: bIsHandled(false)
-		, bCaptureMouse(false)
+	FEventReply(bool IsHandled = false)
+	: NativeReply(IsHandled ? FReply::Handled() : FReply::Unhandled())
 	{
-
 	}
 
-	FReply ToReply( TSharedRef<SWidget> Widget ) const
+	FReply ToReply(TSharedRef<SWidget> Widget) const
 	{
-		FReply Reply = FReply::Unhandled();
+		//FReply Reply = FReply::Unhandled();
 
-		if ( bIsHandled )
-		{
-			Reply = FReply::Handled();
-		}
+		//if ( bIsHandled )
+		//{
+		//	Reply = FReply::Handled();
+		//}
 
-		if ( bCaptureMouse )
-		{
-			Reply = Reply.CaptureMouse(Widget);
-		}
+		//if ( bCaptureMouse )
+		//{
+		//	Reply = Reply.CaptureMouse(Widget);
+		//}
 
-		return Reply;
+		return NativeReply;
 	}
+
+	FReply NativeReply;
 };
 
 /** A struct exposing size param related properties to UMG. */
diff --git a/Engine/Source/Runtime/UMG/Public/Components/Widget.h b/Engine/Source/Runtime/UMG/Public/Components/Widget.h
--- a/Engine/Source/Runtime/UMG/Public/Components/Widget.h
+++ b/Engine/Source/Runtime/UMG/Public/Components/Widget.h
@@ -47,8 +47,8 @@ public:
 	DECLARE_DYNAMIC_DELEGATE_RetVal_OneParam(UWidget*, FGenerateWidgetUObject, UObject*, Item);
 
 	// Events
-	DECLARE_DYNAMIC_DELEGATE_RetVal(FSReply, FOnReply);
-	DECLARE_DYNAMIC_DELEGATE_RetVal_TwoParams(FSReply, FOnPointerEvent, FGeometry, MyGeometry, const FPointerEvent&, MouseEvent);
+	DECLARE_DYNAMIC_DELEGATE_RetVal(FEventReply, FOnReply);
+	DECLARE_DYNAMIC_DELEGATE_RetVal_TwoParams(FEventReply, FOnPointerEvent, FGeometry, MyGeometry, const FPointerEvent&, MouseEvent);
 
 	/**
 	 * Allows controls to be exposed as variables in a blueprint.  Not all controls need to be exposed
diff --git a/Engine/Source/Runtime/UMG/Public/Slate/SDragPanel.h b/Engine/Source/Runtime/UMG/Public/Slate/SDragPanel.h
deleted file mode 100644
diff --git a/Engine/Source/Runtime/UMG/Public/Slate/SDropPanel.h b/Engine/Source/Runtime/UMG/Public/Slate/SDropPanel.h
deleted file mode 100644
diff --git a/Engine/Source/Runtime/UMG/Public/UMG.h b/Engine/Source/Runtime/UMG/Public/UMG.h
--- a/Engine/Source/Runtime/UMG/Public/UMG.h
+++ b/Engine/Source/Runtime/UMG/Public/UMG.h
@@ -99,9 +99,6 @@
 #include "ListView.h"
 #include "TileView.h"
 
-#include "DragPanel.h"
-#include "DropPanel.h"
-
 #include "Viewport.h"
 
 #include "NativeWidgetHost.h"
