UMG - The widget that detects the drag/drop operation now has a chance to do something when the drop isn't handled by anyone.  OnDragCancelled.

diff --git a/Engine/Source/Runtime/UMG/Private/DragDropOperation.cpp b/Engine/Source/Runtime/UMG/Private/DragDropOperation.cpp
--- a/Engine/Source/Runtime/UMG/Private/DragDropOperation.cpp
+++ b/Engine/Source/Runtime/UMG/Private/DragDropOperation.cpp
@@ -15,9 +15,9 @@ void UDragDropOperation::Drop_Implementation(const FPointerEvent& PointerEvent)
 	OnDrop.Broadcast(this);
 }
 
-void UDragDropOperation::DropCanceled_Implementation(const FPointerEvent& PointerEvent)
+void UDragDropOperation::DragCancelled_Implementation(const FPointerEvent& PointerEvent)
 {
-	OnDropCanceled.Broadcast(this);
+	OnDragCancelled.Broadcast(this);
 }
 
 void UDragDropOperation::Dragged_Implementation(const FPointerEvent& PointerEvent)
diff --git a/Engine/Source/Runtime/UMG/Private/Slate/SObjectWidget.cpp b/Engine/Source/Runtime/UMG/Private/Slate/SObjectWidget.cpp
--- a/Engine/Source/Runtime/UMG/Private/Slate/SObjectWidget.cpp
+++ b/Engine/Source/Runtime/UMG/Private/Slate/SObjectWidget.cpp
@@ -231,7 +231,7 @@ FReply SObjectWidget::OnDragDetected(const FGeometry& MyGeometry, const FPointer
 		FVector2D ScreenCursorPos = PointerEvent.GetScreenSpacePosition();
 		FVector2D ScreenDrageePosition = MyGeometry.AbsolutePosition;
 
-		return FReply::Handled().BeginDragDrop(FUMGDragDropOp::New(Operation, ScreenCursorPos, ScreenDrageePosition));
+		return FReply::Handled().BeginDragDrop(FUMGDragDropOp::New(Operation, ScreenCursorPos, ScreenDrageePosition, SharedThis(this)));
 	}
 
 	return FReply::Unhandled();
@@ -283,6 +283,15 @@ FReply SObjectWidget::OnDrop(const FGeometry& MyGeometry, const FDragDropEvent&
 	return FReply::Unhandled();
 }
 
+void SObjectWidget::OnDragCancelled(const FDragDropEvent& DragDropEvent, UDragDropOperation* Operation)
+{
+	TSharedPtr<FUMGDragDropOp> NativeOp = DragDropEvent.GetOperationAs<FUMGDragDropOp>();
+	if ( NativeOp.IsValid() )
+	{
+		WidgetObject->OnDragCancelled(DragDropEvent, NativeOp->GetOperation());
+	}
+}
+
 FReply SObjectWidget::OnControllerButtonPressed(const FGeometry& MyGeometry, const FControllerEvent& ControllerEvent)
 {
 	TSharedPtr<SWidget> SlateWidget = WidgetObject->GetCachedWidget();
diff --git a/Engine/Source/Runtime/UMG/Private/Slate/UMGDragDropOp.cpp b/Engine/Source/Runtime/UMG/Private/Slate/UMGDragDropOp.cpp
--- a/Engine/Source/Runtime/UMG/Private/Slate/UMGDragDropOp.cpp
+++ b/Engine/Source/Runtime/UMG/Private/Slate/UMGDragDropOp.cpp
@@ -17,7 +17,12 @@ void FUMGDragDropOp::OnDrop( bool bDropWasHandled, const FPointerEvent& MouseEve
 		}
 		else
 		{
-			OperationObject->DropCanceled(MouseEvent);
+			if ( SourceUserWidget.IsValid() )
+			{
+				SourceUserWidget->OnDragCancelled(FDragDropEvent(MouseEvent, AsShared()), OperationObject);
+			}
+
+			OperationObject->DragCancelled(MouseEvent);
 		}
 	}
 	
@@ -42,11 +47,12 @@ TSharedPtr<SWidget> FUMGDragDropOp::GetDefaultDecorator() const
 	return DecoratorWidget;
 }
 
-TSharedRef<FUMGDragDropOp> FUMGDragDropOp::New(UDragDropOperation* InOperation, const FVector2D &CursorPosition, const FVector2D &ScreenPositionOfDragee)
+TSharedRef<FUMGDragDropOp> FUMGDragDropOp::New(UDragDropOperation* InOperation, const FVector2D &CursorPosition, const FVector2D &ScreenPositionOfDragee, TSharedPtr<SObjectWidget> SourceUserWidget)
 {
 	TSharedRef<FUMGDragDropOp> Operation = MakeShareable(new FUMGDragDropOp);
 	Operation->Offset = ScreenPositionOfDragee - CursorPosition;
 	Operation->StartingScreenPos = ScreenPositionOfDragee;
+	Operation->SourceUserWidget = SourceUserWidget;
 
 	Operation->DragOperation = InOperation;
 
diff --git a/Engine/Source/Runtime/UMG/Public/Blueprint/DragDropOperation.h b/Engine/Source/Runtime/UMG/Public/Blueprint/DragDropOperation.h
--- a/Engine/Source/Runtime/UMG/Public/Blueprint/DragDropOperation.h
+++ b/Engine/Source/Runtime/UMG/Public/Blueprint/DragDropOperation.h
@@ -34,7 +34,7 @@ public:
 
 	/**  */
 	UPROPERTY(BlueprintAssignable)
-	FOnDragDropMulticast OnDropCanceled;
+	FOnDragDropMulticast OnDragCancelled;
 
 	/**  */
 	UPROPERTY(BlueprintAssignable)
@@ -46,7 +46,7 @@ public:
 
 	/**  */
 	UFUNCTION(BlueprintNativeEvent, Category="Drag and Drop")
-	void DropCanceled(const FPointerEvent& PointerEvent);
+	void DragCancelled(const FPointerEvent& PointerEvent);
 
 	/**  */
 	UFUNCTION(BlueprintNativeEvent, Category="Drag and Drop")
diff --git a/Engine/Source/Runtime/UMG/Public/Blueprint/UserWidget.h b/Engine/Source/Runtime/UMG/Public/Blueprint/UserWidget.h
--- a/Engine/Source/Runtime/UMG/Public/Blueprint/UserWidget.h
+++ b/Engine/Source/Runtime/UMG/Public/Blueprint/UserWidget.h
@@ -193,7 +193,10 @@ public:
 	FEventReply OnMouseButtonDoubleClick(FGeometry InMyGeometry, const FPointerEvent& InMouseEvent);
 
 	UFUNCTION(BlueprintImplementableEvent, Category="Drag and Drop")
-	void OnDragDetected(FGeometry MyGeometry, const FPointerEvent& MouseEvent, UDragDropOperation*& Operation);
+	void OnDragDetected(FGeometry MyGeometry, const FPointerEvent& PointerEvent, UDragDropOperation*& Operation);
+
+	UFUNCTION(BlueprintImplementableEvent, Category="Drag and Drop")
+	void OnDragCancelled(const FPointerEvent& PointerEvent, UDragDropOperation* Operation);
 	
 	UFUNCTION(BlueprintImplementableEvent, Category="Drag and Drop")
 	void OnDragEnter(FGeometry MyGeometry, FPointerEvent PointerEvent, UDragDropOperation* Operation);
diff --git a/Engine/Source/Runtime/UMG/Public/Slate/SObjectWidget.h b/Engine/Source/Runtime/UMG/Public/Slate/SObjectWidget.h
--- a/Engine/Source/Runtime/UMG/Public/Slate/SObjectWidget.h
+++ b/Engine/Source/Runtime/UMG/Public/Slate/SObjectWidget.h
@@ -56,6 +56,8 @@ class SObjectWidget : public SCompoundWidget, public FGCObject
 	virtual void OnDragLeave(const FDragDropEvent& DragDropEvent) override;
 	virtual FReply OnDragOver(const FGeometry& MyGeometry, const FDragDropEvent& DragDropEvent) override;
 	virtual FReply OnDrop(const FGeometry& MyGeometry, const FDragDropEvent& DragDropEvent) override;
+
+	void OnDragCancelled(const FDragDropEvent& DragDropEvent, UDragDropOperation* Operation);
 	
 	virtual FReply OnControllerButtonPressed(const FGeometry& MyGeometry, const FControllerEvent& ControllerEvent) override;
 	virtual FReply OnControllerButtonReleased(const FGeometry& MyGeometry, const FControllerEvent& ControllerEvent) override;
diff --git a/Engine/Source/Runtime/UMG/Public/Slate/UMGDragDropOp.h b/Engine/Source/Runtime/UMG/Public/Slate/UMGDragDropOp.h
--- a/Engine/Source/Runtime/UMG/Public/Slate/UMGDragDropOp.h
+++ b/Engine/Source/Runtime/UMG/Public/Slate/UMGDragDropOp.h
@@ -13,7 +13,7 @@ class FUMGDragDropOp : public FDragDropOperation
 public:
 	DRAG_DROP_OPERATOR_TYPE(FUMGDragDropOp, FDragDropOperation)
 	
-	static TSharedRef<FUMGDragDropOp> New(UDragDropOperation* Operation, const FVector2D &CursorPosition, const FVector2D &ScreenPositionOfNode);
+	static TSharedRef<FUMGDragDropOp> New(UDragDropOperation* Operation, const FVector2D &CursorPosition, const FVector2D &ScreenPositionOfNode, TSharedPtr<SObjectWidget> SourceUserWidget);
 
 	virtual void OnDrop( bool bDropWasHandled, const FPointerEvent& MouseEvent ) override;
 
@@ -27,6 +27,9 @@ private:
 
 	TWeakObjectPtr<UDragDropOperation> DragOperation;
 
+	/** Source User Widget */
+	TSharedPtr<SObjectWidget> SourceUserWidget;
+
 	/** The widget used during the drag/drop action to show something being dragged. */
 	TSharedPtr<SWidget> DecoratorWidget;
 
