UMG - Starting work on a string combobox.

diff --git a/Engine/Source/Runtime/UMG/Private/Components/ComboBox.cpp b/Engine/Source/Runtime/UMG/Private/Components/ComboBox.cpp
--- a/Engine/Source/Runtime/UMG/Private/Components/ComboBox.cpp
+++ b/Engine/Source/Runtime/UMG/Private/Components/ComboBox.cpp
@@ -14,20 +14,23 @@ UComboBox::UComboBox(const FPostConstructInitializeProperties& PCIP)
 
 TSharedRef<SWidget> UComboBox::RebuildWidget()
 {
-	TSharedRef< SComboBox<UObject*> > NewComboBox =
+	TSet<UObject*> UniqueItems(Items);
+	Items = UniqueItems.Array();
+
+	MyComboBox =
 		SNew(SComboBox<UObject*>)
 		.OptionsSource(&Items)
 		.OnGenerateWidget(BIND_UOBJECT_DELEGATE(SComboBox<UObject*>::FOnGenerateWidget, HandleGenerateWidget));
 
-	return NewComboBox;
+	return MyComboBox.ToSharedRef();
 }
 
 TSharedRef<SWidget> UComboBox::HandleGenerateWidget(UObject* Item) const
 {
 	// Call the user's delegate to see if they want to generate a custom widget bound to the data source.
-	if ( OnGenerateWidget.IsBound() )
+	if ( OnGenerateWidgetEvent.IsBound() )
 	{
-		UWidget* Widget = OnGenerateWidget.Execute(Item);
+		UWidget* Widget = OnGenerateWidgetEvent.Execute(Item);
 		if ( Widget != NULL )
 		{
 			return Widget->TakeWidget();
diff --git a/Engine/Source/Runtime/UMG/Private/Components/ComboBoxString.cpp b/Engine/Source/Runtime/UMG/Private/Components/ComboBoxString.cpp
new file mode 100644
--- /dev/null
+++ b/Engine/Source/Runtime/UMG/Private/Components/ComboBoxString.cpp
@@ -0,0 +1,152 @@
+// Copyright 1998-2014 Epic Games, Inc. All Rights Reserved.
+
+#include "UMGPrivatePCH.h"
+
+#define LOCTEXT_NAMESPACE "UMG"
+
+/////////////////////////////////////////////////////
+// UComboBoxString
+
+UComboBoxString::UComboBoxString(const FPostConstructInitializeProperties& PCIP)
+	: Super(PCIP)
+{
+	ContentPadding = FMargin(4.0, 2.0);
+	MaxListHeight = 450.0f;
+	HasDownArrow = true;
+}
+
+TSharedRef<SWidget> UComboBoxString::RebuildWidget()
+{
+	for ( FString& DefaultOptions : DefaultOptions )
+	{
+		AddOption(DefaultOptions);
+	}
+
+	TSharedPtr<FString> InitiallySelectedItem;
+	int32 InitialIndex = FindOptionIndex(InitiallySelectedOption);
+	if ( InitialIndex != -1 )
+	{
+		InitiallySelectedItem = Options[InitialIndex];
+	}
+
+	MyComboBox =
+		SNew(SComboBox< TSharedPtr<FString> >)
+		.OptionsSource(&Options)
+		.InitiallySelectedItem(InitiallySelectedItem)
+		.ContentPadding(ContentPadding)
+		.MaxListHeight(MaxListHeight)
+		.HasDownArrow(HasDownArrow)
+		.Method(SMenuAnchor::UseCurrentWindow)
+		.OnGenerateWidget(BIND_UOBJECT_DELEGATE(SComboBox< TSharedPtr<FString> >::FOnGenerateWidget, HandleGenerateWidget))
+		.OnSelectionChanged(BIND_UOBJECT_DELEGATE(SComboBox< TSharedPtr<FString> >::FOnSelectionChanged, HandleSelectionChanged))
+		.OnComboBoxOpening(BIND_UOBJECT_DELEGATE(FOnComboBoxOpening, HandleOpening))
+		[
+			SAssignNew(ComoboBoxContent, SBox)
+		];
+
+	if ( InitialIndex != -1 )
+	{
+		// Generate the widget for the initially selected widget if needed
+		ComoboBoxContent->SetContent(HandleGenerateWidget(InitiallySelectedItem));
+	}
+
+	//TODO UMG Expose the rest of these options
+
+	//SLATE_STYLE_ARGUMENT(FComboBoxStyle, ComboBoxStyle)
+
+	//	/** The visual style of the button part of the combo box (overrides ComboBoxStyle) */
+	//	SLATE_STYLE_ARGUMENT(FButtonStyle, ButtonStyle)
+
+	//	SLATE_ATTRIBUTE(FMargin, ContentPadding)
+	//	SLATE_ATTRIBUTE(FSlateColor, ForegroundColor)
+
+	//	/** The sound to play when the button is pressed (overrides ComboBoxStyle) */
+	//	SLATE_ARGUMENT(TOptional<FSlateSound>, PressedSoundOverride)
+
+	//	/** The sound to play when the selection changes (overrides ComboBoxStyle) */
+	//	SLATE_ARGUMENT(TOptional<FSlateSound>, SelectionChangeSoundOverride)
+
+	return MyComboBox.ToSharedRef();
+}
+
+void UComboBoxString::AddOption(const FString& Option)
+{
+	Options.Add(MakeShareable(new FString(Option)));
+}
+
+bool UComboBoxString::RemoveOption(const FString& Option)
+{
+	int32 OptionIndex = FindOptionIndex(Option);
+	if ( OptionIndex != -1 )
+	{
+		Options.RemoveAt(OptionIndex);
+		return true;
+	}
+
+	return false;
+}
+
+int32 UComboBoxString::FindOptionIndex(const FString& Option)
+{
+	for ( int32 OptionIndex = 0; OptionIndex < Options.Num(); OptionIndex++ )
+	{
+		TSharedPtr<FString>& OptionAtIndex = Options[OptionIndex];
+
+		if ( ( *OptionAtIndex ) == Option )
+		{
+			return OptionIndex;
+		}
+	}
+
+	return -1;
+}
+
+void UComboBoxString::ClearOptions()
+{
+	Options.Empty();
+}
+
+TSharedRef<SWidget> UComboBoxString::HandleGenerateWidget(TSharedPtr<FString> Item) const
+{
+	// Call the user's delegate to see if they want to generate a custom widget bound to the data source.
+	if ( !IsDesignTime() && OnGenerateWidgetEvent.IsBound() )
+	{
+		UWidget* Widget = OnGenerateWidgetEvent.Execute(*Item);
+		if ( Widget != NULL )
+		{
+			return Widget->TakeWidget();
+		}
+	}
+
+	// If a row wasn't generated just create the default one, a simple text block of the item's name.
+	return SNew(STextBlock).Text(FText::FromString(*Item));
+}
+
+void UComboBoxString::HandleSelectionChanged(TSharedPtr<FString> Item, ESelectInfo::Type SelectionType)
+{
+	if ( !IsDesignTime() )
+	{
+		OnSelectionChanged.Broadcast(*Item, SelectionType);
+	}
+
+	// When the selection changes we always generate another widget to represent the content area of the comobox.
+	ComoboBoxContent->SetContent( HandleGenerateWidget(Item) );
+}
+
+void UComboBoxString::HandleOpening()
+{
+	OnOpening.Broadcast();
+}
+
+#if WITH_EDITOR
+
+const FSlateBrush* UComboBoxString::GetEditorIcon()
+{
+	return FUMGStyle::Get().GetBrush("Widget.ComboBox");
+}
+
+#endif
+
+/////////////////////////////////////////////////////
+
+#undef LOCTEXT_NAMESPACE
diff --git a/Engine/Source/Runtime/UMG/Private/WidgetBlueprintGeneratedClass.cpp b/Engine/Source/Runtime/UMG/Private/WidgetBlueprintGeneratedClass.cpp
--- a/Engine/Source/Runtime/UMG/Private/WidgetBlueprintGeneratedClass.cpp
+++ b/Engine/Source/Runtime/UMG/Private/WidgetBlueprintGeneratedClass.cpp
@@ -158,10 +158,6 @@ void UWidgetBlueprintGeneratedClass::InitializeWidget(UUserWidget* UserWidget) c
 
 		//TODO Add OnWidgetInitialized
 	}
-
-	// Create the input component so that we can simulate input to the blueprint in the forum users are accustomed to.
-	UserWidget->InputComponent = ConstructObject<UInputComponent>(UInputComponent::StaticClass(), UserWidget);
-	UInputDelegateBinding::BindInputDelegates(this, UserWidget->InputComponent);
 }
 
 #undef LOCTEXT_NAMESPACE
diff --git a/Engine/Source/Runtime/UMG/Public/Blueprint/UserWidget.h b/Engine/Source/Runtime/UMG/Public/Blueprint/UserWidget.h
--- a/Engine/Source/Runtime/UMG/Public/Blueprint/UserWidget.h
+++ b/Engine/Source/Runtime/UMG/Public/Blueprint/UserWidget.h
@@ -334,9 +334,6 @@ public:
 	UPROPERTY(Transient)
 	TArray<UUMGSequencePlayer*> StoppedSequencePlayers;
 
-	UPROPERTY(Transient)
-	UInputComponent* InputComponent;
-
 protected:
 	virtual TSharedRef<SWidget> RebuildWidget() override;
 
diff --git a/Engine/Source/Runtime/UMG/Public/Components/ComboBox.h b/Engine/Source/Runtime/UMG/Public/Components/ComboBox.h
--- a/Engine/Source/Runtime/UMG/Public/Components/ComboBox.h
+++ b/Engine/Source/Runtime/UMG/Public/Components/ComboBox.h
@@ -4,8 +4,10 @@
 
 #include "ComboBox.generated.h"
 
-/** The combobox allows you to display a list of options to the user in a dropdown menu for them to select one. */
-UCLASS(meta=( Category="Misc" ), ClassGroup=UserInterface)
+/**
+ * The combobox allows you to display a list of options to the user in a dropdown menu for them to select one.
+ */
+UCLASS( meta=( DisplayName="ComboBox (Object)", Category="Misc" ), ClassGroup=UserInterface)
 class UMG_API UComboBox : public UWidget
 {
 	GENERATED_UCLASS_BODY()
@@ -18,7 +20,7 @@ public:
 
 	/** Called when the widget is needed for the item. */
 	UPROPERTY(EditDefaultsOnly, Category=Events)
-	FGenerateWidgetUObject OnGenerateWidget;
+	FGenerateWidgetForObject OnGenerateWidgetEvent;
 
 #if WITH_EDITOR
 	virtual const FSlateBrush* GetEditorIcon() override;
@@ -30,4 +32,8 @@ protected:
 	// UWidget interface
 	virtual TSharedRef<SWidget> RebuildWidget() override;
 	// End of UWidget interface
+
+protected:
+
+	TSharedPtr< SComboBox<UObject*> > MyComboBox;
 };
diff --git a/Engine/Source/Runtime/UMG/Public/Components/ComboBoxString.h b/Engine/Source/Runtime/UMG/Public/Components/ComboBoxString.h
new file mode 100644
--- /dev/null
+++ b/Engine/Source/Runtime/UMG/Public/Components/ComboBoxString.h
@@ -0,0 +1,95 @@
+// Copyright 1998-2014 Epic Games, Inc. All Rights Reserved.
+
+#pragma once
+
+#include "ComboBoxString.generated.h"
+
+/**
+ * The combobox allows you to display a list of options to the user in a dropdown menu for them to select one.
+ */
+UCLASS( meta=( DisplayName="ComboBox (String)", Category="Misc" ), ClassGroup=UserInterface)
+class UMG_API UComboBoxString : public UWidget
+{
+	GENERATED_UCLASS_BODY()
+
+	DECLARE_DYNAMIC_MULTICAST_DELEGATE_TwoParams(FOnSelectionChangedEvent, FString, SelectedItem, ESelectInfo::Type, SelectionType);
+	DECLARE_DYNAMIC_MULTICAST_DELEGATE(FOnOpeningEvent);
+
+public:
+
+	/** The list of items to be displayed on the combobox. */
+	UPROPERTY(EditDefaultsOnly, Category=Content)
+	TArray<FString> DefaultOptions;
+
+	/** The item in the combobox to select by default */
+	UPROPERTY(EditDefaultsOnly, Category=Content)
+	FString InitiallySelectedOption;
+
+	UPROPERTY(EditDefaultsOnly, Category=Content)
+	FMargin ContentPadding;
+
+	/** The max height of the combobox list that opens */
+	UPROPERTY(EditDefaultsOnly, Category=Content, AdvancedDisplay)
+	float MaxListHeight;
+
+	/**
+	 * When false, the down arrow is not generated and it is up to the API consumer
+	 * to make their own visual hint that this is a drop down.
+	 */
+	UPROPERTY(EditDefaultsOnly, Category=Content, AdvancedDisplay)
+	bool HasDownArrow;
+
+	/** Called when the widget is needed for the item. */
+	UPROPERTY(EditDefaultsOnly, Category=Events)
+	FGenerateWidgetForString OnGenerateWidgetEvent;
+
+	/** Called when a new item is selected in the combobox. */
+	UPROPERTY(BlueprintAssignable, Category=Events)
+	FOnSelectionChangedEvent OnSelectionChanged;
+
+	/** Called when the combobox is opening */
+	UPROPERTY(BlueprintAssignable, Category=Events)
+	FOnOpeningEvent OnOpening;
+
+public:
+
+	UFUNCTION(BlueprintCallable, Category="ComboBox")
+	void AddOption(const FString& Option);
+
+	UFUNCTION(BlueprintCallable, Category="ComboBox")
+	bool RemoveOption(const FString& Option);
+
+	UFUNCTION(BlueprintCallable, Category="ComboBox")
+	int32 FindOptionIndex(const FString& Option);
+
+	UFUNCTION(BlueprintCallable, Category="ComboBox")
+	void ClearOptions();
+
+#if WITH_EDITOR
+	virtual const FSlateBrush* GetEditorIcon() override;
+#endif
+
+protected:
+	/** Called by slate when it needs to generate a new item for the combobox */
+	TSharedRef<SWidget> HandleGenerateWidget(TSharedPtr<FString> Item) const;
+
+	/** Called by slate when the underlying comobobox selection changes */
+	void HandleSelectionChanged(TSharedPtr<FString> Item, ESelectInfo::Type SelectionType);
+
+	/** Called by slate when the underlying comobobox is opening */
+	void HandleOpening();
+
+	// UWidget interface
+	virtual TSharedRef<SWidget> RebuildWidget() override;
+	// End of UWidget interface
+
+protected:
+	/** The true objects bound to the Slate combobox. */
+	TArray< TSharedPtr<FString> > Options;
+
+	/** A shared pointer to the underlying slate combobox */
+	TSharedPtr< SComboBox< TSharedPtr<FString> > > MyComboBox;
+
+	/** A shared pointer to a container that holds the comobobox content that is selected */
+	TSharedPtr< SBox > ComoboBoxContent;
+};
diff --git a/Engine/Source/Runtime/UMG/Public/Components/Widget.h b/Engine/Source/Runtime/UMG/Public/Components/Widget.h
--- a/Engine/Source/Runtime/UMG/Public/Components/Widget.h
+++ b/Engine/Source/Runtime/UMG/Public/Components/Widget.h
@@ -44,7 +44,8 @@ public:
 	DECLARE_DYNAMIC_DELEGATE_RetVal(ESlateCheckBoxState::Type, FGetCheckBoxState);
 	DECLARE_DYNAMIC_DELEGATE_RetVal(UWidget*, FGetContent);
 
-	DECLARE_DYNAMIC_DELEGATE_RetVal_OneParam(UWidget*, FGenerateWidgetUObject, UObject*, Item);
+	DECLARE_DYNAMIC_DELEGATE_RetVal_OneParam(UWidget*, FGenerateWidgetForString, FString, Item);
+	DECLARE_DYNAMIC_DELEGATE_RetVal_OneParam(UWidget*, FGenerateWidgetForObject, UObject*, Item);
 
 	// Events
 	DECLARE_DYNAMIC_DELEGATE_RetVal(FEventReply, FOnReply);
diff --git a/Engine/Source/Runtime/UMG/Public/UMG.h b/Engine/Source/Runtime/UMG/Public/UMG.h
--- a/Engine/Source/Runtime/UMG/Public/UMG.h
+++ b/Engine/Source/Runtime/UMG/Public/UMG.h
@@ -82,7 +82,9 @@
 
 #include "Slider.h"
 #include "ProgressBar.h"
+
 #include "ComboBox.h"
+#include "ComboBoxString.h"
 
 #include "SScaleBox.h"
 #include "ScaleBoxSlot.h"
