UMG - Viewport no longer makes a real viewport at design time.

diff --git a/Engine/Source/Runtime/UMG/Private/Components/Viewport.cpp b/Engine/Source/Runtime/UMG/Private/Components/Viewport.cpp
--- a/Engine/Source/Runtime/UMG/Private/Components/Viewport.cpp
+++ b/Engine/Source/Runtime/UMG/Private/Components/Viewport.cpp
@@ -508,22 +508,38 @@ void UViewport::ReleaseNativeWidget()
 
 TSharedRef<SWidget> UViewport::RebuildWidget()
 {
-	ViewportWidget = SNew(SAutoRefreshViewport);
-
-	if ( GetChildrenCount() > 0 )
+	if ( IsDesignTime() )
 	{
-		ViewportWidget->SetContent(GetContentSlot()->Content ? GetContentSlot()->Content->TakeWidget() : SNullWidget::NullWidget);
+		return BuildDesignTimeWidget(SNew(SBox)
+			.HAlign(HAlign_Center)
+			.VAlign(VAlign_Center)
+			[
+				SNew(STextBlock)
+				.Text(LOCTEXT("Viewport", "Viewport"))
+			]);
 	}
+	else
+	{
+		ViewportWidget = SNew(SAutoRefreshViewport);
+
+		if ( GetChildrenCount() > 0 )
+		{
+			ViewportWidget->SetContent(GetContentSlot()->Content ? GetContentSlot()->Content->TakeWidget() : SNullWidget::NullWidget);
+		}
 
-	return BuildDesignTimeWidget( ViewportWidget.ToSharedRef() );
+		return BuildDesignTimeWidget(ViewportWidget.ToSharedRef());
+	}
 }
 
 void UViewport::SynchronizeProperties()
 {
 	Super::SynchronizeProperties();
 
-	ViewportWidget->ViewportClient->SetBackgroundColor(BackgroundColor);
-	ViewportWidget->ViewportClient->SetEngineShowFlags(ShowFlags);
+	if ( ViewportWidget.IsValid() )
+	{
+		ViewportWidget->ViewportClient->SetBackgroundColor(BackgroundColor);
+		ViewportWidget->ViewportClient->SetEngineShowFlags(ShowFlags);
+	}
 }
 
 void UViewport::OnSlotAdded(UPanelSlot* Slot)
