UMG - Adding IsFocusable to UButton.

diff --git a/Engine/Source/Runtime/UMG/Private/Components/Button.cpp b/Engine/Source/Runtime/UMG/Private/Components/Button.cpp
--- a/Engine/Source/Runtime/UMG/Private/Components/Button.cpp
+++ b/Engine/Source/Runtime/UMG/Private/Components/Button.cpp
@@ -23,6 +23,8 @@ UButton::UButton(const FPostConstructInitializeProperties& PCIP)
 
 	ClickMethod = EButtonClickMethod::DownAndUp;
 	TouchMethod = EButtonTouchMethod::DownAndUp;
+
+	IsFocusable = true;
 }
 
 void UButton::ReleaseNativeWidget()
@@ -36,7 +38,8 @@ TSharedRef<SWidget> UButton::RebuildWidget()
 {
 	MyButton = SNew(SButton)
 		.ClickMethod(ClickMethod)
-		.TouchMethod(TouchMethod);
+		.TouchMethod(TouchMethod)
+		.IsFocusable(IsFocusable);
 
 	if ( GetChildrenCount() > 0 )
 	{
diff --git a/Engine/Source/Runtime/UMG/Public/Blueprint/UserWidget.h b/Engine/Source/Runtime/UMG/Public/Blueprint/UserWidget.h
--- a/Engine/Source/Runtime/UMG/Public/Blueprint/UserWidget.h
+++ b/Engine/Source/Runtime/UMG/Public/Blueprint/UserWidget.h
@@ -363,12 +363,12 @@ T* CreateWidget(UWorld* World, UClass* UserWidgetClass)
 	if ( !UserWidgetClass->IsChildOf(UUserWidget::StaticClass()) )
 	{
 		// TODO UMG Error?
-		return NULL;
+		return nullptr;
 	}
 
 	ULocalPlayer* Player = World->GetFirstLocalPlayerFromController();
 
-	UObject* Outer = ( Player == NULL ) ? StaticCast<UObject*>(World) : StaticCast<UObject*>(Player);
+	UObject* Outer = ( Player == nullptr ) ? StaticCast<UObject*>(World) : StaticCast<UObject*>(Player);
 	UUserWidget* NewWidget = ConstructObject<UUserWidget>(UserWidgetClass, Outer);
 	NewWidget->SetPlayerContext(FLocalPlayerContext(Player));
 	NewWidget->Initialize();
@@ -382,7 +382,7 @@ T* CreateWidget(APlayerController* OwningPlayer, UClass* UserWidgetClass)
 	if ( !UserWidgetClass->IsChildOf(UUserWidget::StaticClass()) )
 	{
 		// TODO UMG Error?
-		return NULL;
+		return nullptr;
 	}
 
 	UUserWidget* NewWidget = ConstructObject<UUserWidget>(UserWidgetClass, OwningPlayer);
diff --git a/Engine/Source/Runtime/UMG/Public/Components/Button.h b/Engine/Source/Runtime/UMG/Public/Components/Button.h
--- a/Engine/Source/Runtime/UMG/Public/Components/Button.h
+++ b/Engine/Source/Runtime/UMG/Public/Components/Button.h
@@ -43,6 +43,10 @@ public:
 	UPROPERTY(EditDefaultsOnly, Category="Interaction")
 	TEnumAsByte<EButtonTouchMethod::Type> TouchMethod;
 
+	/** Sometimes a button should only be mouse-clickable and never keyboard focusable. */
+	UPROPERTY(EditDefaultsOnly, Category="Interaction", AdvancedDisplay)
+	bool IsFocusable;
+
 	/** The sound made when the user 'clicks' the button */
 	UPROPERTY(EditDefaultsOnly, Category=Sound)
 	FSlateSound PressedSound;
