UMG - Preventing the comboboxes from leaking widgets.

diff --git a/Engine/Source/Runtime/UMG/Private/Components/ComboBox.cpp b/Engine/Source/Runtime/UMG/Private/Components/ComboBox.cpp
--- a/Engine/Source/Runtime/UMG/Private/Components/ComboBox.cpp
+++ b/Engine/Source/Runtime/UMG/Private/Components/ComboBox.cpp
@@ -12,6 +12,13 @@ UComboBox::UComboBox(const FPostConstructInitializeProperties& PCIP)
 {
 }
 
+void UComboBox::ReleaseNativeWidget()
+{
+	Super::ReleaseNativeWidget();
+
+	MyComboBox.Reset();
+}
+
 TSharedRef<SWidget> UComboBox::RebuildWidget()
 {
 	TSet<UObject*> UniqueItems(Items);
diff --git a/Engine/Source/Runtime/UMG/Private/Components/ComboBoxString.cpp b/Engine/Source/Runtime/UMG/Private/Components/ComboBoxString.cpp
--- a/Engine/Source/Runtime/UMG/Private/Components/ComboBoxString.cpp
+++ b/Engine/Source/Runtime/UMG/Private/Components/ComboBoxString.cpp
@@ -15,6 +15,14 @@ UComboBoxString::UComboBoxString(const FPostConstructInitializeProperties& PCIP)
 	HasDownArrow = true;
 }
 
+void UComboBoxString::ReleaseNativeWidget()
+{
+	Super::ReleaseNativeWidget();
+
+	MyComboBox.Reset();
+	ComoboBoxContent.Reset();
+}
+
 TSharedRef<SWidget> UComboBoxString::RebuildWidget()
 {
 	for ( FString& DefaultOptions : DefaultOptions )
@@ -57,9 +65,6 @@ TSharedRef<SWidget> UComboBoxString::RebuildWidget()
 	//	/** The visual style of the button part of the combo box (overrides ComboBoxStyle) */
 	//	SLATE_STYLE_ARGUMENT(FButtonStyle, ButtonStyle)
 
-	//	SLATE_ATTRIBUTE(FMargin, ContentPadding)
-	//	SLATE_ATTRIBUTE(FSlateColor, ForegroundColor)
-
 	//	/** The sound to play when the button is pressed (overrides ComboBoxStyle) */
 	//	SLATE_ARGUMENT(TOptional<FSlateSound>, PressedSoundOverride)
 
diff --git a/Engine/Source/Runtime/UMG/Public/Components/ComboBox.h b/Engine/Source/Runtime/UMG/Public/Components/ComboBox.h
--- a/Engine/Source/Runtime/UMG/Public/Components/ComboBox.h
+++ b/Engine/Source/Runtime/UMG/Public/Components/ComboBox.h
@@ -22,6 +22,8 @@ public:
 	UPROPERTY(EditDefaultsOnly, Category=Events)
 	FGenerateWidgetForObject OnGenerateWidgetEvent;
 
+	virtual void ReleaseNativeWidget() override;
+
 #if WITH_EDITOR
 	virtual const FSlateBrush* GetEditorIcon() override;
 #endif
diff --git a/Engine/Source/Runtime/UMG/Public/Components/ComboBoxString.h b/Engine/Source/Runtime/UMG/Public/Components/ComboBoxString.h
--- a/Engine/Source/Runtime/UMG/Public/Components/ComboBoxString.h
+++ b/Engine/Source/Runtime/UMG/Public/Components/ComboBoxString.h
@@ -65,6 +65,8 @@ public:
 	UFUNCTION(BlueprintCallable, Category="ComboBox")
 	void ClearOptions();
 
+	virtual void ReleaseNativeWidget() override;
+
 #if WITH_EDITOR
 	virtual const FSlateBrush* GetEditorIcon() override;
 #endif
