UMG - The style of a button is now dynamically adjustable at runtime.

diff --git a/Engine/Source/Runtime/UMG/Private/Components/Button.cpp b/Engine/Source/Runtime/UMG/Private/Components/Button.cpp
--- a/Engine/Source/Runtime/UMG/Private/Components/Button.cpp
+++ b/Engine/Source/Runtime/UMG/Private/Components/Button.cpp
@@ -53,14 +53,6 @@ void UButton::SynchronizeProperties()
 {
 	Super::SynchronizeProperties();
 
-	SButton::FArguments ButtonDefaults;
-
-	const FButtonStyle* StylePtr = ( Style != NULL ) ? Style->GetStyle<FButtonStyle>() : NULL;
-	if ( StylePtr == NULL )
-	{
-		StylePtr = ButtonDefaults._ButtonStyle;
-	}
-
 	TOptional<FSlateSound> OptionalPressedSound;
 	if ( PressedSound.GetResourceObject() )
 	{
@@ -73,7 +65,15 @@ void UButton::SynchronizeProperties()
 		OptionalHoveredSound = HoveredSound;
 	}
 
-	MyButton->SetButtonStyle( StylePtr );
+	if ( MyStyle.IsSet() )
+	{
+		MyButton->SetButtonStyle(&MyStyle.GetValue());
+	}
+	else
+	{
+		const FButtonStyle* StylePtr = GetStyle();
+		MyButton->SetButtonStyle(StylePtr);
+	}
 
 	MyButton->SetColorAndOpacity( ColorAndOpacity );
 	MyButton->SetBorderBackgroundColor( BackgroundColor );
@@ -109,6 +109,49 @@ void UButton::OnSlotRemoved(UPanelSlot* Slot)
 	}
 }
 
+void UButton::SetStyle(USlateWidgetStyleAsset* InStyle)
+{
+	Style = InStyle;
+	MyStyle = TOptional<FButtonStyle>();
+
+	const FButtonStyle* StylePtr = GetStyle();
+
+	if ( MyButton.IsValid() )
+	{
+		MyButton->SetButtonStyle(StylePtr);
+	}
+}
+
+const FButtonStyle* UButton::GetStyle() const
+{
+	const FButtonStyle* StylePtr = ( Style != NULL ) ? Style->GetStyle<FButtonStyle>() : NULL;
+	if ( StylePtr == NULL )
+	{
+		SButton::FArguments ButtonDefaults;
+		StylePtr = ButtonDefaults._ButtonStyle;
+	}
+
+	return StylePtr;
+}
+
+void UButton::SetButtonStyle(FButtonStyle InButtonStyle)
+{
+	MyStyle = InButtonStyle;
+}
+
+FButtonStyle UButton::GetButtonStyle()
+{
+	// If the dynamic style hasn't been set, default it to a clone of the current
+	// button style asset.
+	if ( !MyStyle.IsSet() )
+	{
+		const FButtonStyle* StylePtr = GetStyle();
+		MyStyle = *StylePtr;
+	}
+	
+	return MyStyle.GetValue();
+}
+
 void UButton::SetColorAndOpacity(FLinearColor Color)
 {
 	ColorAndOpacity = Color;
diff --git a/Engine/Source/Runtime/UMG/Public/Components/Button.h b/Engine/Source/Runtime/UMG/Public/Components/Button.h
--- a/Engine/Source/Runtime/UMG/Public/Components/Button.h
+++ b/Engine/Source/Runtime/UMG/Public/Components/Button.h
@@ -36,11 +36,11 @@ public:
 	FLinearColor BackgroundColor;
 
 	/** The type of mouse action required by the user to trigger the buttons 'Click' */
-	UPROPERTY(EditDefaultsOnly, Category="Interaction")
+	UPROPERTY(EditDefaultsOnly, Category="Interaction", AdvancedDisplay)
 	TEnumAsByte<EButtonClickMethod::Type> ClickMethod;
 
 	/** The type of touch action required by the user to trigger the buttons 'Click' */
-	UPROPERTY(EditDefaultsOnly, Category="Interaction")
+	UPROPERTY(EditDefaultsOnly, Category="Interaction", AdvancedDisplay)
 	TEnumAsByte<EButtonTouchMethod::Type> TouchMethod;
 
 	/** Sometimes a button should only be mouse-clickable and never keyboard focusable. */
@@ -60,17 +60,29 @@ public:
 	/** Called when the button is clicked */
 	UPROPERTY(BlueprintAssignable)
 	FOnButtonClickedEvent OnClicked;
+
+	/** Sets the look and feel of a button by giving it a new style from an existing style asset. */
+	UFUNCTION(BlueprintCallable, Category="Button|Appearance")
+	void SetStyle(USlateWidgetStyleAsset* InStyle);
+
+	/** Sets the look and feel of a button from a new button style struct. */
+	UFUNCTION(BlueprintCallable, Category="Button|Appearance")
+	void SetButtonStyle(FButtonStyle InButtonStyle);
+
+	/** Gets a dynamic button style that can be changed. */
+	UFUNCTION(BlueprintCallable, Category="Button|Appearance")
+	FButtonStyle GetButtonStyle();
 	
 	/** Sets the color multiplier for the button content */
-	UFUNCTION(BlueprintCallable, Category="Appearance")
+	UFUNCTION(BlueprintCallable, Category="Button|Appearance")
 	void SetColorAndOpacity(FLinearColor InColorAndOpacity);
 
 	/** Sets the color multiplier for the button background */
-	UFUNCTION(BlueprintCallable, Category="Appearance")
+	UFUNCTION(BlueprintCallable, Category="Button|Appearance")
 	void SetBackgroundColor(FLinearColor InBackgroundColor);
 
 	/** @return true if the user is actively pressing the button otherwise false.  For detecting Clicks, use the OnClicked event. */
-	UFUNCTION(BlueprintCallable, Category="Widget")
+	UFUNCTION(BlueprintCallable, Category="Button")
 	bool IsPressed() const;
 
 	// UWidget interface
@@ -98,6 +110,7 @@ protected:
 	// End UPanelWidget
 
 protected:
+	/** Handle the actual click event from slate and forward it on */
 	FReply SlateHandleClicked();
 
 protected:
@@ -105,6 +118,13 @@ protected:
 	virtual TSharedRef<SWidget> RebuildWidget() override;
 	// End of UWidget interface
 
+	/** Get the button style of the current style asset */
+	const FButtonStyle* GetStyle() const;
+
 protected:
+	/** Cached pointer to the underlying slate button owned by this UWidget */
 	TSharedPtr<SButton> MyButton;
+
+	/** Stores the style set by users dynamically. */
+	TOptional<FButtonStyle> MyStyle;
 };
