UMG - Setting the FButtonStyle sets the style on the underlying slate widget now.

diff --git a/Engine/Source/Runtime/UMG/Private/Components/Button.cpp b/Engine/Source/Runtime/UMG/Private/Components/Button.cpp
--- a/Engine/Source/Runtime/UMG/Private/Components/Button.cpp
+++ b/Engine/Source/Runtime/UMG/Private/Components/Button.cpp
@@ -137,6 +137,11 @@ const FButtonStyle* UButton::GetStyle() const
 void UButton::SetButtonStyle(FButtonStyle InButtonStyle)
 {
 	MyStyle = InButtonStyle;
+
+	if ( MyButton.IsValid() )
+	{
+		MyButton->SetButtonStyle(&MyStyle.GetValue());
+	}
 }
 
 FButtonStyle UButton::GetButtonStyle()
