UMG - Changing all BlueprintImplementableEvents to BlupeprintNativeEvent on the UUserWidget so that native implementations can handle those events as well.

diff --git a/Engine/Source/Runtime/UMG/Private/UserWidget.cpp b/Engine/Source/Runtime/UMG/Private/UserWidget.cpp
--- a/Engine/Source/Runtime/UMG/Private/UserWidget.cpp
+++ b/Engine/Source/Runtime/UMG/Private/UserWidget.cpp
@@ -5,6 +5,7 @@
 #include "UMGSequencePlayer.h"
 #include "SceneViewport.h"
 #include "WidgetAnimation.h"
+#include "WidgetBlueprintLibrary.h"
 
 /**
  * This class holds onto the widget when it's placed into the viewport.
@@ -200,6 +201,161 @@ UWorld* UUserWidget::GetWorld() const
 	return CachedWorld;
 }
 
+void UUserWidget::Construct_Implementation()
+{
+
+}
+
+void UUserWidget::Tick_Implementation(FGeometry MyGeometry, float InDeltaTime)
+{
+
+}
+
+void UUserWidget::OnPaint_Implementation(FPaintContext& Context) const
+{
+
+}
+
+FEventReply UUserWidget::OnKeyboardFocusReceived_Implementation(FGeometry MyGeometry, FKeyboardFocusEvent InKeyboardFocusEvent)
+{
+	return UWidgetBlueprintLibrary::Unhandled();
+}
+
+void UUserWidget::OnKeyboardFocusLost_Implementation(FKeyboardFocusEvent InKeyboardFocusEvent)
+{
+
+}
+
+FEventReply UUserWidget::OnKeyChar_Implementation(FGeometry MyGeometry, FCharacterEvent InCharacterEvent)
+{
+	return UWidgetBlueprintLibrary::Unhandled();
+}
+
+FEventReply UUserWidget::OnPreviewKeyDown_Implementation(FGeometry MyGeometry, FKeyboardEvent InKeyboardEvent)
+{
+	return UWidgetBlueprintLibrary::Unhandled();
+}
+
+FEventReply UUserWidget::OnKeyDown_Implementation(FGeometry MyGeometry, FKeyboardEvent InKeyboardEvent)
+{
+	return UWidgetBlueprintLibrary::Unhandled();
+}
+
+FEventReply UUserWidget::OnKeyUp_Implementation(FGeometry MyGeometry, FKeyboardEvent InKeyboardEvent)
+{
+	return UWidgetBlueprintLibrary::Unhandled();
+}
+
+FEventReply UUserWidget::OnMouseButtonDown_Implementation(FGeometry MyGeometry, const FPointerEvent& MouseEvent)
+{
+	return UWidgetBlueprintLibrary::Unhandled();
+}
+
+FEventReply UUserWidget::OnPreviewMouseButtonDown_Implementation(FGeometry MyGeometry, const FPointerEvent& MouseEvent)
+{
+	return UWidgetBlueprintLibrary::Unhandled();
+}
+
+FEventReply UUserWidget::OnMouseButtonUp_Implementation(FGeometry MyGeometry, const FPointerEvent& MouseEvent)
+{
+	return UWidgetBlueprintLibrary::Unhandled();
+}
+
+FEventReply UUserWidget::OnMouseMove_Implementation(FGeometry MyGeometry, const FPointerEvent& MouseEvent)
+{
+	return UWidgetBlueprintLibrary::Unhandled();
+}
+
+void UUserWidget::OnMouseEnter_Implementation(FGeometry MyGeometry, const FPointerEvent& MouseEvent)
+{
+
+}
+
+void UUserWidget::OnMouseLeave_Implementation(const FPointerEvent& MouseEvent)
+{
+
+}
+
+FEventReply UUserWidget::OnMouseWheel_Implementation(FGeometry MyGeometry, const FPointerEvent& MouseEvent)
+{
+	return UWidgetBlueprintLibrary::Unhandled();
+}
+
+FEventReply UUserWidget::OnMouseButtonDoubleClick_Implementation(FGeometry InMyGeometry, const FPointerEvent& InMouseEvent)
+{
+	return UWidgetBlueprintLibrary::Unhandled();
+}
+
+void UUserWidget::OnDragDetected_Implementation(FGeometry MyGeometry, const FPointerEvent& PointerEvent, UDragDropOperation*& Operation)
+{
+
+}
+
+void UUserWidget::OnDragCancelled_Implementation(const FPointerEvent& PointerEvent, UDragDropOperation* Operation)
+{
+
+}
+
+void UUserWidget::OnDragEnter_Implementation(FGeometry MyGeometry, FPointerEvent PointerEvent, UDragDropOperation* Operation)
+{
+
+}
+
+void UUserWidget::OnDragLeave_Implementation(FPointerEvent PointerEvent, UDragDropOperation* Operation)
+{
+
+}
+
+bool UUserWidget::OnDragOver_Implementation(FGeometry MyGeometry, FPointerEvent PointerEvent, UDragDropOperation* Operation)
+{
+	return false;
+}
+
+bool UUserWidget::OnDrop_Implementation(FGeometry MyGeometry, FPointerEvent PointerEvent, UDragDropOperation* Operation)
+{
+	return false;
+}
+
+FEventReply UUserWidget::OnControllerButtonPressed_Implementation(FGeometry MyGeometry, FControllerEvent ControllerEvent)
+{
+	return UWidgetBlueprintLibrary::Unhandled();
+}
+
+FEventReply UUserWidget::OnControllerButtonReleased_Implementation(FGeometry MyGeometry, FControllerEvent ControllerEvent)
+{
+	return UWidgetBlueprintLibrary::Unhandled();
+}
+
+FEventReply UUserWidget::OnControllerAnalogValueChanged_Implementation(FGeometry MyGeometry, FControllerEvent ControllerEvent)
+{
+	return UWidgetBlueprintLibrary::Unhandled();
+}
+
+FEventReply UUserWidget::OnTouchGesture_Implementation(FGeometry MyGeometry, const FPointerEvent& GestureEvent)
+{
+	return UWidgetBlueprintLibrary::Unhandled();
+}
+
+FEventReply UUserWidget::OnTouchStarted_Implementation(FGeometry MyGeometry, const FPointerEvent& InTouchEvent)
+{
+	return UWidgetBlueprintLibrary::Unhandled();
+}
+
+FEventReply UUserWidget::OnTouchMoved_Implementation(FGeometry MyGeometry, const FPointerEvent& InTouchEvent)
+{
+	return UWidgetBlueprintLibrary::Unhandled();
+}
+
+FEventReply UUserWidget::OnTouchEnded_Implementation(FGeometry MyGeometry, const FPointerEvent& InTouchEvent)
+{
+	return UWidgetBlueprintLibrary::Unhandled();
+}
+
+FEventReply UUserWidget::OnMotionDetected_Implementation(FGeometry MyGeometry, FMotionEvent InMotionEvent)
+{
+	return UWidgetBlueprintLibrary::Unhandled();
+}
+
 void UUserWidget::PlayAnimation(const UWidgetAnimation* InAnimation)
 {
 	if( InAnimation )
diff --git a/Engine/Source/Runtime/UMG/Public/Blueprint/UserWidget.h b/Engine/Source/Runtime/UMG/Public/Blueprint/UserWidget.h
--- a/Engine/Source/Runtime/UMG/Public/Blueprint/UserWidget.h
+++ b/Engine/Source/Runtime/UMG/Public/Blueprint/UserWidget.h
@@ -146,83 +146,84 @@ public:
 	class APlayerController* GetOwningPlayer() const;
 
 	/** Called when the widget is constructed */
-	UFUNCTION(BlueprintImplementableEvent, Category="User Interface")
+	UFUNCTION(BlueprintNativeEvent, Category="User Interface")
 	void Construct();
 
-	UFUNCTION(BlueprintImplementableEvent, Category="User Interface")
+	UFUNCTION(BlueprintNativeEvent, Category="User Interface")
 	void Tick(FGeometry MyGeometry, float InDeltaTime);
 
 	//TODO UMG HitTest
 
-	UFUNCTION(BlueprintImplementableEvent, Category="User Interface | Painting")
+	UFUNCTION(BlueprintNativeEvent, Category="User Interface | Painting")
 	void OnPaint(UPARAM(ref) FPaintContext& Context) const;
 
-	UFUNCTION(BlueprintImplementableEvent, Category="Keyboard")
+	UFUNCTION(BlueprintNativeEvent, Category="Keyboard")
 	FEventReply OnKeyboardFocusReceived(FGeometry MyGeometry, FKeyboardFocusEvent InKeyboardFocusEvent);
-	UFUNCTION(BlueprintImplementableEvent, Category="Keyboard")
+	UFUNCTION(BlueprintNativeEvent, Category="Keyboard")
 	void OnKeyboardFocusLost(FKeyboardFocusEvent InKeyboardFocusEvent);
-	//UFUNCTION(BlueprintImplementableEvent, Category="User Interface")
+	//UFUNCTION(BlueprintNativeEvent, Category="User Interface")
 	//void OnKeyboardFocusChanging(FWeakWidgetPath PreviousFocusPath, FWidgetPath NewWidgetPath);
 
-	UFUNCTION(BlueprintImplementableEvent, Category="Keyboard")
+	UFUNCTION(BlueprintNativeEvent, Category="Keyboard")
 	FEventReply OnKeyChar(FGeometry MyGeometry, FCharacterEvent InCharacterEvent);
-	UFUNCTION(BlueprintImplementableEvent, Category="Keyboard")
+	UFUNCTION(BlueprintNativeEvent, Category="Keyboard")
 	FEventReply OnPreviewKeyDown(FGeometry MyGeometry, FKeyboardEvent InKeyboardEvent);
-	UFUNCTION(BlueprintImplementableEvent, Category="Keyboard")
+	UFUNCTION(BlueprintNativeEvent, Category="Keyboard")
 	FEventReply OnKeyDown(FGeometry MyGeometry, FKeyboardEvent InKeyboardEvent);
-	UFUNCTION(BlueprintImplementableEvent, Category="Keyboard")
+	UFUNCTION(BlueprintNativeEvent, Category="Keyboard")
 	FEventReply OnKeyUp(FGeometry MyGeometry, FKeyboardEvent InKeyboardEvent);
 
-	UFUNCTION(BlueprintImplementableEvent, Category="Mouse")
+	UFUNCTION(BlueprintNativeEvent, Category="Mouse")
 	FEventReply OnMouseButtonDown(FGeometry MyGeometry, const FPointerEvent& MouseEvent);
-	UFUNCTION(BlueprintImplementableEvent, Category="Mouse")
+	UFUNCTION(BlueprintNativeEvent, Category="Mouse")
 	FEventReply OnPreviewMouseButtonDown(FGeometry MyGeometry, const FPointerEvent& MouseEvent);
-	UFUNCTION(BlueprintImplementableEvent, Category="Mouse")
+	UFUNCTION(BlueprintNativeEvent, Category="Mouse")
 	FEventReply OnMouseButtonUp(FGeometry MyGeometry, const FPointerEvent& MouseEvent);
-	UFUNCTION(BlueprintImplementableEvent, Category="Mouse")
+	UFUNCTION(BlueprintNativeEvent, Category="Mouse")
 	FEventReply OnMouseMove(FGeometry MyGeometry, const FPointerEvent& MouseEvent);
-	UFUNCTION(BlueprintImplementableEvent, Category="Mouse")
+	UFUNCTION(BlueprintNativeEvent, Category="Mouse")
 	void OnMouseEnter(FGeometry MyGeometry, const FPointerEvent& MouseEvent);
-	UFUNCTION(BlueprintImplementableEvent, Category="Mouse")
+	UFUNCTION(BlueprintNativeEvent, Category="Mouse")
 	void OnMouseLeave(const FPointerEvent& MouseEvent);
-	UFUNCTION(BlueprintImplementableEvent, Category="Mouse")
+	UFUNCTION(BlueprintNativeEvent, Category="Mouse")
 	FEventReply OnMouseWheel(FGeometry MyGeometry, const FPointerEvent& MouseEvent);
-	//UFUNCTION(BlueprintImplementableEvent, Category="Mouse")
-	//FCursorReply OnCursorQuery(FGeometry MyGeometry, const FPointerEvent& CursorEvent) const;
-	UFUNCTION(BlueprintImplementableEvent, Category="Mouse")
+	UFUNCTION(BlueprintNativeEvent, Category="Mouse")
 	FEventReply OnMouseButtonDoubleClick(FGeometry InMyGeometry, const FPointerEvent& InMouseEvent);
 
-	UFUNCTION(BlueprintImplementableEvent, Category="Drag and Drop")
+	//UFUNCTION(BlueprintNativeEvent, Category="Mouse")
+	//FCursorReply OnCursorQuery(FGeometry MyGeometry, const FPointerEvent& CursorEvent) const;
+
+	UFUNCTION(BlueprintNativeEvent, Category="Drag and Drop")
 	void OnDragDetected(FGeometry MyGeometry, const FPointerEvent& PointerEvent, UDragDropOperation*& Operation);
 
-	UFUNCTION(BlueprintImplementableEvent, Category="Drag and Drop")
+	UFUNCTION(BlueprintNativeEvent, Category="Drag and Drop")
 	void OnDragCancelled(const FPointerEvent& PointerEvent, UDragDropOperation* Operation);
 	
-	UFUNCTION(BlueprintImplementableEvent, Category="Drag and Drop")
+	UFUNCTION(BlueprintNativeEvent, Category="Drag and Drop")
 	void OnDragEnter(FGeometry MyGeometry, FPointerEvent PointerEvent, UDragDropOperation* Operation);
-	UFUNCTION(BlueprintImplementableEvent, Category="Drag and Drop")
+	UFUNCTION(BlueprintNativeEvent, Category="Drag and Drop")
 	void OnDragLeave(FPointerEvent PointerEvent, UDragDropOperation* Operation);
-	UFUNCTION(BlueprintImplementableEvent, Category="Drag and Drop")
+	UFUNCTION(BlueprintNativeEvent, Category="Drag and Drop")
 	bool OnDragOver(FGeometry MyGeometry, FPointerEvent PointerEvent, UDragDropOperation* Operation);
-	UFUNCTION(BlueprintImplementableEvent, Category="Drag and Drop")
+	UFUNCTION(BlueprintNativeEvent, Category="Drag and Drop")
 	bool OnDrop(FGeometry MyGeometry, FPointerEvent PointerEvent, UDragDropOperation* Operation);
 
-	UFUNCTION(BlueprintImplementableEvent, Category="Gamepad Input")
+	UFUNCTION(BlueprintNativeEvent, Category="Gamepad Input")
 	FEventReply OnControllerButtonPressed(FGeometry MyGeometry, FControllerEvent ControllerEvent);
-	UFUNCTION(BlueprintImplementableEvent, Category="Gamepad Input")
+	UFUNCTION(BlueprintNativeEvent, Category="Gamepad Input")
 	FEventReply OnControllerButtonReleased(FGeometry MyGeometry, FControllerEvent ControllerEvent);
-	UFUNCTION(BlueprintImplementableEvent, Category="Gamepad Input")
+	UFUNCTION(BlueprintNativeEvent, Category="Gamepad Input")
 	FEventReply OnControllerAnalogValueChanged(FGeometry MyGeometry, FControllerEvent ControllerEvent);
 
-	UFUNCTION(BlueprintImplementableEvent, Category="Touch Input")
+	UFUNCTION(BlueprintNativeEvent, Category="Touch Input")
 	FEventReply OnTouchGesture(FGeometry MyGeometry, const FPointerEvent& GestureEvent);
-	UFUNCTION(BlueprintImplementableEvent, Category="Touch Input")
+	UFUNCTION(BlueprintNativeEvent, Category="Touch Input")
 	FEventReply OnTouchStarted(FGeometry MyGeometry, const FPointerEvent& InTouchEvent);
-	UFUNCTION(BlueprintImplementableEvent, Category="Touch Input")
+	UFUNCTION(BlueprintNativeEvent, Category="Touch Input")
 	FEventReply OnTouchMoved(FGeometry MyGeometry, const FPointerEvent& InTouchEvent);
-	UFUNCTION(BlueprintImplementableEvent, Category="Touch Input")
+	UFUNCTION(BlueprintNativeEvent, Category="Touch Input")
 	FEventReply OnTouchEnded(FGeometry MyGeometry, const FPointerEvent& InTouchEvent);
-	UFUNCTION(BlueprintImplementableEvent, Category="Touch Input")
+	UFUNCTION(BlueprintNativeEvent, Category="Touch Input")
 	FEventReply OnMotionDetected(FGeometry MyGeometry, FMotionEvent InMotionEvent);
 
 	//virtual bool OnVisualizeTooltip(const TSharedPtr<SWidget>& TooltipContent);
