- Added UMG SpinBox widget
- Added SpinBoxStyle uobject wrapper

diff --git a/Engine/Source/Runtime/Slate/Private/Framework/Styling/SpinBoxWidgetStyle.cpp b/Engine/Source/Runtime/Slate/Private/Framework/Styling/SpinBoxWidgetStyle.cpp
new file mode 100644
--- /dev/null
+++ b/Engine/Source/Runtime/Slate/Private/Framework/Styling/SpinBoxWidgetStyle.cpp
@@ -0,0 +1,10 @@
+// Copyright 1998-2014 Epic Games, Inc. All Rights Reserved.
+
+#include "SlatePrivatePCH.h"
+
+
+USpinBoxWidgetStyle::USpinBoxWidgetStyle(const class FPostConstructInitializeProperties& PCIP)
+	: Super(PCIP)
+{
+
+}
\ No newline at end of file
diff --git a/Engine/Source/Runtime/Slate/Private/SlatePrivatePCH.h b/Engine/Source/Runtime/Slate/Private/SlatePrivatePCH.h
--- a/Engine/Source/Runtime/Slate/Private/SlatePrivatePCH.h
+++ b/Engine/Source/Runtime/Slate/Private/SlatePrivatePCH.h
@@ -28,3 +28,4 @@
 #include "ProgressWidgetStyle.h"
 #include "ScrollBoxWidgetStyle.h"
 #include "ScrollBarWidgetStyle.h"
+#include "SpinBoxWidgetStyle.h"
diff --git a/Engine/Source/Runtime/Slate/Public/Framework/Styling/SpinBoxWidgetStyle.h b/Engine/Source/Runtime/Slate/Public/Framework/Styling/SpinBoxWidgetStyle.h
new file mode 100644
--- /dev/null
+++ b/Engine/Source/Runtime/Slate/Public/Framework/Styling/SpinBoxWidgetStyle.h
@@ -0,0 +1,24 @@
+// Copyright 1998-2014 Epic Games, Inc. All Rights Reserved.
+
+#pragma once
+
+#include "SlateWidgetStyleContainerBase.h"
+#include "SpinBoxWidgetStyle.generated.h"
+
+/**
+*/
+UCLASS(hidecategories = Object, MinimalAPI)
+class USpinBoxWidgetStyle : public USlateWidgetStyleContainerBase
+{
+	GENERATED_UCLASS_BODY()
+
+public:
+	/** The actual data describing the button's appearance. */
+	UPROPERTY(Category = Appearance, EditAnywhere, meta = (ShowOnlyInnerProperties))
+	FSpinBoxStyle SpinBoxStyle;
+
+	virtual const struct FSlateWidgetStyle* const GetStyle() const override
+	{
+		return static_cast< const struct FSlateWidgetStyle* >(&SpinBoxStyle);
+	}
+};
diff --git a/Engine/Source/Runtime/Slate/Public/Widgets/Input/SSpinBox.h b/Engine/Source/Runtime/Slate/Public/Widgets/Input/SSpinBox.h
--- a/Engine/Source/Runtime/Slate/Public/Widgets/Input/SSpinBox.h
+++ b/Engine/Source/Runtime/Slate/Public/Widgets/Input/SSpinBox.h
@@ -101,7 +101,7 @@ public:
 		MinSliderValue = (InArgs._MinSliderValue.Get().IsSet()) ? InArgs._MinSliderValue : MinValue;
 		MaxSliderValue = (InArgs._MaxSliderValue.Get().IsSet()) ? InArgs._MaxSliderValue : MaxValue;
 
-		bUnlimitedSpinRange = !((InArgs._MinValue.Get().IsSet() && InArgs._MaxValue.Get().IsSet()) || (InArgs._MinSliderValue.Get().IsSet() && InArgs._MaxSliderValue.Get().IsSet()));
+		UpdateIsSpinRangeUnlimited();
 	
 		SliderExponent = InArgs._SliderExponent;
 
@@ -432,6 +432,58 @@ public:
 		return SCompoundWidget::HasKeyboardFocus() || (EditableText.IsValid() && EditableText->HasKeyboardFocus());
 	}
 
+	/** See the Value attribute */
+	float GetValue() const { return ValueAttribute.Get(); }
+	void SetValue(const TAttribute<NumericType>& InValueAttribute) 
+	{
+		ValueAttribute = InValueAttribute; 
+		CommitValue(InValueAttribute.Get(), ECommitMethod::CommittedViaSpin, ETextCommit::Default);
+	}
+
+	/** See the MinValue attribute */
+	NumericType GetMinValue() const { return MinValue.Get().Get(TNumericLimits<NumericType>::Lowest()); }
+	void SetMinValue(const TAttribute<TOptional<NumericType>>& InMinValue) 
+	{ 
+		MinValue = InMinValue;
+		UpdateIsSpinRangeUnlimited();
+	}
+
+	/** See the MaxValue attribute */
+	NumericType GetMaxValue() const { return MaxValue.Get().Get(TNumericLimits<NumericType>::Max()); }
+	void SetMaxValue(const TAttribute<TOptional<NumericType>>& InMaxValue) 
+	{ 
+		MaxValue = InMaxValue; 
+		UpdateIsSpinRangeUnlimited();
+	}
+
+	/** See the MinSliderValue attribute */
+	NumericType GetMinSliderValue() const { return MinSliderValue.Get().Get(TNumericLimits<NumericType>::Lowest()); }
+	void SetMinSliderValue(const TAttribute<TOptional<NumericType>>& InMinSliderValue) 
+	{ 
+		MinSliderValue = (InMinSliderValue.Get().IsSet()) ? InMinSliderValue : MinValue;
+		UpdateIsSpinRangeUnlimited();
+	}
+
+	/** See the MaxSliderValue attribute */
+	NumericType GetMaxSliderValue() const { return MaxSliderValue.Get().Get(TNumericLimits<NumericType>::Max()); }
+	void SetMaxSliderValue(const TAttribute<TOptional<NumericType>>& InMaxSliderValue) 
+	{ 
+		MaxSliderValue = (InMaxSliderValue.Get().IsSet()) ? InMaxSliderValue : MaxValue;;
+		UpdateIsSpinRangeUnlimited();
+	}
+
+	/** See the Delta attribute */
+	NumericType GetDelta() const { return Delta; }
+	void SetDelta(NumericType InDelta) { Delta = InDelta; }
+
+	/** See the SliderExponent attribute */
+	float GetSliderExponent() const { return SliderExponent.Get(); }
+	void SetSliderExponent(const TAttribute<float>& InSliderExponent) { SliderExponent = InSliderExponent; }
+	
+	/** See the MinDesiredWidth attribute */
+	float GetMinDesiredWidth() const { return SliderExponent.Get(); }
+	void SetMinDesiredWidth(const TAttribute<float>& InMinDesiredWidth) { MinDesiredWidth = InMinDesiredWidth; }
+
 protected:
 	/** Make the spinbox switch to keyboard-based input mode. */
 	void EnterTextMode()
@@ -644,6 +696,10 @@ private:
 
 	/** True when no range is specified, spinner can be spun indefinitely */
 	bool bUnlimitedSpinRange;
+	void UpdateIsSpinRangeUnlimited()
+	{
+		bUnlimitedSpinRange = !((MinValue.Get().IsSet() && MaxValue.Get().IsSet()) || (MinSliderValue.Get().IsSet() && MaxSliderValue.Get().IsSet()));
+	}
 
 	const FSlateBrush* BackgroundHoveredBrush;
 	const FSlateBrush* BackgroundBrush;
@@ -652,23 +708,20 @@ private:
 
 	float DistanceDragged;
 	NumericType Delta;
+	TAttribute<float> SliderExponent;
 	TAttribute< TOptional<NumericType> > MinValue;
 	TAttribute< TOptional<NumericType> > MaxValue;
 	TAttribute< TOptional<NumericType> > MinSliderValue;
 	TAttribute< TOptional<NumericType> > MaxSliderValue;
 
-	NumericType GetMinValue() const { return MinValue.Get().Get(TNumericLimits<NumericType>::Lowest()); }
-	NumericType GetMaxValue() const { return MaxValue.Get().Get(TNumericLimits<NumericType>::Max()); }
-	NumericType GetMinSliderValue() const { return MinSliderValue.Get().Get(TNumericLimits<NumericType>::Lowest()); }
-	NumericType GetMaxSliderValue() const { return MaxSliderValue.Get().Get(TNumericLimits<NumericType>::Max()); }
-
+	/** Prevents the spinbox from being smaller than desired in certain cases (e.g. when it is empty) */
+	TAttribute<float> MinDesiredWidth;
 	float GetTextMinDesiredWidth() const
 	{
 		return FMath::Max(0.0f, MinDesiredWidth.Get() - ArrowImageWidth);
 	}
 
-	TAttribute<float> SliderExponent;
-
+	/** Whether the user is dragging the slider */
 	bool bDragging;
 	
 	/** Cached mouse position to restore after scrolling. */
@@ -686,9 +739,6 @@ private:
 		external forces on the spinbox and syncing the internal value to them. Synced when a value is committed to the spinbox. */
 	NumericType CachedExternalValue;
 
-	/** Prevents the spinbox from being smaller than desired in certain cases (e.g. when it is empty) */
-	TAttribute<float> MinDesiredWidth;
-
 	/** The width of the arrow image, used for calculating the min desired width of the internal text controls. */
 	float ArrowImageWidth;
 };
diff --git a/Engine/Source/Runtime/UMG/Private/Components/SpinBox.cpp b/Engine/Source/Runtime/UMG/Private/Components/SpinBox.cpp
new file mode 100644
--- /dev/null
+++ b/Engine/Source/Runtime/UMG/Private/Components/SpinBox.cpp
@@ -0,0 +1,281 @@
+// Copyright 1998-2014 Epic Games, Inc. All Rights Reserved.
+
+#include "UMGPrivatePCH.h"
+
+#define LOCTEXT_NAMESPACE "UMG"
+
+/////////////////////////////////////////////////////
+// USpinBox
+
+USpinBox::USpinBox(const FPostConstructInitializeProperties& PCIP)
+	: Super(PCIP)
+{
+	// HACK Special font initialization hack since there are no font assets yet for slate.
+	Font = FSlateFontInfo(TEXT("Slate/Fonts/Roboto-Bold.ttf"), 12);
+
+	// Grab other defaults from slate arguments.
+	SSpinBox<float>::FArguments Defaults;
+
+	Value = Defaults._Value.Get();
+	MinValue = Defaults._MinValue.Get().GetValue();
+	MaxValue = Defaults._MaxValue.Get().GetValue();
+	MinSliderValue = Defaults._MinSliderValue.Get().GetValue();
+	MaxSliderValue = Defaults._MaxSliderValue.Get().GetValue();
+	Delta = Defaults._Delta.Get();
+	SliderExponent = Defaults._SliderExponent.Get();
+	MinDesiredWidth = Defaults._MinDesiredWidth.Get();
+	ClearKeyboardFocusOnCommit = Defaults._ClearKeyboardFocusOnCommit.Get();
+	SelectAllTextOnCommit = Defaults._SelectAllTextOnCommit.Get();
+}
+
+void USpinBox::ReleaseNativeWidget()
+{
+	Super::ReleaseNativeWidget();
+
+	MySpinBox.Reset();
+}
+
+TSharedRef<SWidget> USpinBox::RebuildWidget()
+{
+	FString FontPath = FPaths::GameContentDir() / Font.FontName.ToString();
+
+	if ( !FPaths::FileExists(FontPath) )
+	{
+		FontPath = FPaths::EngineContentDir() / Font.FontName.ToString();
+	}
+	
+	SSpinBox<float>::FArguments Defaults;
+	
+	const FSpinBoxStyle* StylePtr = (Style != nullptr) ? Style->GetStyle<FSpinBoxStyle>() : nullptr;
+	if ( StylePtr == nullptr )
+	{
+		StylePtr = Defaults._Style;
+	}
+	
+	MySpinBox = SNew(SSpinBox<float>)
+	.Style(StylePtr)
+	.Font(FSlateFontInfo(FontPath, Font.Size))
+	.ClearKeyboardFocusOnCommit(ClearKeyboardFocusOnCommit)
+	.SelectAllTextOnCommit(SelectAllTextOnCommit)
+	.OnValueChanged(BIND_UOBJECT_DELEGATE(FOnFloatValueChanged, HandleOnValueChanged))
+	.OnValueCommitted(BIND_UOBJECT_DELEGATE(FOnFloatValueCommitted, HandleOnValueCommitted))
+	.OnBeginSliderMovement(BIND_UOBJECT_DELEGATE(FSimpleDelegate, HandleOnBeginSliderMovement))
+	.OnEndSliderMovement(BIND_UOBJECT_DELEGATE(FOnFloatValueChanged, HandleOnEndSliderMovement))
+	;
+	
+	return BuildDesignTimeWidget( MySpinBox.ToSharedRef() );
+}
+
+void USpinBox::SynchronizeProperties()
+{
+	Super::SynchronizeProperties();
+
+	TAttribute<float> ValueBinding = OPTIONAL_BINDING(float, Value);
+	MySpinBox->SetValue(ValueBinding);
+	
+	MySpinBox->SetDelta(Delta);
+	MySpinBox->SetSliderExponent(SliderExponent);
+	MySpinBox->SetMinDesiredWidth(MinDesiredWidth);
+
+	// Set optional values
+	bOverride_MinValue ? SetMinValue(MinValue) : ClearMinValue();
+	bOverride_MaxValue ? SetMaxValue(MaxValue) : ClearMaxValue();
+	bOverride_MinSliderValue ? SetMinSliderValue(MinSliderValue) : ClearMinSliderValue();
+	bOverride_MaxSliderValue ? SetMaxSliderValue(MaxSliderValue) : ClearMaxSliderValue();
+}
+
+float USpinBox::GetValue()
+{
+	if (MySpinBox.IsValid())
+	{
+		return MySpinBox->GetValue();
+	}
+
+	return Value;
+}
+
+void USpinBox::SetValue(float InValue)
+{
+	Value = InValue;
+	if (MySpinBox.IsValid())
+	{
+		MySpinBox->SetValue(InValue);
+	}
+}
+
+// MIN VALUE
+float USpinBox::GetMinValue()
+{
+	float ReturnVal = TNumericLimits<float>::Lowest();
+
+	if (MySpinBox.IsValid())
+	{
+		ReturnVal = MySpinBox->GetMinValue();
+	}
+	else if (bOverride_MinValue)
+	{
+		ReturnVal = MinValue;
+	}
+
+	return ReturnVal;
+}
+
+void USpinBox::SetMinValue(float InMinValue)
+{
+	bOverride_MinValue = true;
+	MinValue = InMinValue;
+	if (MySpinBox.IsValid())
+	{
+		MySpinBox->SetMinValue(InMinValue);
+	}
+}
+
+void USpinBox::ClearMinValue()
+{
+	bOverride_MinValue = false;
+	if (MySpinBox.IsValid())
+	{
+		MySpinBox->SetMinValue(TOptional<float>());
+	}
+}
+
+// MAX VALUE
+float USpinBox::GetMaxValue()
+{
+	float ReturnVal = TNumericLimits<float>::Max();
+
+	if (MySpinBox.IsValid())
+	{
+		ReturnVal = MySpinBox->GetMaxValue();
+	}
+	else if (bOverride_MaxValue)
+	{
+		ReturnVal = MaxValue;
+	}
+
+	return ReturnVal;
+}
+
+void USpinBox::SetMaxValue(float InMaxValue)
+{
+	bOverride_MaxValue = true;
+	MaxValue = InMaxValue;
+	if (MySpinBox.IsValid())
+	{
+		MySpinBox->SetMaxValue(InMaxValue);
+	}
+}
+void USpinBox::ClearMaxValue()
+{
+	bOverride_MaxValue = false;
+	if (MySpinBox.IsValid())
+	{
+		MySpinBox->SetMaxValue(TOptional<float>());
+	}
+}
+
+// MIN SLIDER VALUE
+float USpinBox::GetMinSliderValue()
+{
+	float ReturnVal = TNumericLimits<float>::Min();
+
+	if (MySpinBox.IsValid())
+	{
+		ReturnVal = MySpinBox->GetMinSliderValue();
+	}
+	else if (bOverride_MinSliderValue)
+	{
+		ReturnVal = MinSliderValue;
+	}
+
+	return ReturnVal;
+}
+
+void USpinBox::SetMinSliderValue(float InMinSliderValue)
+{
+	bOverride_MinSliderValue = true;
+	MinSliderValue = InMinSliderValue;
+	if (MySpinBox.IsValid())
+	{
+		MySpinBox->SetMinSliderValue(InMinSliderValue);
+	}
+}
+
+void USpinBox::ClearMinSliderValue()
+{
+	bOverride_MinSliderValue = false;
+	if (MySpinBox.IsValid())
+	{
+		MySpinBox->SetMinSliderValue(TOptional<float>());
+	}
+}
+
+// MAX SLIDER VALUE
+float USpinBox::GetMaxSliderValue()
+{
+	float ReturnVal = TNumericLimits<float>::Max();
+
+	if (MySpinBox.IsValid())
+	{
+		ReturnVal = MySpinBox->GetMaxSliderValue();
+	}
+	else if (bOverride_MaxSliderValue)
+	{
+		ReturnVal = MaxSliderValue;
+	}
+
+	return ReturnVal;
+}
+
+void USpinBox::SetMaxSliderValue(float InMaxSliderValue)
+{
+	bOverride_MaxSliderValue = true;
+	MaxSliderValue = InMaxSliderValue;
+	if (MySpinBox.IsValid())
+	{
+		MySpinBox->SetMaxSliderValue(InMaxSliderValue);
+	}
+}
+
+void USpinBox::ClearMaxSliderValue()
+{
+	bOverride_MaxSliderValue = false;
+	if (MySpinBox.IsValid())
+	{
+		MySpinBox->SetMaxSliderValue(TOptional<float>());
+	}
+}
+
+// Event handlers
+void USpinBox::HandleOnValueChanged(float InValue)
+{
+	OnValueChanged.Broadcast(InValue);
+}
+
+void USpinBox::HandleOnValueCommitted(float InValue, ETextCommit::Type CommitMethod)
+{
+	OnValueCommitted.Broadcast(InValue, CommitMethod);
+}
+
+void USpinBox::HandleOnBeginSliderMovement()
+{
+	OnBeginSliderMovement.Broadcast();
+}
+
+void USpinBox::HandleOnEndSliderMovement(float InValue)
+{
+	OnEndSliderMovement.Broadcast(InValue);
+}
+
+#if WITH_EDITOR
+
+const FSlateBrush* USpinBox::GetEditorIcon()
+{
+	return FUMGStyle::Get().GetBrush("Widget.SpinBox");
+}
+
+#endif
+
+/////////////////////////////////////////////////////
+
+#undef LOCTEXT_NAMESPACE
\ No newline at end of file
diff --git a/Engine/Source/Runtime/UMG/Public/Components/SpinBox.h b/Engine/Source/Runtime/UMG/Public/Components/SpinBox.h
new file mode 100644
--- /dev/null
+++ b/Engine/Source/Runtime/UMG/Public/Components/SpinBox.h
@@ -0,0 +1,189 @@
+// Copyright 1998-2014 Epic Games, Inc. All Rights Reserved.
+
+#pragma once
+
+#include "SpinBox.generated.h"
+
+/** Spin box widget */
+UCLASS(meta = (Category = "Primitive"), ClassGroup = UserInterface)
+class UMG_API USpinBox : public UWidget
+{
+	GENERATED_UCLASS_BODY()
+
+public:
+	DECLARE_DYNAMIC_MULTICAST_DELEGATE_OneParam(FOnSpinBoxValueChangedEvent, float, InValue);
+	DECLARE_DYNAMIC_MULTICAST_DELEGATE_TwoParams(FOnSpinBoxValueCommittedEvent, float, InValue, ETextCommit::Type, CommitMethod);
+	DECLARE_DYNAMIC_MULTICAST_DELEGATE(FOnSpinBoxBeginSliderMovement);
+
+public:
+	/** Text style */
+	UPROPERTY(EditDefaultsOnly, Category = Style, meta = (DisplayThumbnail = "true"))
+	USlateWidgetStyleAsset* Style;
+
+	/** Value stored in this spin box */
+	UPROPERTY(EditDefaultsOnly, Category = Content)
+	float Value;
+
+	/** The amount by which to change the spin box value as the slider moves. */
+	UPROPERTY(EditDefaultsOnly, Category = Content)
+	float Delta;
+
+	/** The exponent by which to increase the delta as the mouse moves. 1 is constant (never increases the delta). */
+	UPROPERTY(EditDefaultsOnly, Category = Content)
+	float SliderExponent;
+	
+	/** Font color and opacity (overrides style) */
+	UPROPERTY(EditDefaultsOnly, Category = Content)
+	FSlateFontInfo Font;
+
+	/** The minimum width of the spin box */
+	UPROPERTY(EditDefaultsOnly, Category = Content, AdvancedDisplay, DisplayName = "Minimum Desired Width")
+	float MinDesiredWidth;
+
+	/** Whether to remove the keyboard focus from the spin box when the value is committed */
+	UPROPERTY(EditDefaultsOnly, Category = Content, AdvancedDisplay)
+	bool ClearKeyboardFocusOnCommit;
+
+	/** Whether to select the text in the spin box when the value is committed */
+	UPROPERTY(EditDefaultsOnly, Category = Content, AdvancedDisplay)
+	bool SelectAllTextOnCommit;
+
+public:
+	/** Called when the value is changed interactively by the user */
+	UPROPERTY(BlueprintAssignable, Category = "Widget Event")
+	FOnSpinBoxValueChangedEvent OnValueChanged;
+
+	/** Called when the value is committed. Occurs when the user presses Enter or the text box loses focus. */
+	UPROPERTY(BlueprintAssignable, Category = "Widget Event")
+	FOnSpinBoxValueCommittedEvent OnValueCommitted;
+
+	/** Called right before the slider begins to move */
+	UPROPERTY(BlueprintAssignable, Category = "Widget Event")
+	FOnSpinBoxBeginSliderMovement OnBeginSliderMovement;
+
+	/** Called right after the slider handle is released by the user */
+	UPROPERTY(BlueprintAssignable, Category = "Widget Event")
+	FOnSpinBoxValueChangedEvent OnEndSliderMovement;
+
+	/** Get the current value of the spin box. */
+	UFUNCTION(BlueprintPure, Category = "Behavior")
+	float GetValue();
+
+	/** Set the value of the spin box. */
+	UFUNCTION(BlueprintCallable, Category = "Behavior")
+	void SetValue(float NewValue);
+
+	// MIN VALUE
+	/** Get the current minimum value that can be manually set in the spin box. */
+	UFUNCTION(BlueprintPure, Category = "Behavior")
+	float GetMinValue();
+
+	/** Set the minimum value that can be manually set in the spin box. */
+	UFUNCTION(BlueprintCallable, Category = "Behavior")
+	void SetMinValue(float NewValue);
+
+	/** Clear the minimum value that can be manually set in the spin box. */
+	UFUNCTION(BlueprintCallable, Category = "Behavior")
+	void ClearMinValue();
+
+	// MAX VALUE
+	/** Get the current maximum value that can be manually set in the spin box. */
+	UFUNCTION(BlueprintPure, Category = "Behavior")
+	float GetMaxValue();
+
+	/** Set the maximum value that can be manually set in the spin box. */
+	UFUNCTION(BlueprintCallable, Category = "Behavior")
+	void SetMaxValue(float NewValue);
+
+	/** Clear the maximum value that can be manually set in the spin box. */
+	UFUNCTION(BlueprintCallable, Category = "Behavior")
+	void ClearMaxValue();
+
+	// MIN SLIDER VALUE
+	/** Get the current minimum value that can be specified using the slider. */
+	UFUNCTION(BlueprintPure, Category = "Behavior")
+	float GetMinSliderValue();
+
+	/** Set the minimum value that can be specified using the slider. */
+	UFUNCTION(BlueprintCallable, Category = "Behavior")
+	void SetMinSliderValue(float NewValue);
+
+	/** Clear the minimum value that can be specified using the slider. */
+	UFUNCTION(BlueprintCallable, Category = "Behavior")
+	void ClearMinSliderValue();
+
+	// MAX SLIDER VALUE
+	/** Get the current maximum value that can be specified using the slider. */
+	UFUNCTION(BlueprintPure, Category = "Behavior")
+	float GetMaxSliderValue();
+
+	/** Set the maximum value that can be specified using the slider. */
+	UFUNCTION(BlueprintCallable, Category = "Behavior")
+	void SetMaxSliderValue(float NewValue);
+
+	/** Clear the maximum value that can be specified using the slider. */
+	UFUNCTION(BlueprintCallable, Category = "Behavior")
+	void ClearMaxSliderValue();
+
+	// UVisual interface
+	virtual void ReleaseNativeWidget() override;
+	// End of UVisual interface
+
+	// UWidget interface
+	virtual void SynchronizeProperties() override;
+
+#if WITH_EDITOR
+	virtual const FSlateBrush* GetEditorIcon() override;
+#endif
+	// End of UWidget interface
+
+	/** A bindable delegate to allow logic to drive the value of the widget */
+	UPROPERTY()
+	FGetFloat ValueDelegate;
+
+protected:
+	// UWidget interface
+	virtual TSharedRef<SWidget> RebuildWidget() override;
+	// End of UWidget
+
+	void HandleOnValueChanged(float InValue);
+	void HandleOnValueCommitted(float InValue, ETextCommit::Type CommitMethod);
+	void HandleOnBeginSliderMovement();
+	void HandleOnEndSliderMovement(float InValue);
+
+protected:
+	/** Whether the optional MinValue attribute of the widget is set */
+	UPROPERTY()
+	uint32 bOverride_MinValue : 1;
+
+	/** Whether the optional MaxValue attribute of the widget is set */
+	UPROPERTY()
+	uint32 bOverride_MaxValue : 1;
+
+	/** Whether the optional MinSliderValue attribute of the widget is set */
+	UPROPERTY()
+	uint32 bOverride_MinSliderValue : 1;
+
+	/** Whether the optional MaxSliderValue attribute of the widget is set */
+	UPROPERTY()
+	uint32 bOverride_MaxSliderValue : 1;
+
+	/** The minimum allowable value that can be manually entered into the spin box */
+	UPROPERTY(EditDefaultsOnly, Category = Content, DisplayName = "Minimum Value", meta = (editcondition = "bOverride_MinValue"))
+	float MinValue;
+
+	/** The maximum allowable value that can be manually entered into the spin box */
+	UPROPERTY(EditDefaultsOnly, Category = Content, DisplayName = "Maximum Value", meta = (editcondition = "bOverride_MaxValue"))
+	float MaxValue;
+
+	/** The minimum allowable value that can be specified using the slider */
+	UPROPERTY(EditDefaultsOnly, Category = Content, DisplayName = "Minimum Slider Value", meta = (editcondition = "bOverride_MinSliderValue"))
+	float MinSliderValue;
+
+	/** The maximum allowable value that can be specified using the slider */
+	UPROPERTY(EditDefaultsOnly, Category = Content, DisplayName = "Maximum Slider Value", meta = (editcondition = "bOverride_MaxSliderValue"))
+	float MaxSliderValue;
+
+protected:
+	TSharedPtr<SSpinBox<float>> MySpinBox;
+};
\ No newline at end of file
diff --git a/Engine/Source/Runtime/UMG/Public/UMG.h b/Engine/Source/Runtime/UMG/Public/UMG.h
--- a/Engine/Source/Runtime/UMG/Public/UMG.h
+++ b/Engine/Source/Runtime/UMG/Public/UMG.h
@@ -74,6 +74,8 @@
 #include "CheckBox.h"
 #include "Image.h"
 
+#include "SpinBox.h"
+
 #include "EditableText.h"
 #include "EditableTextBox.h"
 
