UMG - Move spacer to the primitive group of widgets.

diff --git a/Engine/Source/Runtime/UMG/Public/Components/Spacer.h b/Engine/Source/Runtime/UMG/Public/Components/Spacer.h
--- a/Engine/Source/Runtime/UMG/Public/Components/Spacer.h
+++ b/Engine/Source/Runtime/UMG/Public/Components/Spacer.h
@@ -5,7 +5,7 @@
 #include "Spacer.generated.h"
 
 /** A spacer widget; it does not have a visual representation, and just provides padding between other widgets. */
-UCLASS(meta=( Category="Misc" ), ClassGroup=UserInterface)
+UCLASS(meta=( Category="Primitive" ), ClassGroup=UserInterface)
 class UMG_API USpacer : public UWidget
 {
 	GENERATED_UCLASS_BODY()
