UMG - No longer using a specialized move handle on canvases.  Fixing movement on canvases that have non-zero anchor points.

diff --git a/Engine/Source/Editor/UMGEditor/Private/Designer/SDesignerView.cpp b/Engine/Source/Editor/UMGEditor/Private/Designer/SDesignerView.cpp
--- a/Engine/Source/Editor/UMGEditor/Private/Designer/SDesignerView.cpp
+++ b/Engine/Source/Editor/UMGEditor/Private/Designer/SDesignerView.cpp
@@ -256,24 +256,17 @@ void SDesignerView::Construct(const FArguments& InArgs, TSharedPtr<FWidgetBluepr
 			[
 				SNew(SDisappearingBar)
 				[
-					SNew(SHorizontalBox)
-
-					+ SHorizontalBox::Slot()
-					.FillWidth(1.0f)
-					.Padding(0)
+					SNew(SBorder)
+					.BorderImage(FEditorStyle::GetBrush("ToolPanel.GroupBorder"))
+					.BorderBackgroundColor(FLinearColor(1, 1, 1, 0.75))
+					.HAlign(HAlign_Center)
+					.VAlign(VAlign_Center)
+					.Padding(FMargin(0, 5))
+					.Visibility(this, &SDesignerView::GetInfoBarVisibility)
 					[
-						SNew(SBorder)
-						.BorderImage(FEditorStyle::GetBrush("ToolPanel.GroupBorder"))
-						.BorderBackgroundColor(FLinearColor(1, 1, 1, 0.75))
-						.HAlign(HAlign_Center)
-						.VAlign(VAlign_Center)
-						.Padding(FMargin(0, 5))
-						.Visibility(this, &SDesignerView::GetInfoBarVisibility)
-						[
-							SNew(STextBlock)
-							.TextStyle(FEditorStyle::Get(), "Graph.ZoomText")
-							.Text(this, &SDesignerView::GetInfoBarText)
-						]
+						SNew(STextBlock)
+						.TextStyle(FEditorStyle::Get(), "Graph.ZoomText")
+						.Text(this, &SDesignerView::GetInfoBarText)
 					]
 				]
 			]
@@ -366,7 +359,7 @@ FSlateRect SDesignerView::ComputeAreaBounds() const
 
 EVisibility SDesignerView::GetInfoBarVisibility() const
 {
-	if ( DesignerMessage != EDesignerMessage::None && bMovingExistingWidget )
+	if ( DesignerMessage != EDesignerMessage::None )
 	{
 		return EVisibility::Visible;
 	}
@@ -657,10 +650,11 @@ FReply SDesignerView::OnMouseButtonDown(const FGeometry& MyGeometry, const FPoin
 
 FReply SDesignerView::OnMouseButtonUp(const FGeometry& MyGeometry, const FPointerEvent& MouseEvent)
 {
-	if ( MouseEvent.GetEffectingButton() == EKeys::LeftMouseButton )
+	if ( HasMouseCapture() && MouseEvent.GetEffectingButton() == EKeys::LeftMouseButton )
 	{
 		bMouseDown = false;
 		bMovingExistingWidget = false;
+		DesignerMessage = EDesignerMessage::None;
 
 		if ( CurrentHandle != DH_NONE )
 		{
@@ -729,6 +723,7 @@ FReply SDesignerView::OnMouseMove(const FGeometry& MyGeometry, const FPointerEve
 			if ( !bIsRootWidget )
 			{
 				bMovingExistingWidget = true;
+				DesignerMessage = EDesignerMessage::None;
 				//Drag selected widgets
 				return FReply::Handled().DetectDrag(AsShared(), EKeys::LeftMouseButton);
 			}
@@ -1362,9 +1357,34 @@ UWidget* SDesignerView::ProcessDropAndAddWidget(const FGeometry& MyGeometry, con
 			FVector2D LocalPosition = ArrangedWidget.Geometry.AbsoluteToLocal(DragDropEvent.GetScreenSpacePosition());
 			if ( UPanelSlot* Slot = NewParent->AddChild(Widget) )
 			{
-				FWidgetBlueprintEditorUtils::ImportPropertiesFromText(Slot, SelectedDragDropOp->ExportedSlotProperties);
-
-				Slot->SetDesiredPosition(LocalPosition - SelectedWidgetContextMenuLocation);
+				// HACK UMG: In order to correctly drop items into the canvas that have a non-zero anchor,
+				// we need to know the layout information after slate has performed a prepass.  So we have
+				// to rebase the layout and reinterpret the new position based on anchor point layout data.
+				// This should be pulled out into an extension of some kind so that this can be fixed for
+				// other widgets as well that may need to do work like this.
+				if ( UCanvasPanelSlot* CanvasSlot = Cast<UCanvasPanelSlot>(Slot) )
+				{
+					if ( bIsPreview )
+					{
+						FWidgetBlueprintEditorUtils::ImportPropertiesFromText(Slot, SelectedDragDropOp->ExportedSlotProperties);
+
+						CanvasSlot->SaveBaseLayout();
+						Slot->SetDesiredPosition(LocalPosition - SelectedWidgetContextMenuLocation);
+						CanvasSlot->RebaseLayout();
+
+						FWidgetBlueprintEditorUtils::ExportPropertiesToText(Slot, SelectedDragDropOp->ExportedSlotProperties);
+					}
+					else
+					{
+						FWidgetBlueprintEditorUtils::ImportPropertiesFromText(Slot, SelectedDragDropOp->ExportedSlotProperties);
+						//Slot->SetDesiredPosition(LocalPosition - SelectedWidgetContextMenuLocation);
+					}
+				}
+				else
+				{
+					FWidgetBlueprintEditorUtils::ImportPropertiesFromText(Slot, SelectedDragDropOp->ExportedSlotProperties);
+					Slot->SetDesiredPosition(LocalPosition - SelectedWidgetContextMenuLocation);
+				}
 
 				DropPreviewParent = NewParent;
 
@@ -1401,6 +1421,7 @@ FReply SDesignerView::OnDrop(const FGeometry& MyGeometry, const FDragDropEvent&
 {
 	bMouseDown = false;
 	bMovingExistingWidget = false;
+	DesignerMessage = EDesignerMessage::None;
 
 	UWidgetBlueprint* BP = GetBlueprint();
 	
diff --git a/Engine/Source/Editor/UMGEditor/Private/Extensions/CanvasSlotExtension.cpp b/Engine/Source/Editor/UMGEditor/Private/Extensions/CanvasSlotExtension.cpp
--- a/Engine/Source/Editor/UMGEditor/Private/Extensions/CanvasSlotExtension.cpp
+++ b/Engine/Source/Editor/UMGEditor/Private/Extensions/CanvasSlotExtension.cpp
@@ -38,23 +38,6 @@ void FCanvasSlotExtension::ExtendSelection(const TArray< FWidgetReference >& Sel
 {
 	SelectionCache = Selection;
 
-	MoveHandle =
-		SNew(SBorder)
-		.BorderImage(FEditorStyle::Get().GetBrush("NoBrush"))
-		.OnMouseButtonDown(this, &FCanvasSlotExtension::HandleBeginDrag)
-		.OnMouseButtonUp(this, &FCanvasSlotExtension::HandleEndDrag)
-		.OnMouseMove(this, &FCanvasSlotExtension::HandleDragging)
-		.Padding(FMargin(0))
-		.Cursor(EMouseCursor::GrabHand)
-		[
-			SNew(SImage)
-			.Image(FCoreStyle::Get().GetBrush("SoftwareCursor_CardinalCross"))
-		];
-
-	MoveHandle->SlatePrepass();
-
-	SurfaceElements.Add(MakeShareable(new FDesignerSurfaceElement(MoveHandle.ToSharedRef(), EExtensionLayoutLocation::TopLeft, FVector2D(0, -(MoveHandle->GetDesiredSize().Y + 10)))));
-
 	AnchorWidgets.SetNumZeroed(EAnchorWidget::MAX_COUNT);
 	AnchorWidgets[EAnchorWidget::Center] = MakeAnchorWidget(EAnchorWidget::Center, 16, 16);
 
@@ -374,78 +357,6 @@ FReply FCanvasSlotExtension::HandleAnchorDragging(const FGeometry& Geometry, con
 	return FReply::Unhandled();
 }
 
-FReply FCanvasSlotExtension::HandleBeginDrag(const FGeometry& Geometry, const FPointerEvent& Event)
-{
-	BeginTransaction(LOCTEXT("MoveWidget", "Move Widget"));
-
-	return FReply::Handled().CaptureMouse(MoveHandle.ToSharedRef());
-}
-
-FReply FCanvasSlotExtension::HandleEndDrag(const FGeometry& Geometry, const FPointerEvent& Event)
-{
-	EndTransaction();
-
-	FBlueprintEditorUtils::MarkBlueprintAsModified(Blueprint);
-
-	return FReply::Handled().ReleaseMouseCapture();
-}
-
-FReply FCanvasSlotExtension::HandleDragging(const FGeometry& Geometry, const FPointerEvent& Event)
-{
-	if ( MoveHandle->HasMouseCapture() )
-	{
-		float InverseScale = ( 1.0f / Designer->GetPreviewScale() );
-
-		for ( FWidgetReference& Selection : SelectionCache )
-		{
-			MoveByAmount(Selection, Event.GetCursorDelta() * InverseScale);
-		}
-
-		return FReply::Handled();
-	}
-
-	return FReply::Unhandled();
-}
-
-void FCanvasSlotExtension::MoveByAmount(FWidgetReference& WidgetRef, FVector2D Delta)
-{
-	if ( Delta.IsZero() )
-	{
-		return;
-	}
-
-	UWidget* Widget = WidgetRef.GetPreview();
-
-	UCanvasPanelSlot* CanvasSlot = CastChecked<UCanvasPanelSlot>(Widget->Slot);
-	UCanvasPanel* Parent = CastChecked<UCanvasPanel>(CanvasSlot->Parent);
-
-	FMargin Offsets = CanvasSlot->LayoutData.Offsets;
-	Offsets.Left += Delta.X;
-	Offsets.Top += Delta.Y;
-
-	// If the slot is stretched horizontally we need to move the right side as it no longer represents width, but
-	// now represents margin from the right stretched side.
-	if ( CanvasSlot->LayoutData.Anchors.IsStretchedHorizontal() )
-	{
-		Offsets.Right -= Delta.X;
-	}
-
-	// If the slot is stretched vertically we need to move the bottom side as it no longer represents width, but
-	// now represents margin from the bottom stretched side.
-	if ( CanvasSlot->LayoutData.Anchors.IsStretchedVertical() )
-	{
-		Offsets.Bottom -= Delta.Y;
-	}
-
-	CanvasSlot->SetOffsets(Offsets);
-
-	// Update the Template widget to match
-	UWidget* TemplateWidget = WidgetRef.GetTemplate();
-	UCanvasPanelSlot* TemplateSlot = CastChecked<UCanvasPanelSlot>(TemplateWidget->Slot);
-
-	TemplateSlot->SetOffsets(Offsets);
-}
-
 void FCanvasSlotExtension::PaintCollisionLines(const TSet< FWidgetReference >& Selection, const FGeometry& AllottedGeometry, const FSlateRect& MyClippingRect, FSlateWindowElementList& OutDrawElements, int32 LayerId, const FWidgetStyle& InWidgetStyle, bool bParentEnabled) const
 {
 	for ( const FWidgetReference& WidgetRef : Selection )
@@ -462,7 +373,7 @@ void FCanvasSlotExtension::PaintCollisionLines(const TSet< FWidgetReference >& S
 			if ( UCanvasPanel* Canvas = CastChecked<UCanvasPanel>(CanvasSlot->Parent) )
 			{
 				//TODO UMG Only show guide lines when near them and dragging
-				if ( MoveHandle->HasMouseCapture() )
+				if ( false )
 				{
 					FGeometry MyArrangedGeometry;
 					if ( !Canvas->GetGeometryForSlot(CanvasSlot, MyArrangedGeometry) )
diff --git a/Engine/Source/Editor/UMGEditor/Private/Extensions/CanvasSlotExtension.h b/Engine/Source/Editor/UMGEditor/Private/Extensions/CanvasSlotExtension.h
--- a/Engine/Source/Editor/UMGEditor/Private/Extensions/CanvasSlotExtension.h
+++ b/Engine/Source/Editor/UMGEditor/Private/Extensions/CanvasSlotExtension.h
@@ -45,18 +45,12 @@ private:
 	FReply HandleAnchorEndDrag(const FGeometry& Geometry, const FPointerEvent& Event, EAnchorWidget::Type AnchorType);
 	FReply HandleAnchorDragging(const FGeometry& Geometry, const FPointerEvent& Event, EAnchorWidget::Type AnchorType);
 
-	FReply HandleBeginDrag(const FGeometry& Geometry, const FPointerEvent& Event);
-	FReply HandleEndDrag(const FGeometry& Geometry, const FPointerEvent& Event);
-	FReply HandleDragging(const FGeometry& Geometry, const FPointerEvent& Event);
-
 	TSharedRef<SWidget> MakeAnchorWidget(EAnchorWidget::Type AnchorType, float Width, float Height);
 
 	const FSlateBrush* GetAnchorBrush(EAnchorWidget::Type AnchorType) const;
 	EVisibility GetAnchorVisibility(EAnchorWidget::Type AnchorType) const;
 	FVector2D GetAnchorAlignment(EAnchorWidget::Type AnchorType) const;
 
-	void MoveByAmount(FWidgetReference& WidgetRef, FVector2D Delta);
-
 	static bool GetCollisionSegmentsForSlot(UCanvasPanel* Canvas, int32 SlotIndex, TArray<FVector2D>& Segments);
 	static bool GetCollisionSegmentsForSlot(UCanvasPanel* Canvas, UCanvasPanelSlot* Slot, TArray<FVector2D>& Segments);
 	static void GetCollisionSegmentsFromGeometry(FGeometry ArrangedGeometry, TArray<FVector2D>& Segments);
@@ -65,9 +59,6 @@ private:
 
 private:
 
-	/** */
-	TSharedPtr<SBorder> MoveHandle;
-
 	/** */
 	TArray< TSharedPtr<SWidget> > AnchorWidgets;
 
diff --git a/Engine/Source/Runtime/UMG/Private/Components/CanvasPanelSlot.cpp b/Engine/Source/Runtime/UMG/Private/Components/CanvasPanelSlot.cpp
--- a/Engine/Source/Runtime/UMG/Private/Components/CanvasPanelSlot.cpp
+++ b/Engine/Source/Runtime/UMG/Private/Components/CanvasPanelSlot.cpp
@@ -34,7 +34,27 @@ void UCanvasPanelSlot::BuildSlot(TSharedRef<SConstraintCanvas> Canvas)
 
 void UCanvasPanelSlot::SetDesiredPosition(FVector2D InPosition)
 {
-	SetOffsets(FMargin(InPosition.X, InPosition.Y, LayoutData.Offsets.Right, LayoutData.Offsets.Bottom));
+	FVector2D Delta = FVector2D(InPosition.X - LayoutData.Offsets.Left, InPosition.Y - LayoutData.Offsets.Top);
+
+	FMargin NewOffset = LayoutData.Offsets;
+	NewOffset.Left += Delta.X;
+	NewOffset.Top += Delta.Y;
+
+	// If the slot is stretched horizontally we need to move the right side as it no longer represents width, but
+	// now represents margin from the right stretched side.
+	if ( LayoutData.Anchors.IsStretchedHorizontal() )
+	{
+		NewOffset.Right -= Delta.X;
+	}
+
+	// If the slot is stretched vertically we need to move the bottom side as it no longer represents width, but
+	// now represents margin from the bottom stretched side.
+	if ( LayoutData.Anchors.IsStretchedVertical() )
+	{
+		NewOffset.Bottom -= Delta.Y;
+	}
+
+	SetOffsets(NewOffset);
 }
 
 void UCanvasPanelSlot::SetDesiredSize(FVector2D InSize)
@@ -216,59 +236,81 @@ void UCanvasPanelSlot::RebaseLayout(bool PreserveSize)
 				LayoutData.Anchors.Maximum.Y * CanvasSize.Y);
 			FVector2D DefaultAnchorPosition = FVector2D(AnchorPositions.Left, AnchorPositions.Top);
 
-			// Determine the amount that would be offset from the anchor position if alignment was applied.
-			FVector2D AlignmentOffset = LayoutData.Alignment * PreEditGeometry.Size;
-
 			// Determine where the widget's new position needs to be to maintain a stable location when the anchors change.
 			FVector2D LeftTopDelta = PreEditGeometry.Position - DefaultAnchorPosition;
 
-			// Adjust the size to remain constant
-			if ( !LayoutData.Anchors.IsStretchedHorizontal() && PreEditLayoutData.Anchors.IsStretchedHorizontal() )
-			{
-				// Adjust the position to remain constant
-				LayoutData.Offsets.Left = LeftTopDelta.X + AlignmentOffset.X;
-				LayoutData.Offsets.Right = PreEditGeometry.Size.X;
-			}
-			else if ( !PreserveSize && LayoutData.Anchors.IsStretchedHorizontal() && !PreEditLayoutData.Anchors.IsStretchedHorizontal() )
-			{
-				// Adjust the position to remain constant
-				LayoutData.Offsets.Left = 0;
-				LayoutData.Offsets.Right = 0;
-			}
-			else if ( LayoutData.Anchors.IsStretchedHorizontal() )
-			{
-				// Adjust the position to remain constant
-				LayoutData.Offsets.Left = LeftTopDelta.X;
-				LayoutData.Offsets.Right = AnchorPositions.Right - ( AnchorPositions.Left + LayoutData.Offsets.Left + PreEditGeometry.Size.X );
-			}
-			else
-			{
-				// Adjust the position to remain constant
-				LayoutData.Offsets.Left = LeftTopDelta.X + AlignmentOffset.X;
-			}
-
-			if ( !LayoutData.Anchors.IsStretchedVertical() && PreEditLayoutData.Anchors.IsStretchedVertical() )
-			{
-				// Adjust the position to remain constant
-				LayoutData.Offsets.Top = LeftTopDelta.Y + AlignmentOffset.Y;
-				LayoutData.Offsets.Bottom = PreEditGeometry.Size.Y;
-			}
-			else if ( !PreserveSize && LayoutData.Anchors.IsStretchedVertical() && !PreEditLayoutData.Anchors.IsStretchedVertical() )
-			{
-				// Adjust the position to remain constant
-				LayoutData.Offsets.Top = 0;
-				LayoutData.Offsets.Bottom = 0;
-			}
-			else if ( LayoutData.Anchors.IsStretchedVertical() )
+			if ( PreEditLayoutData.Anchors.Minimum != LayoutData.Anchors.Minimum || PreEditLayoutData.Anchors.Maximum != LayoutData.Anchors.Maximum )
 			{
-				// Adjust the position to remain constant
-				LayoutData.Offsets.Top = LeftTopDelta.Y;
-				LayoutData.Offsets.Bottom = AnchorPositions.Bottom - ( AnchorPositions.Top + LayoutData.Offsets.Top + PreEditGeometry.Size.Y );
+				// Determine the amount that would be offset from the anchor position if alignment was applied.
+				FVector2D AlignmentOffset = LayoutData.Alignment * PreEditGeometry.Size;
+
+				// Adjust the size to remain constant
+				if ( !LayoutData.Anchors.IsStretchedHorizontal() && PreEditLayoutData.Anchors.IsStretchedHorizontal() )
+				{
+					// Adjust the position to remain constant
+					LayoutData.Offsets.Left = LeftTopDelta.X + AlignmentOffset.X;
+					LayoutData.Offsets.Right = PreEditGeometry.Size.X;
+				}
+				else if ( !PreserveSize && LayoutData.Anchors.IsStretchedHorizontal() && !PreEditLayoutData.Anchors.IsStretchedHorizontal() )
+				{
+					// Adjust the position to remain constant
+					LayoutData.Offsets.Left = 0;
+					LayoutData.Offsets.Right = 0;
+				}
+				else if ( LayoutData.Anchors.IsStretchedHorizontal() )
+				{
+					// Adjust the position to remain constant
+					LayoutData.Offsets.Left = LeftTopDelta.X;
+					LayoutData.Offsets.Right = AnchorPositions.Right - ( AnchorPositions.Left + LayoutData.Offsets.Left + PreEditGeometry.Size.X );
+				}
+				else
+				{
+					// Adjust the position to remain constant
+					LayoutData.Offsets.Left = LeftTopDelta.X + AlignmentOffset.X;
+				}
+
+				if ( !LayoutData.Anchors.IsStretchedVertical() && PreEditLayoutData.Anchors.IsStretchedVertical() )
+				{
+					// Adjust the position to remain constant
+					LayoutData.Offsets.Top = LeftTopDelta.Y + AlignmentOffset.Y;
+					LayoutData.Offsets.Bottom = PreEditGeometry.Size.Y;
+				}
+				else if ( !PreserveSize && LayoutData.Anchors.IsStretchedVertical() && !PreEditLayoutData.Anchors.IsStretchedVertical() )
+				{
+					// Adjust the position to remain constant
+					LayoutData.Offsets.Top = 0;
+					LayoutData.Offsets.Bottom = 0;
+				}
+				else if ( LayoutData.Anchors.IsStretchedVertical() )
+				{
+					// Adjust the position to remain constant
+					LayoutData.Offsets.Top = LeftTopDelta.Y;
+					LayoutData.Offsets.Bottom = AnchorPositions.Bottom - ( AnchorPositions.Top + LayoutData.Offsets.Top + PreEditGeometry.Size.Y );
+				}
+				else
+				{
+					// Adjust the position to remain constant
+					LayoutData.Offsets.Top = LeftTopDelta.Y + AlignmentOffset.Y;
+				}
 			}
-			else
+			else if ( PreEditLayoutData.Offsets.Left != LayoutData.Offsets.Left || PreEditLayoutData.Offsets.Top != LayoutData.Offsets.Top )
 			{
-				// Adjust the position to remain constant
-				LayoutData.Offsets.Top = LeftTopDelta.Y + AlignmentOffset.Y;
+				LayoutData.Offsets.Left += LeftTopDelta.X;
+				LayoutData.Offsets.Top += LeftTopDelta.Y;
+
+				// If the slot is stretched horizontally we need to move the right side as it no longer represents width, but
+				// now represents margin from the right stretched side.
+				if ( LayoutData.Anchors.IsStretchedHorizontal() )
+				{
+					LayoutData.Offsets.Right -= LeftTopDelta.X;
+				}
+
+				// If the slot is stretched vertically we need to move the bottom side as it no longer represents width, but
+				// now represents margin from the bottom stretched side.
+				if ( LayoutData.Anchors.IsStretchedVertical() )
+				{
+					LayoutData.Offsets.Bottom -= LeftTopDelta.Y;
+				}
 			}
 		}
 
