UMG - Fixing a crash on exit, the SObjectWidget can survive after the UObject GC has been shutdown, so we ensure UObjects have still been initialized before accessing the UObject on destruction.

diff --git a/Engine/Source/Runtime/UMG/Private/Slate/SObjectWidget.cpp b/Engine/Source/Runtime/UMG/Private/Slate/SObjectWidget.cpp
--- a/Engine/Source/Runtime/UMG/Private/Slate/SObjectWidget.cpp
+++ b/Engine/Source/Runtime/UMG/Private/Slate/SObjectWidget.cpp
@@ -16,7 +16,7 @@ void SObjectWidget::Construct(const FArguments& InArgs, UUserWidget* InWidgetObj
 
 SObjectWidget::~SObjectWidget(void)
 {
-	if ( WidgetObject )
+	if ( UObjectInitialized() && WidgetObject )
 	{
 		WidgetObject->ReleaseNativeWidget();
 	}
diff --git a/Engine/Source/Runtime/UMG/Public/Slate/SObjectWidget.h b/Engine/Source/Runtime/UMG/Public/Slate/SObjectWidget.h
--- a/Engine/Source/Runtime/UMG/Public/Slate/SObjectWidget.h
+++ b/Engine/Source/Runtime/UMG/Public/Slate/SObjectWidget.h
@@ -2,6 +2,8 @@
 
 #pragma once
 
+class UUserWidget;
+
 /**
  * The SObjectWidet allows UMG to insert an SWidget into the hierarchy that manages the lifetime of the
  * UMG UWidget that created it.  Once the SObjectWidget is destroyed it frees the reference it holds to
@@ -20,7 +22,7 @@ class SObjectWidget : public SCompoundWidget, public FGCObject
 
 	virtual ~SObjectWidget(void);
 
-	void Construct(const FArguments& InArgs, class UUserWidget* InWidgetObject);
+	void Construct(const FArguments& InArgs, UUserWidget* InWidgetObject);
 
 	// FGCObject interface
 	virtual void AddReferencedObjects(FReferenceCollector& Collector) override;
@@ -71,5 +73,5 @@ class SObjectWidget : public SCompoundWidget, public FGCObject
 
 private:
 	/** The UWidget that created this SObjectWidget who needs to be kept alive. */
-	class UUserWidget* WidgetObject;
+	UUserWidget* WidgetObject;
 };
\ No newline at end of file
