-Moved all widget toolbox categories from metadata to virtual UWidget::GetToolboxCategory()

diff --git a/Engine/Source/Editor/UMGEditor/Private/Templates/WidgetTemplateClass.cpp b/Engine/Source/Editor/UMGEditor/Private/Templates/WidgetTemplateClass.cpp
--- a/Engine/Source/Editor/UMGEditor/Private/Templates/WidgetTemplateClass.cpp
+++ b/Engine/Source/Editor/UMGEditor/Private/Templates/WidgetTemplateClass.cpp
@@ -15,26 +15,8 @@ FWidgetTemplateClass::FWidgetTemplateClass(TSubclassOf<UWidget> InWidgetClass)
 
 FText FWidgetTemplateClass::GetCategory() const
 {
-	if ( Category.IsEmpty() )
-	{
-		const FString& MetadatCategory = WidgetClass->GetMetaData("Category");
-
-		if ( MetadatCategory.IsEmpty() )
-		{
-			if ( WidgetClass->IsChildOf(UUserWidget::StaticClass()) )
-			{
-				return LOCTEXT("UserControls", "User Controls");
-			}
-
-			return LOCTEXT("Misc", "Misc");
-		}
-		else
-		{
-			Category = FText::FromString(MetadatCategory);
-		}
-	}
-
-	return Category;
+	auto DefaultWidget = WidgetClass->GetDefaultObject<UWidget>();
+	return DefaultWidget->GetToolboxCategory();
 }
 
 UWidget* FWidgetTemplateClass::Create(UWidgetTree* Tree)
diff --git a/Engine/Source/Editor/UMGEditor/Private/Templates/WidgetTemplateClass.h b/Engine/Source/Editor/UMGEditor/Private/Templates/WidgetTemplateClass.h
--- a/Engine/Source/Editor/UMGEditor/Private/Templates/WidgetTemplateClass.h
+++ b/Engine/Source/Editor/UMGEditor/Private/Templates/WidgetTemplateClass.h
@@ -31,7 +31,4 @@ public:
 protected:
 	/** The widget class that will be created by this template */
 	TSubclassOf<UWidget> WidgetClass;
-
-	/** The cached category for the widget. */
-	mutable FText Category;
 };
diff --git a/Engine/Source/Runtime/UMG/Private/Components/Border.cpp b/Engine/Source/Runtime/UMG/Private/Components/Border.cpp
--- a/Engine/Source/Runtime/UMG/Private/Components/Border.cpp
+++ b/Engine/Source/Runtime/UMG/Private/Components/Border.cpp
@@ -183,6 +183,11 @@ const FSlateBrush* UBorder::GetEditorIcon()
 	return FUMGStyle::Get().GetBrush("Widget.Border");
 }
 
+const FText UBorder::GetToolboxCategory()
+{
+	return LOCTEXT("Common", "Common");
+}
+
 #endif
 
 /////////////////////////////////////////////////////
diff --git a/Engine/Source/Runtime/UMG/Private/Components/Button.cpp b/Engine/Source/Runtime/UMG/Private/Components/Button.cpp
--- a/Engine/Source/Runtime/UMG/Private/Components/Button.cpp
+++ b/Engine/Source/Runtime/UMG/Private/Components/Button.cpp
@@ -215,6 +215,11 @@ const FSlateBrush* UButton::GetEditorIcon()
 	return FUMGStyle::Get().GetBrush("Widget.Button");
 }
 
+const FText UButton::GetToolboxCategory()
+{
+	return LOCTEXT("Common", "Common");
+}
+
 #endif
 
 /////////////////////////////////////////////////////
diff --git a/Engine/Source/Runtime/UMG/Private/Components/CanvasPanel.cpp b/Engine/Source/Runtime/UMG/Private/Components/CanvasPanel.cpp
--- a/Engine/Source/Runtime/UMG/Private/Components/CanvasPanel.cpp
+++ b/Engine/Source/Runtime/UMG/Private/Components/CanvasPanel.cpp
@@ -2,6 +2,8 @@
 
 #include "UMGPrivatePCH.h"
 
+#define LOCTEXT_NAMESPACE "UMG"
+
 /////////////////////////////////////////////////////
 // UCanvasPanel
 
@@ -108,4 +110,13 @@ const FSlateBrush* UCanvasPanel::GetEditorIcon()
 	return FUMGStyle::Get().GetBrush("Widget.Canvas");
 }
 
-#endif
\ No newline at end of file
+const FText UCanvasPanel::GetToolboxCategory()
+{
+	return LOCTEXT("Panel", "Panel");
+}
+
+#endif
+
+/////////////////////////////////////////////////////
+
+#undef LOCTEXT_NAMESPACE
\ No newline at end of file
diff --git a/Engine/Source/Runtime/UMG/Private/Components/CheckBox.cpp b/Engine/Source/Runtime/UMG/Private/Components/CheckBox.cpp
--- a/Engine/Source/Runtime/UMG/Private/Components/CheckBox.cpp
+++ b/Engine/Source/Runtime/UMG/Private/Components/CheckBox.cpp
@@ -183,6 +183,11 @@ const FSlateBrush* UCheckBox::GetEditorIcon()
 	return FUMGStyle::Get().GetBrush("Widget.CheckBox");
 }
 
+const FText UCheckBox::GetToolboxCategory()
+{
+	return LOCTEXT("Common", "Common");
+}
+
 #endif
 
 /////////////////////////////////////////////////////
diff --git a/Engine/Source/Runtime/UMG/Private/Components/CircularThrobber.cpp b/Engine/Source/Runtime/UMG/Private/Components/CircularThrobber.cpp
--- a/Engine/Source/Runtime/UMG/Private/Components/CircularThrobber.cpp
+++ b/Engine/Source/Runtime/UMG/Private/Components/CircularThrobber.cpp
@@ -101,6 +101,11 @@ const FSlateBrush* UCircularThrobber::GetEditorIcon()
 	return FUMGStyle::Get().GetBrush("Widget.CircularThrobber");
 }
 
+const FText UCircularThrobber::GetToolboxCategory()
+{
+	return LOCTEXT("Misc", "Misc");
+}
+
 #endif
 
 /////////////////////////////////////////////////////
diff --git a/Engine/Source/Runtime/UMG/Private/Components/ComboBox.cpp b/Engine/Source/Runtime/UMG/Private/Components/ComboBox.cpp
--- a/Engine/Source/Runtime/UMG/Private/Components/ComboBox.cpp
+++ b/Engine/Source/Runtime/UMG/Private/Components/ComboBox.cpp
@@ -55,6 +55,11 @@ const FSlateBrush* UComboBox::GetEditorIcon()
 	return FUMGStyle::Get().GetBrush("Widget.ComboBox");
 }
 
+const FText UComboBox::GetToolboxCategory()
+{
+	return LOCTEXT("Misc", "Misc");
+}
+
 #endif
 
 /////////////////////////////////////////////////////
diff --git a/Engine/Source/Runtime/UMG/Private/Components/ComboBoxString.cpp b/Engine/Source/Runtime/UMG/Private/Components/ComboBoxString.cpp
--- a/Engine/Source/Runtime/UMG/Private/Components/ComboBoxString.cpp
+++ b/Engine/Source/Runtime/UMG/Private/Components/ComboBoxString.cpp
@@ -150,6 +150,11 @@ const FSlateBrush* UComboBoxString::GetEditorIcon()
 	return FUMGStyle::Get().GetBrush("Widget.ComboBox");
 }
 
+const FText UComboBoxString::GetToolboxCategory()
+{
+	return LOCTEXT("Misc", "Misc");
+}
+
 #endif
 
 /////////////////////////////////////////////////////
diff --git a/Engine/Source/Runtime/UMG/Private/Components/EditableText.cpp b/Engine/Source/Runtime/UMG/Private/Components/EditableText.cpp
--- a/Engine/Source/Runtime/UMG/Private/Components/EditableText.cpp
+++ b/Engine/Source/Runtime/UMG/Private/Components/EditableText.cpp
@@ -123,6 +123,11 @@ const FSlateBrush* UEditableText::GetEditorIcon()
 	return FUMGStyle::Get().GetBrush("Widget.EditableText");
 }
 
+const FText UEditableText::GetToolboxCategory()
+{
+	return LOCTEXT("Primitive", "Primitive");
+}
+
 #endif
 
 /////////////////////////////////////////////////////
diff --git a/Engine/Source/Runtime/UMG/Private/Components/EditableTextBox.cpp b/Engine/Source/Runtime/UMG/Private/Components/EditableTextBox.cpp
--- a/Engine/Source/Runtime/UMG/Private/Components/EditableTextBox.cpp
+++ b/Engine/Source/Runtime/UMG/Private/Components/EditableTextBox.cpp
@@ -139,6 +139,11 @@ const FSlateBrush* UEditableTextBox::GetEditorIcon()
 	return FUMGStyle::Get().GetBrush("Widget.EditableTextBox");
 }
 
+const FText UEditableTextBox::GetToolboxCategory()
+{
+	return LOCTEXT("Common", "Common");
+}
+
 #endif
 
 /////////////////////////////////////////////////////
diff --git a/Engine/Source/Runtime/UMG/Private/Components/GridPanel.cpp b/Engine/Source/Runtime/UMG/Private/Components/GridPanel.cpp
--- a/Engine/Source/Runtime/UMG/Private/Components/GridPanel.cpp
+++ b/Engine/Source/Runtime/UMG/Private/Components/GridPanel.cpp
@@ -88,6 +88,11 @@ const FSlateBrush* UGridPanel::GetEditorIcon()
 	return FUMGStyle::Get().GetBrush("Widget.Grid");
 }
 
+const FText UGridPanel::GetToolboxCategory()
+{
+	return LOCTEXT("Panel", "Panel");
+}
+
 #endif
 
 /////////////////////////////////////////////////////
diff --git a/Engine/Source/Runtime/UMG/Private/Components/HorizontalBox.cpp b/Engine/Source/Runtime/UMG/Private/Components/HorizontalBox.cpp
--- a/Engine/Source/Runtime/UMG/Private/Components/HorizontalBox.cpp
+++ b/Engine/Source/Runtime/UMG/Private/Components/HorizontalBox.cpp
@@ -2,6 +2,8 @@
 
 #include "UMGPrivatePCH.h"
 
+#define LOCTEXT_NAMESPACE "UMG"
+
 /////////////////////////////////////////////////////
 // UHorizontalBox
 
@@ -81,4 +83,13 @@ const FSlateBrush* UHorizontalBox::GetEditorIcon()
 	return FUMGStyle::Get().GetBrush("Widget.HorizontalBox");
 }
 
+const FText UHorizontalBox::GetToolboxCategory()
+{
+	return LOCTEXT("Panel", "Panel");
+}
+
 #endif
+
+/////////////////////////////////////////////////////
+
+#undef LOCTEXT_NAMESPACE
\ No newline at end of file
diff --git a/Engine/Source/Runtime/UMG/Private/Components/Image.cpp b/Engine/Source/Runtime/UMG/Private/Components/Image.cpp
--- a/Engine/Source/Runtime/UMG/Private/Components/Image.cpp
+++ b/Engine/Source/Runtime/UMG/Private/Components/Image.cpp
@@ -197,6 +197,11 @@ const FSlateBrush* UImage::GetEditorIcon()
 	return FUMGStyle::Get().GetBrush("Widget.Image");
 }
 
+const FText UImage::GetToolboxCategory()
+{
+	return LOCTEXT("Common", "Common");
+}
+
 #endif
 
 
diff --git a/Engine/Source/Runtime/UMG/Private/Components/ListView.cpp b/Engine/Source/Runtime/UMG/Private/Components/ListView.cpp
--- a/Engine/Source/Runtime/UMG/Private/Components/ListView.cpp
+++ b/Engine/Source/Runtime/UMG/Private/Components/ListView.cpp
@@ -66,6 +66,11 @@ const FSlateBrush* UListView::GetEditorIcon()
 	return FUMGStyle::Get().GetBrush("Widget.ListView");
 }
 
+const FText UListView::GetToolboxCategory()
+{
+	return LOCTEXT("Misc", "Misc");
+}
+
 #endif
 
 /////////////////////////////////////////////////////
diff --git a/Engine/Source/Runtime/UMG/Private/Components/MenuAnchor.cpp b/Engine/Source/Runtime/UMG/Private/Components/MenuAnchor.cpp
--- a/Engine/Source/Runtime/UMG/Private/Components/MenuAnchor.cpp
+++ b/Engine/Source/Runtime/UMG/Private/Components/MenuAnchor.cpp
@@ -103,6 +103,11 @@ const FSlateBrush* UMenuAnchor::GetEditorIcon()
 	return FUMGStyle::Get().GetBrush("Widget.MenuAnchor");
 }
 
+const FText UMenuAnchor::GetToolboxCategory()
+{
+	return LOCTEXT("Advanced", "Advanced");
+}
+
 #endif
 
 /////////////////////////////////////////////////////
diff --git a/Engine/Source/Runtime/UMG/Private/Components/MultiLineEditableText.cpp b/Engine/Source/Runtime/UMG/Private/Components/MultiLineEditableText.cpp
--- a/Engine/Source/Runtime/UMG/Private/Components/MultiLineEditableText.cpp
+++ b/Engine/Source/Runtime/UMG/Private/Components/MultiLineEditableText.cpp
@@ -101,6 +101,11 @@ const FSlateBrush* UMultiLineEditableText::GetEditorIcon()
 	return FUMGStyle::Get().GetBrush("Widget.MultiLineEditableText");
 }
 
+const FText UMultiLineEditableText::GetToolboxCategory()
+{
+	return LOCTEXT("Primitive", "Primitive");
+}
+
 #endif
 
 /////////////////////////////////////////////////////
diff --git a/Engine/Source/Runtime/UMG/Private/Components/MultiLineEditableTextBox.cpp b/Engine/Source/Runtime/UMG/Private/Components/MultiLineEditableTextBox.cpp
--- a/Engine/Source/Runtime/UMG/Private/Components/MultiLineEditableTextBox.cpp
+++ b/Engine/Source/Runtime/UMG/Private/Components/MultiLineEditableTextBox.cpp
@@ -110,6 +110,11 @@ const FSlateBrush* UMultiLineEditableTextBox::GetEditorIcon()
 	return FUMGStyle::Get().GetBrush("Widget.MultiLineEditableTextBox");
 }
 
+const FText UMultiLineEditableTextBox::GetToolboxCategory()
+{
+	return LOCTEXT("Common", "Common");
+}
+
 #endif
 
 /////////////////////////////////////////////////////
diff --git a/Engine/Source/Runtime/UMG/Private/Components/NativeWidgetHost.cpp b/Engine/Source/Runtime/UMG/Private/Components/NativeWidgetHost.cpp
--- a/Engine/Source/Runtime/UMG/Private/Components/NativeWidgetHost.cpp
+++ b/Engine/Source/Runtime/UMG/Private/Components/NativeWidgetHost.cpp
@@ -52,6 +52,11 @@ const FSlateBrush* UNativeWidgetHost::GetEditorIcon()
 	return FUMGStyle::Get().GetBrush("Widget.NativeWidgetHost");
 }
 
+const FText UNativeWidgetHost::GetToolboxCategory()
+{
+	return LOCTEXT("Advanced", "Advanced");
+}
+
 #endif
 
 /////////////////////////////////////////////////////
diff --git a/Engine/Source/Runtime/UMG/Private/Components/Overlay.cpp b/Engine/Source/Runtime/UMG/Private/Components/Overlay.cpp
--- a/Engine/Source/Runtime/UMG/Private/Components/Overlay.cpp
+++ b/Engine/Source/Runtime/UMG/Private/Components/Overlay.cpp
@@ -2,6 +2,8 @@
 
 #include "UMGPrivatePCH.h"
 
+#define LOCTEXT_NAMESPACE "UMG"
+
 /////////////////////////////////////////////////////
 // UOverlay
 
@@ -71,4 +73,13 @@ const FSlateBrush* UOverlay::GetEditorIcon()
 	return FUMGStyle::Get().GetBrush("Widget.Overlay");
 }
 
+const FText UOverlay::GetToolboxCategory()
+{
+	return LOCTEXT("Panel", "Panel");
+}
+
 #endif
+
+/////////////////////////////////////////////////////
+
+#undef LOCTEXT_NAMESPACE
\ No newline at end of file
diff --git a/Engine/Source/Runtime/UMG/Private/Components/ProcessBar.cpp b/Engine/Source/Runtime/UMG/Private/Components/ProcessBar.cpp
deleted file mode 100644
diff --git a/Engine/Source/Runtime/UMG/Private/Components/ProgressBar.cpp b/Engine/Source/Runtime/UMG/Private/Components/ProgressBar.cpp
new file mode 100644
--- /dev/null
+++ b/Engine/Source/Runtime/UMG/Private/Components/ProgressBar.cpp
@@ -0,0 +1,73 @@
+// Copyright 1998-2014 Epic Games, Inc. All Rights Reserved.
+
+#include "UMGPrivatePCH.h"
+
+#define LOCTEXT_NAMESPACE "UMG"
+
+/////////////////////////////////////////////////////
+// UProgressBar
+
+UProgressBar::UProgressBar(const FPostConstructInitializeProperties& PCIP)
+	: Super(PCIP)
+{
+	BarFillType = EProgressBarFillType::LeftToRight;
+	bIsMarquee = false;
+	Percent = 0;
+	FillColorAndOpacity = FLinearColor::White;
+	BorderPadding = FVector2D(1, 0);
+}
+
+TSharedRef<SWidget> UProgressBar::RebuildWidget()
+{
+	MyProgressBar = SNew(SProgressBar);
+
+	return MyProgressBar.ToSharedRef();
+}
+
+void UProgressBar::SynchronizeProperties()
+{
+	const FProgressBarStyle* StylePtr = ( Style != NULL ) ? Style->GetStyle<FProgressBarStyle>() : NULL;
+	if ( StylePtr )
+	{
+		MyProgressBar->SetStyle(StylePtr);
+	}
+
+	TAttribute< TOptional<float> > PercentBinding = OPTIONAL_BINDING_CONVERT(float, Percent, TOptional<float>, ConvertFloatToOptionalFloat);
+
+	MyProgressBar->SetPercent(bIsMarquee ? TOptional<float>() : PercentBinding);
+	
+	MyProgressBar->SetBarFillType(BarFillType);
+	MyProgressBar->SetFillColorAndOpacity(FillColorAndOpacity);
+	MyProgressBar->SetBorderPadding(BorderPadding);
+	
+	MyProgressBar->SetBackgroundImage(BackgroundImage ? &BackgroundImage->Brush : nullptr);
+	MyProgressBar->SetFillImage(FillImage ? &FillImage->Brush : nullptr);
+	MyProgressBar->SetMarqueeImage(MarqueeImage ? &MarqueeImage->Brush : nullptr);
+}
+
+void UProgressBar::SetPercent(float InPercent)
+{
+	Percent = InPercent;
+	if ( MyProgressBar.IsValid() )
+	{
+		MyProgressBar->SetPercent(InPercent);
+	}
+}
+
+#if WITH_EDITOR
+
+const FSlateBrush* UProgressBar::GetEditorIcon()
+{
+	return FUMGStyle::Get().GetBrush("Widget.ProgressBar");
+}
+
+const FText UProgressBar::GetToolboxCategory()
+{
+	return LOCTEXT("Common", "Common");
+}
+
+#endif
+
+/////////////////////////////////////////////////////
+
+#undef LOCTEXT_NAMESPACE
diff --git a/Engine/Source/Runtime/UMG/Private/Components/ScaleBox.cpp b/Engine/Source/Runtime/UMG/Private/Components/ScaleBox.cpp
--- a/Engine/Source/Runtime/UMG/Private/Components/ScaleBox.cpp
+++ b/Engine/Source/Runtime/UMG/Private/Components/ScaleBox.cpp
@@ -73,6 +73,11 @@ const FSlateBrush* UScaleBox::GetEditorIcon()
 	return FUMGStyle::Get().GetBrush("Widget.ScaleBox");
 }
 
+const FText UScaleBox::GetToolboxCategory()
+{
+	return LOCTEXT("Panel", "Panel");
+}
+
 #endif
 
 /////////////////////////////////////////////////////
diff --git a/Engine/Source/Runtime/UMG/Private/Components/ScrollBar.cpp b/Engine/Source/Runtime/UMG/Private/Components/ScrollBar.cpp
--- a/Engine/Source/Runtime/UMG/Private/Components/ScrollBar.cpp
+++ b/Engine/Source/Runtime/UMG/Private/Components/ScrollBar.cpp
@@ -57,6 +57,11 @@ const FSlateBrush* UScrollBar::GetEditorIcon()
 	return FUMGStyle::Get().GetBrush("Widget.ScrollBar");
 }
 
+const FText UScrollBar::GetToolboxCategory()
+{
+	return LOCTEXT("Advanced", "Advanced");
+}
+
 #endif
 
 /////////////////////////////////////////////////////
diff --git a/Engine/Source/Runtime/UMG/Private/Components/ScrollBox.cpp b/Engine/Source/Runtime/UMG/Private/Components/ScrollBox.cpp
--- a/Engine/Source/Runtime/UMG/Private/Components/ScrollBox.cpp
+++ b/Engine/Source/Runtime/UMG/Private/Components/ScrollBox.cpp
@@ -107,6 +107,11 @@ const FSlateBrush* UScrollBox::GetEditorIcon()
 	return FUMGStyle::Get().GetBrush("Widget.ScrollBox");
 }
 
+const FText UScrollBox::GetToolboxCategory()
+{
+	return LOCTEXT("Panel", "Panel");
+}
+
 #endif
 
 /////////////////////////////////////////////////////
diff --git a/Engine/Source/Runtime/UMG/Private/Components/SizeBox.cpp b/Engine/Source/Runtime/UMG/Private/Components/SizeBox.cpp
--- a/Engine/Source/Runtime/UMG/Private/Components/SizeBox.cpp
+++ b/Engine/Source/Runtime/UMG/Private/Components/SizeBox.cpp
@@ -236,6 +236,11 @@ const FSlateBrush* USizeBox::GetEditorIcon()
 	return FUMGStyle::Get().GetBrush("Widget.SizeBox");
 }
 
+const FText USizeBox::GetToolboxCategory()
+{
+	return LOCTEXT("Common", "Common");
+}
+
 #endif
 
 /////////////////////////////////////////////////////
diff --git a/Engine/Source/Runtime/UMG/Private/Components/Slider.cpp b/Engine/Source/Runtime/UMG/Private/Components/Slider.cpp
--- a/Engine/Source/Runtime/UMG/Private/Components/Slider.cpp
+++ b/Engine/Source/Runtime/UMG/Private/Components/Slider.cpp
@@ -78,6 +78,11 @@ const FSlateBrush* USlider::GetEditorIcon()
 	return FUMGStyle::Get().GetBrush("Widget.Slider");
 }
 
+const FText USlider::GetToolboxCategory()
+{
+	return LOCTEXT("Common", "Common");
+}
+
 #endif
 
 /////////////////////////////////////////////////////
diff --git a/Engine/Source/Runtime/UMG/Private/Components/Spacer.cpp b/Engine/Source/Runtime/UMG/Private/Components/Spacer.cpp
--- a/Engine/Source/Runtime/UMG/Private/Components/Spacer.cpp
+++ b/Engine/Source/Runtime/UMG/Private/Components/Spacer.cpp
@@ -48,6 +48,11 @@ const FSlateBrush* USpacer::GetEditorIcon()
 	return FUMGStyle::Get().GetBrush("Widget.Spacer");
 }
 
+const FText USpacer::GetToolboxCategory()
+{
+	return LOCTEXT("Primitive", "Primitive");
+}
+
 #endif
 
 /////////////////////////////////////////////////////
diff --git a/Engine/Source/Runtime/UMG/Private/Components/SpinBox.cpp b/Engine/Source/Runtime/UMG/Private/Components/SpinBox.cpp
--- a/Engine/Source/Runtime/UMG/Private/Components/SpinBox.cpp
+++ b/Engine/Source/Runtime/UMG/Private/Components/SpinBox.cpp
@@ -274,6 +274,11 @@ const FSlateBrush* USpinBox::GetEditorIcon()
 	return FUMGStyle::Get().GetBrush("Widget.SpinBox");
 }
 
+const FText USpinBox::GetToolboxCategory()
+{
+	return LOCTEXT("Common", "Common");
+}
+
 #endif
 
 /////////////////////////////////////////////////////
diff --git a/Engine/Source/Runtime/UMG/Private/Components/TextBlock.cpp b/Engine/Source/Runtime/UMG/Private/Components/TextBlock.cpp
--- a/Engine/Source/Runtime/UMG/Private/Components/TextBlock.cpp
+++ b/Engine/Source/Runtime/UMG/Private/Components/TextBlock.cpp
@@ -148,6 +148,11 @@ const FSlateBrush* UTextBlock::GetEditorIcon()
 	return FUMGStyle::Get().GetBrush("Widget.TextBlock");
 }
 
+const FText UTextBlock::GetToolboxCategory()
+{
+	return LOCTEXT("Primitive", "Primitive");
+}
+
 #endif
 
 /////////////////////////////////////////////////////
diff --git a/Engine/Source/Runtime/UMG/Private/Components/Throbber.cpp b/Engine/Source/Runtime/UMG/Private/Components/Throbber.cpp
--- a/Engine/Source/Runtime/UMG/Private/Components/Throbber.cpp
+++ b/Engine/Source/Runtime/UMG/Private/Components/Throbber.cpp
@@ -112,6 +112,11 @@ const FSlateBrush* UThrobber::GetEditorIcon()
 	return FUMGStyle::Get().GetBrush("Widget.Throbber");
 }
 
+const FText UThrobber::GetToolboxCategory()
+{
+	return LOCTEXT("Misc", "Misc");
+}
+
 #endif
 
 /////////////////////////////////////////////////////
diff --git a/Engine/Source/Runtime/UMG/Private/Components/TileView.cpp b/Engine/Source/Runtime/UMG/Private/Components/TileView.cpp
--- a/Engine/Source/Runtime/UMG/Private/Components/TileView.cpp
+++ b/Engine/Source/Runtime/UMG/Private/Components/TileView.cpp
@@ -83,6 +83,11 @@ const FSlateBrush* UTileView::GetEditorIcon()
 	return FUMGStyle::Get().GetBrush("Widget.TileView");
 }
 
+const FText UTileView::GetToolboxCategory()
+{
+	return LOCTEXT("Misc", "Misc");
+}
+
 #endif
 
 /////////////////////////////////////////////////////
diff --git a/Engine/Source/Runtime/UMG/Private/Components/UniformGridPanel.cpp b/Engine/Source/Runtime/UMG/Private/Components/UniformGridPanel.cpp
--- a/Engine/Source/Runtime/UMG/Private/Components/UniformGridPanel.cpp
+++ b/Engine/Source/Runtime/UMG/Private/Components/UniformGridPanel.cpp
@@ -83,6 +83,11 @@ const FSlateBrush* UUniformGridPanel::GetEditorIcon()
 	return FUMGStyle::Get().GetBrush("Widget.UniformGrid");
 }
 
+const FText UUniformGridPanel::GetToolboxCategory()
+{
+	return LOCTEXT("Panel", "Panel");
+}
+
 #endif
 
 /////////////////////////////////////////////////////
diff --git a/Engine/Source/Runtime/UMG/Private/Components/VerticalBox.cpp b/Engine/Source/Runtime/UMG/Private/Components/VerticalBox.cpp
--- a/Engine/Source/Runtime/UMG/Private/Components/VerticalBox.cpp
+++ b/Engine/Source/Runtime/UMG/Private/Components/VerticalBox.cpp
@@ -2,6 +2,8 @@
 
 #include "UMGPrivatePCH.h"
 
+#define LOCTEXT_NAMESPACE "UMG"
+
 /////////////////////////////////////////////////////
 // UVerticalBox
 
@@ -76,4 +78,13 @@ const FSlateBrush* UVerticalBox::GetEditorIcon()
 	return FUMGStyle::Get().GetBrush("Widget.VerticalBox");
 }
 
+const FText UVerticalBox::GetToolboxCategory()
+{
+	return LOCTEXT("Panel", "Panel");
+}
+
 #endif
+
+/////////////////////////////////////////////////////
+
+#undef LOCTEXT_NAMESPACE
\ No newline at end of file
diff --git a/Engine/Source/Runtime/UMG/Private/Components/Viewport.cpp b/Engine/Source/Runtime/UMG/Private/Components/Viewport.cpp
--- a/Engine/Source/Runtime/UMG/Private/Components/Viewport.cpp
+++ b/Engine/Source/Runtime/UMG/Private/Components/Viewport.cpp
@@ -626,6 +626,11 @@ const FSlateBrush* UViewport::GetEditorIcon()
 	return FUMGStyle::Get().GetBrush("Widget.Viewport");
 }
 
+const FText UViewport::GetToolboxCategory()
+{
+	return LOCTEXT("Advanced", "Advanced");
+}
+
 #endif
 
 /////////////////////////////////////////////////////
diff --git a/Engine/Source/Runtime/UMG/Private/Components/Widget.cpp b/Engine/Source/Runtime/UMG/Private/Components/Widget.cpp
--- a/Engine/Source/Runtime/UMG/Private/Components/Widget.cpp
+++ b/Engine/Source/Runtime/UMG/Private/Components/Widget.cpp
@@ -225,6 +225,7 @@ TSharedRef<SWidget> UWidget::BuildDesignTimeWidget(TSharedRef<SWidget> WrapWidge
 }
 
 #if WITH_EDITOR
+#define LOCTEXT_NAMESPACE "UMGEditor"
 
 bool UWidget::IsGeneratedName() const
 {
@@ -256,6 +257,11 @@ FString UWidget::GetLabel() const
 	}
 }
 
+const FText UWidget::GetToolboxCategory()
+{
+	return LOCTEXT("Uncategorized", "Uncategorized");
+}
+
 const FSlateBrush* UWidget::GetEditorIcon()
 {
 	return FUMGStyle::Get().GetBrush("Widget");
@@ -301,6 +307,7 @@ void UWidget::Deselect()
 	}
 }
 
+#undef LOCTEXT_NAMESPACE
 #endif
 
 bool UWidget::Modify(bool bAlwaysMarkDirty)
diff --git a/Engine/Source/Runtime/UMG/Private/Components/WidgetSwitcher.cpp b/Engine/Source/Runtime/UMG/Private/Components/WidgetSwitcher.cpp
--- a/Engine/Source/Runtime/UMG/Private/Components/WidgetSwitcher.cpp
+++ b/Engine/Source/Runtime/UMG/Private/Components/WidgetSwitcher.cpp
@@ -122,6 +122,11 @@ const FSlateBrush* UWidgetSwitcher::GetEditorIcon()
 	return FUMGStyle::Get().GetBrush("Widget.WidgetSwitcher");
 }
 
+const FText UWidgetSwitcher::GetToolboxCategory()
+{
+	return LOCTEXT("Panel", "Panel");
+}
+
 void UWidgetSwitcher::OnDescendantSelected(UWidget* DescendantWidget)
 {
 	// Temporarily sets the active child to the selected child to make
diff --git a/Engine/Source/Runtime/UMG/Private/UserWidget.cpp b/Engine/Source/Runtime/UMG/Private/UserWidget.cpp
--- a/Engine/Source/Runtime/UMG/Private/UserWidget.cpp
+++ b/Engine/Source/Runtime/UMG/Private/UserWidget.cpp
@@ -7,6 +7,8 @@
 #include "WidgetAnimation.h"
 #include "WidgetBlueprintLibrary.h"
 
+#define LOCTEXT_NAMESPACE "UMG"
+
 /**
  * This class holds onto the widget when it's placed into the viewport.
  */
@@ -719,4 +721,13 @@ const FSlateBrush* UUserWidget::GetEditorIcon()
 	return FUMGStyle::Get().GetBrush("Widget.UserWidget");
 }
 
+const FText UUserWidget::GetToolboxCategory()
+{
+	return LOCTEXT("UserCreated", "User Created");
+}
+
 #endif
+
+/////////////////////////////////////////////////////
+
+#undef LOCTEXT_NAMESPACE
diff --git a/Engine/Source/Runtime/UMG/Public/Blueprint/UserWidget.h b/Engine/Source/Runtime/UMG/Public/Blueprint/UserWidget.h
--- a/Engine/Source/Runtime/UMG/Public/Blueprint/UserWidget.h
+++ b/Engine/Source/Runtime/UMG/Public/Blueprint/UserWidget.h
@@ -268,6 +268,7 @@ public:
 #if WITH_EDITOR
 	// UWidget interface
 	virtual const FSlateBrush* GetEditorIcon() override;
+	virtual const FText GetToolboxCategory() override;
 	// End UWidget interface
 #endif
 
diff --git a/Engine/Source/Runtime/UMG/Public/Components/Border.h b/Engine/Source/Runtime/UMG/Public/Components/Border.h
--- a/Engine/Source/Runtime/UMG/Public/Components/Border.h
+++ b/Engine/Source/Runtime/UMG/Public/Components/Border.h
@@ -10,7 +10,7 @@ class USlateBrushAsset;
  * A border is a container widget that can contain one child widget, providing an opportunity 
  * to surround it with a border image and adjustable padding.
  */
-UCLASS(meta=( Category="Common" ), ClassGroup=UserInterface)
+UCLASS(ClassGroup=UserInterface)
 class UMG_API UBorder : public UContentWidget
 {
 	GENERATED_UCLASS_BODY()
@@ -89,6 +89,7 @@ public:
 
 #if WITH_EDITOR
 	virtual const FSlateBrush* GetEditorIcon() override;
+	virtual const FText GetToolboxCategory() override;
 #endif
 
 protected:
diff --git a/Engine/Source/Runtime/UMG/Public/Components/Button.h b/Engine/Source/Runtime/UMG/Public/Components/Button.h
--- a/Engine/Source/Runtime/UMG/Public/Components/Button.h
+++ b/Engine/Source/Runtime/UMG/Public/Components/Button.h
@@ -9,7 +9,7 @@ DECLARE_DYNAMIC_MULTICAST_DELEGATE(FOnButtonClickedEvent);
 /**
  * The button is a clickable primitive widget to enable basic interaction.
  */
-UCLASS(meta=( Category="Common" ), ClassGroup=UserInterface)
+UCLASS(ClassGroup=UserInterface)
 class UMG_API UButton : public UContentWidget
 {
 	GENERATED_UCLASS_BODY()
@@ -99,6 +99,7 @@ public:
 
 #if WITH_EDITOR
 	virtual const FSlateBrush* GetEditorIcon() override;
+	virtual const FText GetToolboxCategory() override;
 #endif
 
 protected:
diff --git a/Engine/Source/Runtime/UMG/Public/Components/CanvasPanel.h b/Engine/Source/Runtime/UMG/Public/Components/CanvasPanel.h
--- a/Engine/Source/Runtime/UMG/Public/Components/CanvasPanel.h
+++ b/Engine/Source/Runtime/UMG/Public/Components/CanvasPanel.h
@@ -4,7 +4,7 @@
 
 #include "CanvasPanel.generated.h"
 
-UCLASS(meta=( Category="Panel" ), ClassGroup=UserInterface)
+UCLASS(ClassGroup=UserInterface)
 class UMG_API UCanvasPanel : public UPanelWidget
 {
 	GENERATED_UCLASS_BODY()
@@ -25,6 +25,7 @@ public:
 #if WITH_EDITOR
 	// UWidget interface
 	virtual const FSlateBrush* GetEditorIcon() override;
+	virtual const FText GetToolboxCategory() override;
 	// End UWidget interface
 
 	// UWidget interface
diff --git a/Engine/Source/Runtime/UMG/Public/Components/CheckBox.h b/Engine/Source/Runtime/UMG/Public/Components/CheckBox.h
--- a/Engine/Source/Runtime/UMG/Public/Components/CheckBox.h
+++ b/Engine/Source/Runtime/UMG/Public/Components/CheckBox.h
@@ -9,7 +9,7 @@ DECLARE_DYNAMIC_MULTICAST_DELEGATE_OneParam( FOnCheckBoxComponentStateChanged, b
 /**
  * Check box widget
  */
-UCLASS(meta=( Category="Common" ), ClassGroup=UserInterface)
+UCLASS(ClassGroup=UserInterface)
 class UMG_API UCheckBox : public UContentWidget
 {
 	GENERATED_UCLASS_BODY()
@@ -122,6 +122,7 @@ public:
 
 #if WITH_EDITOR
 	virtual const FSlateBrush* GetEditorIcon() override;
+	virtual const FText GetToolboxCategory() override;
 #endif
 
 protected:
diff --git a/Engine/Source/Runtime/UMG/Public/Components/CircularThrobber.h b/Engine/Source/Runtime/UMG/Public/Components/CircularThrobber.h
--- a/Engine/Source/Runtime/UMG/Public/Components/CircularThrobber.h
+++ b/Engine/Source/Runtime/UMG/Public/Components/CircularThrobber.h
@@ -7,7 +7,7 @@
 class USlateBrushAsset;
 
 /**  A throbber widget that orients images in a spinning circle. */
-UCLASS(meta=( Category="Misc" ), ClassGroup=UserInterface)
+UCLASS(ClassGroup=UserInterface)
 class UMG_API UCircularThrobber : public UWidget
 {
 	GENERATED_UCLASS_BODY()
@@ -38,6 +38,7 @@ public:
 
 #if WITH_EDITOR
 	virtual const FSlateBrush* GetEditorIcon() override;
+	virtual const FText GetToolboxCategory() override;
 #endif
 
 protected:
diff --git a/Engine/Source/Runtime/UMG/Public/Components/ComboBox.h b/Engine/Source/Runtime/UMG/Public/Components/ComboBox.h
--- a/Engine/Source/Runtime/UMG/Public/Components/ComboBox.h
+++ b/Engine/Source/Runtime/UMG/Public/Components/ComboBox.h
@@ -7,7 +7,7 @@
 /**
  * The combobox allows you to display a list of options to the user in a dropdown menu for them to select one.
  */
-UCLASS( meta=( DisplayName="ComboBox (Object)", Category="Misc" ), ClassGroup=UserInterface)
+UCLASS( meta=( DisplayName="ComboBox (Object)"), ClassGroup=UserInterface)
 class UMG_API UComboBox : public UWidget
 {
 	GENERATED_UCLASS_BODY()
@@ -26,6 +26,7 @@ public:
 
 #if WITH_EDITOR
 	virtual const FSlateBrush* GetEditorIcon() override;
+	virtual const FText GetToolboxCategory() override;
 #endif
 
 protected:
diff --git a/Engine/Source/Runtime/UMG/Public/Components/ComboBoxString.h b/Engine/Source/Runtime/UMG/Public/Components/ComboBoxString.h
--- a/Engine/Source/Runtime/UMG/Public/Components/ComboBoxString.h
+++ b/Engine/Source/Runtime/UMG/Public/Components/ComboBoxString.h
@@ -7,7 +7,7 @@
 /**
  * The combobox allows you to display a list of options to the user in a dropdown menu for them to select one.
  */
-UCLASS( meta=( DisplayName="ComboBox (String)", Category="Misc" ), ClassGroup=UserInterface)
+UCLASS( meta=( DisplayName="ComboBox (String)"), ClassGroup=UserInterface)
 class UMG_API UComboBoxString : public UWidget
 {
 	GENERATED_UCLASS_BODY()
@@ -69,6 +69,7 @@ public:
 
 #if WITH_EDITOR
 	virtual const FSlateBrush* GetEditorIcon() override;
+	virtual const FText GetToolboxCategory() override;
 #endif
 
 protected:
diff --git a/Engine/Source/Runtime/UMG/Public/Components/EditableText.h b/Engine/Source/Runtime/UMG/Public/Components/EditableText.h
--- a/Engine/Source/Runtime/UMG/Public/Components/EditableText.h
+++ b/Engine/Source/Runtime/UMG/Public/Components/EditableText.h
@@ -107,19 +107,26 @@ public:
 	UFUNCTION(BlueprintCallable, Category="Widget")
 	FText GetText() const;
 
-	/**  */
+	/** 
+	* Directly sets the text.
+	* *Warning*: This will wipe any binding created for the Text property!
+	*
+	* @param InText The text to assign to the widget
+	*/
 	UFUNCTION(BlueprintCallable, Category="Widget")
 	void SetText(FText InText);
 	
+	virtual void ReleaseNativeWidget() override;
+
+public:
 	// UWidget interface
 	virtual void SynchronizeProperties() override;
-	// End of UWidget interface
-
-	virtual void ReleaseNativeWidget() override;
 
 #if WITH_EDITOR
 	virtual const FSlateBrush* GetEditorIcon() override;
+	virtual const FText GetToolboxCategory() override;
 #endif
+	// End of UWidget interface
 
 protected:
 	// UWidget interface
diff --git a/Engine/Source/Runtime/UMG/Public/Components/EditableTextBox.h b/Engine/Source/Runtime/UMG/Public/Components/EditableTextBox.h
--- a/Engine/Source/Runtime/UMG/Public/Components/EditableTextBox.h
+++ b/Engine/Source/Runtime/UMG/Public/Components/EditableTextBox.h
@@ -5,7 +5,7 @@
 #include "EditableTextBox.generated.h"
 
 /** Editable text box widget */
-UCLASS(meta=( Category="Common" ), ClassGroup=UserInterface)
+UCLASS(ClassGroup=UserInterface)
 class UMG_API UEditableTextBox : public UWidget
 {
 	GENERATED_UCLASS_BODY()
@@ -124,6 +124,7 @@ public:
 
 #if WITH_EDITOR
 	virtual const FSlateBrush* GetEditorIcon() override;
+	virtual const FText GetToolboxCategory() override;
 #endif
 
 protected:
diff --git a/Engine/Source/Runtime/UMG/Public/Components/GridPanel.h b/Engine/Source/Runtime/UMG/Public/Components/GridPanel.h
--- a/Engine/Source/Runtime/UMG/Public/Components/GridPanel.h
+++ b/Engine/Source/Runtime/UMG/Public/Components/GridPanel.h
@@ -5,7 +5,7 @@
 #include "GridPanel.generated.h"
 
 /** A panel that evenly divides up available space between all of its children. */
-UCLASS(meta=( Category="Panel" ), ClassGroup=UserInterface)
+UCLASS(ClassGroup=UserInterface)
 class UMG_API UGridPanel : public UPanelWidget
 {
 	GENERATED_UCLASS_BODY()
@@ -30,6 +30,7 @@ public:
 
 #if WITH_EDITOR
 	virtual const FSlateBrush* GetEditorIcon() override;
+	virtual const FText GetToolboxCategory() override;
 #endif
 
 protected:
diff --git a/Engine/Source/Runtime/UMG/Public/Components/HorizontalBox.h b/Engine/Source/Runtime/UMG/Public/Components/HorizontalBox.h
--- a/Engine/Source/Runtime/UMG/Public/Components/HorizontalBox.h
+++ b/Engine/Source/Runtime/UMG/Public/Components/HorizontalBox.h
@@ -4,7 +4,7 @@
 
 #include "HorizontalBox.generated.h"
 
-UCLASS(meta=( Category="Panel" ), ClassGroup=UserInterface)
+UCLASS(ClassGroup=UserInterface)
 class UMG_API UHorizontalBox : public UPanelWidget
 {
 	GENERATED_UCLASS_BODY()
@@ -20,6 +20,7 @@ class UMG_API UHorizontalBox : public UPanelWidget
 #if WITH_EDITOR
 	// UWidget interface
 	virtual const FSlateBrush* GetEditorIcon() override;
+	virtual const FText GetToolboxCategory() override;
 	// End UWidget interface
 #endif
 
diff --git a/Engine/Source/Runtime/UMG/Public/Components/Image.h b/Engine/Source/Runtime/UMG/Public/Components/Image.h
--- a/Engine/Source/Runtime/UMG/Public/Components/Image.h
+++ b/Engine/Source/Runtime/UMG/Public/Components/Image.h
@@ -9,7 +9,7 @@ class USlateBrushAsset;
 /**
  * The image widget allows you to display a Slate Brush, or texture or material in the UI.
  */
-UCLASS(meta=( Category="Common" ), ClassGroup=UserInterface)
+UCLASS(ClassGroup=UserInterface)
 class UMG_API UImage : public UWidget
 {
 	GENERATED_UCLASS_BODY()
@@ -72,6 +72,7 @@ public:
 #if WITH_EDITOR
 	// UWidget interface
 	virtual const FSlateBrush* GetEditorIcon() override;
+	virtual const FText GetToolboxCategory() override;
 	// End UWidget interface
 #endif
 
diff --git a/Engine/Source/Runtime/UMG/Public/Components/ListView.h b/Engine/Source/Runtime/UMG/Public/Components/ListView.h
--- a/Engine/Source/Runtime/UMG/Public/Components/ListView.h
+++ b/Engine/Source/Runtime/UMG/Public/Components/ListView.h
@@ -5,7 +5,7 @@
 #include "ListView.generated.h"
 
 /** Allows thousands of items to be displayed in a list.  Generates widgets dynamically for each item. */
-UCLASS(meta=( Category="Misc" ), ClassGroup=UserInterface)
+UCLASS(ClassGroup=UserInterface)
 class UMG_API UListView : public UTableViewBase
 {
 	GENERATED_UCLASS_BODY()
@@ -30,6 +30,7 @@ public:
 
 #if WITH_EDITOR
 	virtual const FSlateBrush* GetEditorIcon() override;
+	virtual const FText GetToolboxCategory() override;
 #endif
 
 protected:
diff --git a/Engine/Source/Runtime/UMG/Public/Components/MenuAnchor.h b/Engine/Source/Runtime/UMG/Public/Components/MenuAnchor.h
--- a/Engine/Source/Runtime/UMG/Public/Components/MenuAnchor.h
+++ b/Engine/Source/Runtime/UMG/Public/Components/MenuAnchor.h
@@ -8,7 +8,7 @@
  * The Menu Anchor allows you to specify an location that a popup menu should be anchored to, and should be
  * summoned from.
  */
-UCLASS(meta=( Category="Advanced" ), ClassGroup=UserInterface, meta=( IsAdvanced = "True" ))
+UCLASS(ClassGroup=UserInterface, meta=( IsAdvanced = "True" ))
 class UMG_API UMenuAnchor : public UContentWidget
 {
 	GENERATED_UCLASS_BODY()
@@ -56,6 +56,7 @@ public:
 
 #if WITH_EDITOR
 	virtual const FSlateBrush* GetEditorIcon() override;
+	virtual const FText GetToolboxCategory() override;
 #endif
 
 protected:
diff --git a/Engine/Source/Runtime/UMG/Public/Components/MultiLineEditableText.h b/Engine/Source/Runtime/UMG/Public/Components/MultiLineEditableText.h
--- a/Engine/Source/Runtime/UMG/Public/Components/MultiLineEditableText.h
+++ b/Engine/Source/Runtime/UMG/Public/Components/MultiLineEditableText.h
@@ -5,7 +5,7 @@
 #include "MultiLineEditableText.generated.h"
 
 /** Editable text box widget */
-UCLASS(meta=( Category="Primitive" ), ClassGroup=UserInterface)
+UCLASS(ClassGroup=UserInterface)
 class UMG_API UMultiLineEditableText : public UWidget
 {
 	GENERATED_UCLASS_BODY()
@@ -51,6 +51,7 @@ public:
 
 #if WITH_EDITOR
 	virtual const FSlateBrush* GetEditorIcon() override;
+	virtual const FText GetToolboxCategory() override;
 #endif
 
 protected:
diff --git a/Engine/Source/Runtime/UMG/Public/Components/MultiLineEditableTextBox.h b/Engine/Source/Runtime/UMG/Public/Components/MultiLineEditableTextBox.h
--- a/Engine/Source/Runtime/UMG/Public/Components/MultiLineEditableTextBox.h
+++ b/Engine/Source/Runtime/UMG/Public/Components/MultiLineEditableTextBox.h
@@ -5,7 +5,7 @@
 #include "MultiLineEditableTextBox.generated.h"
 
 /** Editable text box widget */
-UCLASS(meta=( Category="Common" ), ClassGroup=UserInterface)
+UCLASS(ClassGroup=UserInterface)
 class UMG_API UMultiLineEditableTextBox : public UWidget
 {
 	GENERATED_UCLASS_BODY()
@@ -73,6 +73,7 @@ public:
 
 #if WITH_EDITOR
 	virtual const FSlateBrush* GetEditorIcon() override;
+	virtual const FText GetToolboxCategory() override;
 #endif
 
 protected:
diff --git a/Engine/Source/Runtime/UMG/Public/Components/NativeWidgetHost.h b/Engine/Source/Runtime/UMG/Public/Components/NativeWidgetHost.h
--- a/Engine/Source/Runtime/UMG/Public/Components/NativeWidgetHost.h
+++ b/Engine/Source/Runtime/UMG/Public/Components/NativeWidgetHost.h
@@ -8,7 +8,7 @@
  * A NativeWidgetHost is a container widget that can contain one child slate widget.  This should
  * be used when all you need is to nest a native widget inside a UMG widget.
  */
-UCLASS(meta=( Category="Advanced" ), ClassGroup=UserInterface)
+UCLASS(ClassGroup=UserInterface)
 class UMG_API UNativeWidgetHost : public UWidget
 {
 	GENERATED_UCLASS_BODY()
@@ -19,6 +19,7 @@ class UMG_API UNativeWidgetHost : public UWidget
 
 #if WITH_EDITOR
 	virtual const FSlateBrush* GetEditorIcon() override;
+	virtual const FText GetToolboxCategory() override;
 #endif
 
 protected:
diff --git a/Engine/Source/Runtime/UMG/Public/Components/Overlay.h b/Engine/Source/Runtime/UMG/Public/Components/Overlay.h
--- a/Engine/Source/Runtime/UMG/Public/Components/Overlay.h
+++ b/Engine/Source/Runtime/UMG/Public/Components/Overlay.h
@@ -4,7 +4,7 @@
 
 #include "Overlay.generated.h"
 
-UCLASS(meta=( Category="Panel" ), ClassGroup=UserInterface)
+UCLASS(ClassGroup=UserInterface)
 class UMG_API UOverlay : public UPanelWidget
 {
 	GENERATED_UCLASS_BODY()
@@ -13,6 +13,7 @@ class UMG_API UOverlay : public UPanelWidget
 
 #if WITH_EDITOR
 	virtual const FSlateBrush* GetEditorIcon() override;
+	virtual const FText GetToolboxCategory() override;
 #endif
 
 protected:
diff --git a/Engine/Source/Runtime/UMG/Public/Components/ProgressBar.h b/Engine/Source/Runtime/UMG/Public/Components/ProgressBar.h
--- a/Engine/Source/Runtime/UMG/Public/Components/ProgressBar.h
+++ b/Engine/Source/Runtime/UMG/Public/Components/ProgressBar.h
@@ -5,7 +5,7 @@
 #include "ProgressBar.generated.h"
 
 /** ProgressBar widget */
-UCLASS(meta=( Category="Common" ), ClassGroup=UserInterface)
+UCLASS(ClassGroup=UserInterface)
 class UMG_API UProgressBar : public UWidget
 {
 	GENERATED_UCLASS_BODY()
@@ -64,6 +64,7 @@ public:
 #if WITH_EDITOR
 	// UWidget interface
 	virtual const FSlateBrush* GetEditorIcon() override;
+	virtual const FText GetToolboxCategory() override;
 	// End UWidget interface
 #endif
 
diff --git a/Engine/Source/Runtime/UMG/Public/Components/ScaleBox.h b/Engine/Source/Runtime/UMG/Public/Components/ScaleBox.h
--- a/Engine/Source/Runtime/UMG/Public/Components/ScaleBox.h
+++ b/Engine/Source/Runtime/UMG/Public/Components/ScaleBox.h
@@ -9,7 +9,7 @@
  * you needed to have a background image scale to fill an area but not become distorted with different aspect ratios, or if you need
  * to auto fit some text to an area, this is the control for you.
  */
-UCLASS(meta=( Category="Panel" ), ClassGroup=UserInterface)
+UCLASS(ClassGroup=UserInterface)
 class UMG_API UScaleBox : public UContentWidget
 {
 	GENERATED_UCLASS_BODY()
@@ -34,6 +34,7 @@ public:
 
 #if WITH_EDITOR
 	virtual const FSlateBrush* GetEditorIcon() override;
+	virtual const FText GetToolboxCategory() override;
 #endif
 
 protected:
diff --git a/Engine/Source/Runtime/UMG/Public/Components/ScrollBar.h b/Engine/Source/Runtime/UMG/Public/Components/ScrollBar.h
--- a/Engine/Source/Runtime/UMG/Public/Components/ScrollBar.h
+++ b/Engine/Source/Runtime/UMG/Public/Components/ScrollBar.h
@@ -5,7 +5,7 @@
 #include "ScrollBar.generated.h"
 
 /** */
-UCLASS(meta=( Category="Advanced" ), ClassGroup=UserInterface, meta=( IsAdvanced = "True" ))
+UCLASS(ClassGroup=UserInterface, meta=( IsAdvanced = "True" ))
 class UMG_API UScrollBar : public UPanelWidget
 {
 	GENERATED_UCLASS_BODY()
@@ -64,6 +64,7 @@ public:
 #if WITH_EDITOR
 	// UWidget interface
 	virtual const FSlateBrush* GetEditorIcon() override;
+	virtual const FText GetToolboxCategory() override;
 	// End UWidget interface
 #endif
 
diff --git a/Engine/Source/Runtime/UMG/Public/Components/ScrollBox.h b/Engine/Source/Runtime/UMG/Public/Components/ScrollBox.h
--- a/Engine/Source/Runtime/UMG/Public/Components/ScrollBox.h
+++ b/Engine/Source/Runtime/UMG/Public/Components/ScrollBox.h
@@ -5,7 +5,7 @@
 #include "ScrollBox.generated.h"
 
 /** An arbitrary scrollable collection of widgets.  Great for presenting 10-100 widgets in a list.  Doesn't support virtualization. */
-UCLASS(meta=( Category="Panel" ), ClassGroup=UserInterface)
+UCLASS(ClassGroup=UserInterface)
 class UMG_API UScrollBox : public UPanelWidget
 {
 	GENERATED_UCLASS_BODY()
@@ -35,6 +35,7 @@ class UMG_API UScrollBox : public UPanelWidget
 #if WITH_EDITOR
 	// UWidget interface
 	virtual const FSlateBrush* GetEditorIcon() override;
+	virtual const FText GetToolboxCategory() override;
 	// End UWidget interface
 #endif
 
diff --git a/Engine/Source/Runtime/UMG/Public/Components/SizeBox.h b/Engine/Source/Runtime/UMG/Public/Components/SizeBox.h
--- a/Engine/Source/Runtime/UMG/Public/Components/SizeBox.h
+++ b/Engine/Source/Runtime/UMG/Public/Components/SizeBox.h
@@ -10,7 +10,7 @@ class USlateBrushAsset;
  * A SizeBox is a container widget that can contain one child widget, providing an opportunity 
  * to surround it with a SizeBox image and adjustable padding.
  */
-UCLASS(meta=( Category="Common" ), ClassGroup=UserInterface)
+UCLASS(ClassGroup=UserInterface)
 class UMG_API USizeBox : public UContentWidget
 {
 	GENERATED_UCLASS_BODY()
@@ -71,6 +71,7 @@ public:
 
 #if WITH_EDITOR
 	virtual const FSlateBrush* GetEditorIcon() override;
+	virtual const FText GetToolboxCategory() override;
 #endif
 
 protected:
diff --git a/Engine/Source/Runtime/UMG/Public/Components/Slider.h b/Engine/Source/Runtime/UMG/Public/Components/Slider.h
--- a/Engine/Source/Runtime/UMG/Public/Components/Slider.h
+++ b/Engine/Source/Runtime/UMG/Public/Components/Slider.h
@@ -9,7 +9,7 @@ DECLARE_DYNAMIC_MULTICAST_DELEGATE(FOnMouseCaptureEndEvent);
 DECLARE_DYNAMIC_MULTICAST_DELEGATE_OneParam(FOnFloatValueChangedEvent, float, Value);
 
 /** Slider widget */
-UCLASS(meta=( Category="Common" ), ClassGroup=UserInterface)
+UCLASS(ClassGroup=UserInterface)
 class UMG_API USlider : public UWidget
 {
 	GENERATED_UCLASS_BODY()
@@ -73,6 +73,7 @@ public:
 
 #if WITH_EDITOR
 	virtual const FSlateBrush* GetEditorIcon() override;
+	virtual const FText GetToolboxCategory() override;
 #endif
 
 protected:
diff --git a/Engine/Source/Runtime/UMG/Public/Components/Spacer.h b/Engine/Source/Runtime/UMG/Public/Components/Spacer.h
--- a/Engine/Source/Runtime/UMG/Public/Components/Spacer.h
+++ b/Engine/Source/Runtime/UMG/Public/Components/Spacer.h
@@ -5,7 +5,7 @@
 #include "Spacer.generated.h"
 
 /** A spacer widget; it does not have a visual representation, and just provides padding between other widgets. */
-UCLASS(meta=( Category="Primitive" ), ClassGroup=UserInterface)
+UCLASS(ClassGroup=UserInterface)
 class UMG_API USpacer : public UWidget
 {
 	GENERATED_UCLASS_BODY()
@@ -28,6 +28,7 @@ public:
 
 #if WITH_EDITOR
 	virtual const FSlateBrush* GetEditorIcon() override;
+	virtual const FText GetToolboxCategory() override;
 #endif
 
 protected:
diff --git a/Engine/Source/Runtime/UMG/Public/Components/SpinBox.h b/Engine/Source/Runtime/UMG/Public/Components/SpinBox.h
--- a/Engine/Source/Runtime/UMG/Public/Components/SpinBox.h
+++ b/Engine/Source/Runtime/UMG/Public/Components/SpinBox.h
@@ -5,7 +5,7 @@
 #include "SpinBox.generated.h"
 
 /** Spin box widget */
-UCLASS(meta = (Category = "Primitive"), ClassGroup = UserInterface)
+UCLASS(ClassGroup = UserInterface)
 class UMG_API USpinBox : public UWidget
 {
 	GENERATED_UCLASS_BODY()
@@ -134,6 +134,7 @@ public:
 
 #if WITH_EDITOR
 	virtual const FSlateBrush* GetEditorIcon() override;
+	virtual const FText GetToolboxCategory() override;
 #endif
 	// End of UWidget interface
 
diff --git a/Engine/Source/Runtime/UMG/Public/Components/TextBlock.h b/Engine/Source/Runtime/UMG/Public/Components/TextBlock.h
--- a/Engine/Source/Runtime/UMG/Public/Components/TextBlock.h
+++ b/Engine/Source/Runtime/UMG/Public/Components/TextBlock.h
@@ -5,7 +5,7 @@
 #include "TextBlock.generated.h"
 
 /** A simple static text widget */
-UCLASS(meta=( Category="Common" ), ClassGroup=UserInterface )
+UCLASS(ClassGroup=UserInterface )
 class UMG_API UTextBlock : public UWidget
 {
 	GENERATED_UCLASS_BODY()
@@ -99,6 +99,7 @@ public:
 #if WITH_EDITOR
 	// UWidget interface
 	virtual const FSlateBrush* GetEditorIcon() override;
+	virtual const FText GetToolboxCategory() override;
 	// End UWidget interface
 
 	virtual FString GetLabelMetadata() const override;
diff --git a/Engine/Source/Runtime/UMG/Public/Components/Throbber.h b/Engine/Source/Runtime/UMG/Public/Components/Throbber.h
--- a/Engine/Source/Runtime/UMG/Public/Components/Throbber.h
+++ b/Engine/Source/Runtime/UMG/Public/Components/Throbber.h
@@ -7,7 +7,7 @@
 class USlateBrushAsset;
 
 /** A Throbber widget that shows several zooming circles in a row. */
-UCLASS(meta=( Category="Misc" ), ClassGroup=UserInterface)
+UCLASS(ClassGroup=UserInterface)
 class UMG_API UThrobber : public UWidget
 {
 	GENERATED_UCLASS_BODY()
@@ -40,6 +40,7 @@ public:
 
 #if WITH_EDITOR
 	virtual const FSlateBrush* GetEditorIcon() override;
+	virtual const FText GetToolboxCategory() override;
 #endif
 
 protected:
diff --git a/Engine/Source/Runtime/UMG/Public/Components/TileView.h b/Engine/Source/Runtime/UMG/Public/Components/TileView.h
--- a/Engine/Source/Runtime/UMG/Public/Components/TileView.h
+++ b/Engine/Source/Runtime/UMG/Public/Components/TileView.h
@@ -5,7 +5,7 @@
 #include "TileView.generated.h"
 
 /** A flow panel that presents the contents as a set of tiles all uniformly sized. */
-UCLASS(meta=( Category="Misc" ), ClassGroup=UserInterface)
+UCLASS(ClassGroup=UserInterface)
 class UMG_API UTileView : public UTableViewBase
 {
 	GENERATED_UCLASS_BODY()
@@ -48,6 +48,7 @@ public:
 
 #if WITH_EDITOR
 	virtual const FSlateBrush* GetEditorIcon() override;
+	virtual const FText GetToolboxCategory() override;
 #endif
 
 protected:
diff --git a/Engine/Source/Runtime/UMG/Public/Components/UniformGridPanel.h b/Engine/Source/Runtime/UMG/Public/Components/UniformGridPanel.h
--- a/Engine/Source/Runtime/UMG/Public/Components/UniformGridPanel.h
+++ b/Engine/Source/Runtime/UMG/Public/Components/UniformGridPanel.h
@@ -5,7 +5,7 @@
 #include "UniformGridPanel.generated.h"
 
 /** A panel that evenly divides up available space between all of its children. */
-UCLASS(meta=( Category="Panel" ), ClassGroup=UserInterface)
+UCLASS(ClassGroup=UserInterface)
 class UMG_API UUniformGridPanel : public UPanelWidget
 {
 	GENERATED_UCLASS_BODY()
@@ -34,6 +34,7 @@ public:
 
 #if WITH_EDITOR
 	virtual const FSlateBrush* GetEditorIcon() override;
+	virtual const FText GetToolboxCategory() override;
 #endif
 
 protected:
diff --git a/Engine/Source/Runtime/UMG/Public/Components/VerticalBox.h b/Engine/Source/Runtime/UMG/Public/Components/VerticalBox.h
--- a/Engine/Source/Runtime/UMG/Public/Components/VerticalBox.h
+++ b/Engine/Source/Runtime/UMG/Public/Components/VerticalBox.h
@@ -8,7 +8,7 @@
  * A vertical box widget is a layout panel allowing child widgets to be automatically laid out
  * vertically.
  */
-UCLASS(meta=( Category="Panel" ), ClassGroup=UserInterface)
+UCLASS(ClassGroup=UserInterface)
 class UMG_API UVerticalBox : public UPanelWidget
 {
 	GENERATED_UCLASS_BODY()
@@ -22,6 +22,7 @@ public:
 #if WITH_EDITOR
 	// UWidget interface
 	virtual const FSlateBrush* GetEditorIcon() override;
+	virtual const FText GetToolboxCategory() override;
 	// End UWidget interface
 #endif
 
diff --git a/Engine/Source/Runtime/UMG/Public/Components/Viewport.h b/Engine/Source/Runtime/UMG/Public/Components/Viewport.h
--- a/Engine/Source/Runtime/UMG/Public/Components/Viewport.h
+++ b/Engine/Source/Runtime/UMG/Public/Components/Viewport.h
@@ -238,7 +238,7 @@ protected:
 /**
  * 
  */
-UCLASS(meta=( Category="Advanced" ), ClassGroup=UserInterface)
+UCLASS(ClassGroup=UserInterface)
 class UMG_API UViewport : public UContentWidget
 {
 	GENERATED_UCLASS_BODY()
@@ -272,6 +272,7 @@ class UMG_API UViewport : public UContentWidget
 
 #if WITH_EDITOR
 	virtual const FSlateBrush* GetEditorIcon() override;
+	virtual const FText GetToolboxCategory() override;
 #endif
 
 protected:
diff --git a/Engine/Source/Runtime/UMG/Public/Components/Widget.h b/Engine/Source/Runtime/UMG/Public/Components/Widget.h
--- a/Engine/Source/Runtime/UMG/Public/Components/Widget.h
+++ b/Engine/Source/Runtime/UMG/Public/Components/Widget.h
@@ -205,6 +205,9 @@ public:
 	/** Gets the label to display to the user for this widget. */
 	FString GetLabel() const;
 
+	/** Gets the toolbox category of the widget */
+	virtual const FText GetToolboxCategory();
+
 	/** Gets the editor icon */
 	virtual const FSlateBrush* GetEditorIcon();
 
diff --git a/Engine/Source/Runtime/UMG/Public/Components/WidgetSwitcher.h b/Engine/Source/Runtime/UMG/Public/Components/WidgetSwitcher.h
--- a/Engine/Source/Runtime/UMG/Public/Components/WidgetSwitcher.h
+++ b/Engine/Source/Runtime/UMG/Public/Components/WidgetSwitcher.h
@@ -8,7 +8,7 @@
  * A vertical box widget is a layout panel allowing child widgets to be automatically laid out
  * vertically.
  */
-UCLASS(meta=( Category="Panel" ), ClassGroup=UserInterface)
+UCLASS(ClassGroup=UserInterface)
 class UMG_API UWidgetSwitcher : public UPanelWidget
 {
 	GENERATED_UCLASS_BODY()
@@ -42,6 +42,7 @@ public:
 
 #if WITH_EDITOR
 	virtual const FSlateBrush* GetEditorIcon() override;
+	virtual const FText GetToolboxCategory() override;
 	virtual void OnDescendantSelected(UWidget* DescendantWidget) override;
 	virtual void OnDescendantDeselected(UWidget* DescendantWidget) override;
 #endif
