UMG - Fixing a bug with grid slots not initializing their slot with the defaults setup in UMG.

diff --git a/Engine/Source/Runtime/UMG/Private/Components/GridSlot.cpp b/Engine/Source/Runtime/UMG/Private/Components/GridSlot.cpp
--- a/Engine/Source/Runtime/UMG/Private/Components/GridSlot.cpp
+++ b/Engine/Source/Runtime/UMG/Private/Components/GridSlot.cpp
@@ -26,6 +26,13 @@ void UGridSlot::ReleaseNativeWidget()
 void UGridSlot::BuildSlot(TSharedRef<SGridPanel> GridPanel)
 {
 	Slot = &GridPanel->AddSlot(Column, Row)
+		.HAlign(HorizontalAlignment)
+		.VAlign(VerticalAlignment)
+		.Row(Row)
+		.RowSpan(RowSpan)
+		.Column(Row)
+		.ColumnSpan(RowSpan)
+		.Nudge(Nudge)
 		[
 			Content == NULL ? SNullWidget::NullWidget : Content->TakeWidget()
 		];
diff --git a/Engine/Source/Runtime/UMG/Private/Components/UniformGridSlot.cpp b/Engine/Source/Runtime/UMG/Private/Components/UniformGridSlot.cpp
--- a/Engine/Source/Runtime/UMG/Private/Components/UniformGridSlot.cpp
+++ b/Engine/Source/Runtime/UMG/Private/Components/UniformGridSlot.cpp
@@ -23,6 +23,8 @@ void UUniformGridSlot::ReleaseNativeWidget()
 void UUniformGridSlot::BuildSlot(TSharedRef<SUniformGridPanel> GridPanel)
 {
 	Slot = &GridPanel->AddSlot(Column, Row)
+		.HAlign(HorizontalAlignment)
+		.VAlign(VerticalAlignment)
 		[
 			Content == NULL ? SNullWidget::NullWidget : Content->TakeWidget()
 		];
