UMG - Adding utility functions to create a brush from material, textures, assets.  Also adding easier ways to set border content slot settings.

diff --git a/Engine/Source/Runtime/UMG/Private/Components/Border.cpp b/Engine/Source/Runtime/UMG/Private/Components/Border.cpp
--- a/Engine/Source/Runtime/UMG/Private/Components/Border.cpp
+++ b/Engine/Source/Runtime/UMG/Private/Components/Border.cpp
@@ -17,8 +17,12 @@ UBorder::UBorder(const FPostConstructInitializeProperties& PCIP)
 	ContentColorAndOpacity = FLinearColor::White;
 	BrushColor = FLinearColor::White;
 
-	SBorder::FArguments BorderDefaults;
-	bShowEffectWhenDisabled = BorderDefaults._ShowEffectWhenDisabled.Get();
+	Padding = FMargin(4, 2);
+
+	HorizontalAlignment = HAlign_Fill;
+	VerticalAlignment = VAlign_Fill;
+
+	bShowEffectWhenDisabled = true;
 }
 
 void UBorder::ReleaseNativeWidget()
@@ -101,6 +105,33 @@ void UBorder::SetContentColorAndOpacity(FLinearColor Color)
 	}
 }
 
+void UBorder::SetPadding(FMargin InPadding)
+{
+	Padding = InPadding;
+	if ( MyBorder.IsValid() )
+	{
+		MyBorder->SetPadding(InPadding);
+	}
+}
+
+void UBorder::SetHorizontalAlignment(EHorizontalAlignment InHorizontalAlignment)
+{
+	HorizontalAlignment = InHorizontalAlignment;
+	if ( MyBorder.IsValid() )
+	{
+		MyBorder->SetHAlign(InHorizontalAlignment);
+	}
+}
+
+void UBorder::SetVerticalAlignment(EVerticalAlignment InVerticalAlignment)
+{
+	VerticalAlignment = InVerticalAlignment;
+	if ( MyBorder.IsValid() )
+	{
+		MyBorder->SetVAlign(InVerticalAlignment);
+	}
+}
+
 void UBorder::SetBrushColor(FLinearColor Color)
 {
 	BrushColor = Color;
diff --git a/Engine/Source/Runtime/UMG/Private/Components/BorderSlot.cpp b/Engine/Source/Runtime/UMG/Private/Components/BorderSlot.cpp
--- a/Engine/Source/Runtime/UMG/Private/Components/BorderSlot.cpp
+++ b/Engine/Source/Runtime/UMG/Private/Components/BorderSlot.cpp
@@ -12,10 +12,6 @@ UBorderSlot::UBorderSlot(const FPostConstructInitializeProperties& PCIP)
 
 	HorizontalAlignment = HAlign_Fill;
 	VerticalAlignment = VAlign_Fill;
-
-	SBorder::FArguments BorderDefaults;
-
-	Padding = BorderDefaults._Padding.Get();
 }
 
 void UBorderSlot::ReleaseNativeWidget()
@@ -38,29 +34,17 @@ void UBorderSlot::BuildSlot(TSharedRef<SBorder> InBorder)
 
 void UBorderSlot::SetPadding(FMargin InPadding)
 {
-	Padding = InPadding;
-	if ( Border.IsValid() )
-	{
-		Border->SetPadding(InPadding);
-	}
+	CastChecked<UBorder>(Parent)->SetPadding(InPadding);
 }
 
 void UBorderSlot::SetHorizontalAlignment(EHorizontalAlignment InHorizontalAlignment)
 {
-	HorizontalAlignment = InHorizontalAlignment;
-	if ( Border.IsValid() )
-	{
-		Border->SetHAlign(InHorizontalAlignment);
-	}
+	CastChecked<UBorder>(Parent)->SetHorizontalAlignment(InHorizontalAlignment);
 }
 
 void UBorderSlot::SetVerticalAlignment(EVerticalAlignment InVerticalAlignment)
 {
-	VerticalAlignment = InVerticalAlignment;
-	if ( Border.IsValid() )
-	{
-		Border->SetVAlign(InVerticalAlignment);
-	}
+	CastChecked<UBorder>(Parent)->SetVerticalAlignment(InVerticalAlignment);
 }
 
 void UBorderSlot::SynchronizeProperties()
diff --git a/Engine/Source/Runtime/UMG/Private/WidgetBlueprintLibrary.cpp b/Engine/Source/Runtime/UMG/Private/WidgetBlueprintLibrary.cpp
--- a/Engine/Source/Runtime/UMG/Private/WidgetBlueprintLibrary.cpp
+++ b/Engine/Source/Runtime/UMG/Private/WidgetBlueprintLibrary.cpp
@@ -188,4 +188,42 @@ FEventReply UWidgetBlueprintLibrary::EndDragDrop(UPARAM(ref) FEventReply& Reply)
 	return Reply;
 }
 
+FSlateBrush UWidgetBlueprintLibrary::MakeBrushFromAsset(USlateBrushAsset* BrushAsset)
+{
+	if ( BrushAsset )
+	{
+		return BrushAsset->Brush;
+	}
+
+	return FSlateNoResource();
+}
+
+FSlateBrush UWidgetBlueprintLibrary::MakeBrushFromTexture(UTexture2D* Texture)
+{
+	if ( Texture )
+	{
+		FSlateBrush Brush;
+		Brush.SetResourceObject(Texture);
+		Brush.ImageSize = FVector2D(Texture->GetSizeX(), Texture->GetSizeY());
+
+		return Brush;
+	}
+	
+	return FSlateNoResource();
+}
+
+FSlateBrush UWidgetBlueprintLibrary::MakeBrushFromMaterial(UMaterialInterface* Material, int32 Width, int32 Height)
+{
+	if ( Material )
+	{
+		FSlateBrush Brush;
+		Brush.SetResourceObject(Material);
+		Brush.ImageSize = FVector2D(Width, Height);
+
+		return Brush;
+	}
+
+	return FSlateNoResource();
+}
+
 #undef LOCTEXT_NAMESPACE
\ No newline at end of file
diff --git a/Engine/Source/Runtime/UMG/Public/Blueprint/WidgetBlueprintLibrary.h b/Engine/Source/Runtime/UMG/Public/Blueprint/WidgetBlueprintLibrary.h
--- a/Engine/Source/Runtime/UMG/Public/Blueprint/WidgetBlueprintLibrary.h
+++ b/Engine/Source/Runtime/UMG/Public/Blueprint/WidgetBlueprintLibrary.h
@@ -83,4 +83,16 @@ class UWidgetBlueprintLibrary : public UBlueprintFunctionLibrary
 	 */
 	UFUNCTION(BlueprintPure, Category="Widget|Event Reply")
 	static FEventReply EndDragDrop(UPARAM(ref) FEventReply& Reply);
+
+	/**  */
+	UFUNCTION(BlueprintPure, Category="Widget|Brush")
+	static FSlateBrush MakeBrushFromAsset(USlateBrushAsset* BrushAsset);
+
+	/**  */
+	UFUNCTION(BlueprintPure, Category="Widget|Brush")
+	static FSlateBrush MakeBrushFromTexture(UTexture2D* Texture);
+
+	/**  */
+	UFUNCTION(BlueprintPure, Category="Widget|Brush")
+	static FSlateBrush MakeBrushFromMaterial(UMaterialInterface* Material, int32 Width, int32 Height);
 };
diff --git a/Engine/Source/Runtime/UMG/Public/Components/Border.h b/Engine/Source/Runtime/UMG/Public/Components/Border.h
--- a/Engine/Source/Runtime/UMG/Public/Components/Border.h
+++ b/Engine/Source/Runtime/UMG/Public/Components/Border.h
@@ -75,6 +75,15 @@ public:
 	UFUNCTION(BlueprintCallable, Category="Appearance")
 	void SetContentColorAndOpacity(FLinearColor InContentColorAndOpacity);
 
+	UFUNCTION(BlueprintCallable, Category="Appearance")
+	void SetPadding(FMargin InPadding);
+
+	UFUNCTION(BlueprintCallable, Category="Appearance")
+	void SetHorizontalAlignment(EHorizontalAlignment InHorizontalAlignment);
+
+	UFUNCTION(BlueprintCallable, Category="Appearance")
+	void SetVerticalAlignment(EVerticalAlignment InVerticalAlignment);
+
 	/**  */
 	UFUNCTION(BlueprintCallable, Category="Appearance")
 	void SetBrushColor(FLinearColor InBrushColor);
diff --git a/Engine/Source/Runtime/UMG/Public/Components/BorderSlot.h b/Engine/Source/Runtime/UMG/Public/Components/BorderSlot.h
--- a/Engine/Source/Runtime/UMG/Public/Components/BorderSlot.h
+++ b/Engine/Source/Runtime/UMG/Public/Components/BorderSlot.h
@@ -6,6 +6,8 @@
 
 #include "BorderSlot.generated.h"
 
+class UBorder;
+
 /** The Slot for the UBorderSlot, contains the widget displayed in a button's single slot */
 UCLASS()
 class UMG_API UBorderSlot : public UPanelSlot
@@ -13,7 +15,17 @@ class UMG_API UBorderSlot : public UPanelSlot
 	GENERATED_UCLASS_BODY()
 
 public:
-	
+
+	UFUNCTION(BlueprintCallable, Category="Appearance")
+	void SetPadding(FMargin InPadding);
+
+	UFUNCTION(BlueprintCallable, Category="Appearance")
+	void SetHorizontalAlignment(EHorizontalAlignment InHorizontalAlignment);
+
+	UFUNCTION(BlueprintCallable, Category="Appearance")
+	void SetVerticalAlignment(EVerticalAlignment InVerticalAlignment);
+
+protected:
 	/** The padding area between the slot and the content it contains. */
 	UPROPERTY(EditDefaultsOnly, Category=Layout)
 	FMargin Padding;
@@ -26,14 +38,7 @@ public:
 	UPROPERTY(EditDefaultsOnly, Category=Layout)
 	TEnumAsByte<EVerticalAlignment> VerticalAlignment;
 
-	UFUNCTION(BlueprintCallable, Category="Appearance")
-	void SetPadding(FMargin InPadding);
-
-	UFUNCTION(BlueprintCallable, Category="Appearance")
-	void SetHorizontalAlignment(EHorizontalAlignment InHorizontalAlignment);
-
-	UFUNCTION(BlueprintCallable, Category="Appearance")
-	void SetVerticalAlignment(EVerticalAlignment InVerticalAlignment);
+public:
 
 	// UPanelSlot interface
 	virtual void SynchronizeProperties() override;
@@ -58,4 +63,6 @@ private:
 
 	/** A pointer to the button to allow us to adjust the size, padding...etc at runtime. */
 	TSharedPtr<SBorder> Border;
+
+	friend UBorder;
 };
