#UE4

- Adding Width and Height parameters to MakeBrushFromTexture

diff --git a/Engine/Source/Runtime/UMG/Private/WidgetBlueprintLibrary.cpp b/Engine/Source/Runtime/UMG/Private/WidgetBlueprintLibrary.cpp
--- a/Engine/Source/Runtime/UMG/Private/WidgetBlueprintLibrary.cpp
+++ b/Engine/Source/Runtime/UMG/Private/WidgetBlueprintLibrary.cpp
@@ -198,14 +198,15 @@ FSlateBrush UWidgetBlueprintLibrary::MakeBrushFromAsset(USlateBrushAsset* BrushA
 	return FSlateNoResource();
 }
 
-FSlateBrush UWidgetBlueprintLibrary::MakeBrushFromTexture(UTexture2D* Texture)
+FSlateBrush UWidgetBlueprintLibrary::MakeBrushFromTexture(UTexture2D* Texture, int32 Width, int32 Height)
 {
 	if ( Texture )
 	{
 		FSlateBrush Brush;
 		Brush.SetResourceObject(Texture);
-		Brush.ImageSize = FVector2D(Texture->GetSizeX(), Texture->GetSizeY());
-
+		Width = (Width > 0) ? Width : Texture->GetSizeX();
+		Height = (Height > 0) ? Height : Texture->GetSizeY();
+		Brush.ImageSize = FVector2D(Width, Height);
 		return Brush;
 	}
 	
diff --git a/Engine/Source/Runtime/UMG/Public/Blueprint/WidgetBlueprintLibrary.h b/Engine/Source/Runtime/UMG/Public/Blueprint/WidgetBlueprintLibrary.h
--- a/Engine/Source/Runtime/UMG/Public/Blueprint/WidgetBlueprintLibrary.h
+++ b/Engine/Source/Runtime/UMG/Public/Blueprint/WidgetBlueprintLibrary.h
@@ -88,9 +88,14 @@ class UWidgetBlueprintLibrary : public UBlueprintFunctionLibrary
 	UFUNCTION(BlueprintPure, Category="Widget|Brush")
 	static FSlateBrush MakeBrushFromAsset(USlateBrushAsset* BrushAsset);
 
-	/**  */
+   /** 
+	* Creates a Slate Brush from a Texture2D
+	*
+	* @param Width  When less than or equal to zero, the Width of the brush will default to the Width of the Texture
+	* @param Height  When less than or equal to zero, the Height of the brush will default to the Height of the Texture
+	*/
 	UFUNCTION(BlueprintPure, Category="Widget|Brush")
-	static FSlateBrush MakeBrushFromTexture(UTexture2D* Texture);
+	static FSlateBrush MakeBrushFromTexture(UTexture2D* Texture, int32 Width = 0, int32 Height = 0);
 
 	/**  */
 	UFUNCTION(BlueprintPure, Category="Widget|Brush")
