- Fixed const-ness to make functions visible w/ "context sensitive" enabled for some UMG widgets

diff --git a/Engine/Source/Runtime/UMG/Private/Components/Slider.cpp b/Engine/Source/Runtime/UMG/Private/Components/Slider.cpp
--- a/Engine/Source/Runtime/UMG/Private/Components/Slider.cpp
+++ b/Engine/Source/Runtime/UMG/Private/Components/Slider.cpp
@@ -52,7 +52,7 @@ void USlider::HandleOnMouseCaptureEnd()
 	OnMouseCaptureEnd.Broadcast();
 }
 
-float USlider::GetValue()
+float USlider::GetValue() const
 {
 	if ( MySlider.IsValid() )
 	{
diff --git a/Engine/Source/Runtime/UMG/Private/Components/SpinBox.cpp b/Engine/Source/Runtime/UMG/Private/Components/SpinBox.cpp
--- a/Engine/Source/Runtime/UMG/Private/Components/SpinBox.cpp
+++ b/Engine/Source/Runtime/UMG/Private/Components/SpinBox.cpp
@@ -84,7 +84,7 @@ void USpinBox::SynchronizeProperties()
 	bOverride_MaxSliderValue ? SetMaxSliderValue(MaxSliderValue) : ClearMaxSliderValue();
 }
 
-float USpinBox::GetValue()
+float USpinBox::GetValue() const
 {
 	if (MySpinBox.IsValid())
 	{
@@ -104,7 +104,7 @@ void USpinBox::SetValue(float InValue)
 }
 
 // MIN VALUE
-float USpinBox::GetMinValue()
+float USpinBox::GetMinValue() const
 {
 	float ReturnVal = TNumericLimits<float>::Lowest();
 
@@ -140,7 +140,7 @@ void USpinBox::ClearMinValue()
 }
 
 // MAX VALUE
-float USpinBox::GetMaxValue()
+float USpinBox::GetMaxValue() const
 {
 	float ReturnVal = TNumericLimits<float>::Max();
 
@@ -175,7 +175,7 @@ void USpinBox::ClearMaxValue()
 }
 
 // MIN SLIDER VALUE
-float USpinBox::GetMinSliderValue()
+float USpinBox::GetMinSliderValue() const
 {
 	float ReturnVal = TNumericLimits<float>::Min();
 
@@ -211,7 +211,7 @@ void USpinBox::ClearMinSliderValue()
 }
 
 // MAX SLIDER VALUE
-float USpinBox::GetMaxSliderValue()
+float USpinBox::GetMaxSliderValue() const
 {
 	float ReturnVal = TNumericLimits<float>::Max();
 
diff --git a/Engine/Source/Runtime/UMG/Private/UserWidget.cpp b/Engine/Source/Runtime/UMG/Private/UserWidget.cpp
--- a/Engine/Source/Runtime/UMG/Private/UserWidget.cpp
+++ b/Engine/Source/Runtime/UMG/Private/UserWidget.cpp
@@ -619,12 +619,12 @@ void UUserWidget::RemoveFromViewport()
 	}
 }
 
-bool UUserWidget::GetIsVisible()
+bool UUserWidget::GetIsVisible() const
 {
 	return FullScreenWidget.IsValid();
 }
 
-TEnumAsByte<ESlateVisibility::Type> UUserWidget::GetVisiblity()
+TEnumAsByte<ESlateVisibility::Type> UUserWidget::GetVisiblity() const
 {
 	if ( FullScreenWidget.IsValid() )
 	{
diff --git a/Engine/Source/Runtime/UMG/Public/Blueprint/UserWidget.h b/Engine/Source/Runtime/UMG/Public/Blueprint/UserWidget.h
--- a/Engine/Source/Runtime/UMG/Public/Blueprint/UserWidget.h
+++ b/Engine/Source/Runtime/UMG/Public/Blueprint/UserWidget.h
@@ -127,10 +127,10 @@ public:
 
 	/*  */
 	UFUNCTION(BlueprintPure, Category="Appearance")
-	bool GetIsVisible();
+	bool GetIsVisible() const;
 
 	UFUNCTION(BlueprintPure, Category="Appearance")
-	TEnumAsByte<ESlateVisibility::Type> GetVisiblity();
+	TEnumAsByte<ESlateVisibility::Type> GetVisiblity() const;
 
 	/** Sets the player context associated with this UI. */
 	void SetPlayerContext(FLocalPlayerContext InPlayerContext);
diff --git a/Engine/Source/Runtime/UMG/Public/Components/Slider.h b/Engine/Source/Runtime/UMG/Public/Components/Slider.h
--- a/Engine/Source/Runtime/UMG/Public/Components/Slider.h
+++ b/Engine/Source/Runtime/UMG/Public/Components/Slider.h
@@ -60,8 +60,8 @@ public:
 	FOnFloatValueChangedEvent OnValueChanged;
 
 	/** Gets the current value of the slider. */
-	UFUNCTION(BlueprintPure, Category="Behavior")
-	float GetValue();
+	UFUNCTION(BlueprintCallable, Category = "Behavior")
+	float GetValue() const;
 
 	/** Sets the current value of the slider. */
 	UFUNCTION(BlueprintCallable, Category="Behavior")
diff --git a/Engine/Source/Runtime/UMG/Public/Components/SpinBox.h b/Engine/Source/Runtime/UMG/Public/Components/SpinBox.h
--- a/Engine/Source/Runtime/UMG/Public/Components/SpinBox.h
+++ b/Engine/Source/Runtime/UMG/Public/Components/SpinBox.h
@@ -66,8 +66,8 @@ public:
 	FOnSpinBoxValueChangedEvent OnEndSliderMovement;
 
 	/** Get the current value of the spin box. */
-	UFUNCTION(BlueprintPure, Category = "Behavior")
-	float GetValue();
+	UFUNCTION(BlueprintCallable, Category = "Behavior")
+	float GetValue() const;
 
 	/** Set the value of the spin box. */
 	UFUNCTION(BlueprintCallable, Category = "Behavior")
@@ -75,8 +75,8 @@ public:
 
 	// MIN VALUE
 	/** Get the current minimum value that can be manually set in the spin box. */
-	UFUNCTION(BlueprintPure, Category = "Behavior")
-	float GetMinValue();
+	UFUNCTION(BlueprintCallable, Category = "Behavior")
+	float GetMinValue() const;
 
 	/** Set the minimum value that can be manually set in the spin box. */
 	UFUNCTION(BlueprintCallable, Category = "Behavior")
@@ -88,8 +88,8 @@ public:
 
 	// MAX VALUE
 	/** Get the current maximum value that can be manually set in the spin box. */
-	UFUNCTION(BlueprintPure, Category = "Behavior")
-	float GetMaxValue();
+	UFUNCTION(BlueprintCallable, Category = "Behavior")
+	float GetMaxValue() const;
 
 	/** Set the maximum value that can be manually set in the spin box. */
 	UFUNCTION(BlueprintCallable, Category = "Behavior")
@@ -101,8 +101,8 @@ public:
 
 	// MIN SLIDER VALUE
 	/** Get the current minimum value that can be specified using the slider. */
-	UFUNCTION(BlueprintPure, Category = "Behavior")
-	float GetMinSliderValue();
+	UFUNCTION(BlueprintCallable, Category = "Behavior")
+	float GetMinSliderValue() const;
 
 	/** Set the minimum value that can be specified using the slider. */
 	UFUNCTION(BlueprintCallable, Category = "Behavior")
@@ -114,8 +114,8 @@ public:
 
 	// MAX SLIDER VALUE
 	/** Get the current maximum value that can be specified using the slider. */
-	UFUNCTION(BlueprintPure, Category = "Behavior")
-	float GetMaxSliderValue();
+	UFUNCTION(BlueprintCallable, Category = "Behavior")
+	float GetMaxSliderValue() const;
 
 	/** Set the maximum value that can be specified using the slider. */
 	UFUNCTION(BlueprintCallable, Category = "Behavior")
