-Added GetText() to UMG text block

diff --git a/Engine/Source/Runtime/UMG/Private/Components/TextBlock.cpp b/Engine/Source/Runtime/UMG/Private/Components/TextBlock.cpp
--- a/Engine/Source/Runtime/UMG/Private/Components/TextBlock.cpp
+++ b/Engine/Source/Runtime/UMG/Private/Components/TextBlock.cpp
@@ -111,6 +111,16 @@ void UTextBlock::SynchronizeProperties()
 #endif
 }
 
+FText UTextBlock::GetText() const
+{
+	if (MyTextBlock.IsValid())
+	{
+		return MyTextBlock->GetText();
+	}
+
+	return Text;
+}
+
 void UTextBlock::SetText(FText InText)
 {
 	Text = InText;
diff --git a/Engine/Source/Runtime/UMG/Public/Components/EditableText.h b/Engine/Source/Runtime/UMG/Public/Components/EditableText.h
--- a/Engine/Source/Runtime/UMG/Public/Components/EditableText.h
+++ b/Engine/Source/Runtime/UMG/Public/Components/EditableText.h
@@ -103,18 +103,20 @@ public:
 	UPROPERTY(BlueprintAssignable, Category="Widget Event")
 	FOnEditableTextCommittedEvent OnTextCommitted;
 
-	/**  */
-	UFUNCTION(BlueprintCallable, Category="Widget")
-	FText GetText() const;
-
-	/** 
-	* Directly sets the text.
-	* *Warning*: This will wipe any binding created for the Text property!
-	*
-	* @param InText The text to assign to the widget
+	/**
+	Gets the widget text
+	@return The widget text
 	*/
-	UFUNCTION(BlueprintCallable, Category="Widget")
-	void SetText(FText InText);
+	UFUNCTION(BlueprintCallable, Category = "Widget")
+		FText GetText() const;
+
+	/**
+	Directly sets the widget text.
+	Warning: This will wipe any binding created for the Text property!
+	@param InText The text to assign to the widget
+	*/
+	UFUNCTION(BlueprintCallable, Category = "Widget")
+		void SetText(FText InText);
 	
 	virtual void ReleaseNativeWidget() override;
 
diff --git a/Engine/Source/Runtime/UMG/Public/Components/TextBlock.h b/Engine/Source/Runtime/UMG/Public/Components/TextBlock.h
--- a/Engine/Source/Runtime/UMG/Public/Components/TextBlock.h
+++ b/Engine/Source/Runtime/UMG/Public/Components/TextBlock.h
@@ -88,7 +88,18 @@ public:
 	///** Called when this text is double clicked */
 	//SLATE_EVENT(FOnClicked, OnDoubleClicked)
 
-	/**  */
+	/** 
+	Gets the widget text
+	@return The widget text
+	*/
+	UFUNCTION(BlueprintCallable, Category = "Widget")
+	FText GetText() const;
+
+	/**
+	Directly sets the widget text.
+	Warning: This will wipe any binding created for the Text property!
+	@param InText The text to assign to the widget
+	*/
 	UFUNCTION(BlueprintCallable, Category="Widget")
 	void SetText(FText InText);
 
