UMG - Adding clearer ways to set size and position on the canvas panel slot object.  Also changing several struct bits for anchor data to make them blueprint safe.

diff --git a/Engine/Source/Runtime/UMG/Private/Components/CanvasPanelSlot.cpp b/Engine/Source/Runtime/UMG/Private/Components/CanvasPanelSlot.cpp
--- a/Engine/Source/Runtime/UMG/Private/Components/CanvasPanelSlot.cpp
+++ b/Engine/Source/Runtime/UMG/Private/Components/CanvasPanelSlot.cpp
@@ -123,6 +123,28 @@ bool UCanvasPanelSlot::CanResize(const FVector2D& Direction) const
 	return true;
 }
 
+void UCanvasPanelSlot::SetPosition(FVector2D InPosition)
+{
+	LayoutData.Offsets.Left = InPosition.X;
+	LayoutData.Offsets.Top = InPosition.Y;
+
+	if ( Slot )
+	{
+		Slot->Offset(LayoutData.Offsets);
+	}
+}
+
+void UCanvasPanelSlot::SetSize(FVector2D InSize)
+{
+	LayoutData.Offsets.Right = InSize.X;
+	LayoutData.Offsets.Bottom = InSize.Y;
+
+	if ( Slot )
+	{
+		Slot->Offset(LayoutData.Offsets);
+	}
+}
+
 void UCanvasPanelSlot::SetOffsets(FMargin InOffset)
 {
 	LayoutData.Offsets = InOffset;
diff --git a/Engine/Source/Runtime/UMG/Public/Components/CanvasPanelSlot.h b/Engine/Source/Runtime/UMG/Public/Components/CanvasPanelSlot.h
--- a/Engine/Source/Runtime/UMG/Public/Components/CanvasPanelSlot.h
+++ b/Engine/Source/Runtime/UMG/Public/Components/CanvasPanelSlot.h
@@ -6,7 +6,7 @@
 
 #include "CanvasPanelSlot.generated.h"
 
-USTRUCT()
+USTRUCT(BlueprintType)
 struct FAnchorData
 {
 public:
@@ -15,15 +15,15 @@ public:
 public:
 
 	/** Offset. */
-	UPROPERTY(EditAnywhere, Category=AnchorData)
+	UPROPERTY(EditAnywhere, BlueprintReadWrite, Category=AnchorData)
 	FMargin Offsets;
 	
 	/** Anchors. */
-	UPROPERTY(EditAnywhere, Category=AnchorData)
+	UPROPERTY(EditAnywhere, BlueprintReadWrite, Category=AnchorData)
 	FAnchors Anchors;
 
 	/** Alignment. */
-	UPROPERTY(EditAnywhere, Category=AnchorData)
+	UPROPERTY(EditAnywhere, BlueprintReadWrite, Category=AnchorData)
 	FVector2D Alignment;
 };
 
@@ -33,11 +33,11 @@ class UMG_API UCanvasPanelSlot : public UPanelSlot
 	GENERATED_UCLASS_BODY()
 
 	/** The anchoring information for the slot */
-	UPROPERTY(EditDefaultsOnly, Category=Appearance)
+	UPROPERTY(EditDefaultsOnly, Category="Layout (Canvas Slot)")
 	FAnchorData LayoutData;
 
 	/** The order priority this widget is rendered in.  Higher values are rendered last (and so they will appear to be on top). */
-	UPROPERTY(EditDefaultsOnly, Category=AnchorData)
+	UPROPERTY(EditDefaultsOnly, Category="Layout (Canvas Slot)")
 	int32 ZOrder;
 
 	void BuildSlot(TSharedRef<SConstraintCanvas> Canvas);
@@ -51,19 +51,27 @@ class UMG_API UCanvasPanelSlot : public UPanelSlot
 	virtual bool CanResize(const FVector2D& Direction) const override;
 
 	/** Sets the position of the slot */
-	UFUNCTION(BlueprintCallable, Category="Appearance")
+	UFUNCTION(BlueprintCallable, Category="Layout (Canvas Slot)")
+	void SetPosition(FVector2D InPosition);
+
+	/** Sets the size of the slot */
+	UFUNCTION(BlueprintCallable, Category="Layout (Canvas Slot)")
+	void SetSize(FVector2D InSize);
+
+	/** Sets the position of the slot */
+	UFUNCTION(BlueprintCallable, Category="Layout (Canvas Slot)")
 	void SetOffsets(FMargin InOffset);
 	
 	/** Sets the anchors on the slot */
-	UFUNCTION(BlueprintCallable, Category="Appearance")
+	UFUNCTION(BlueprintCallable, Category="Layout (Canvas Slot)")
 	void SetAnchors(FAnchors InAnchors);
 
 	/** Sets the alignment on the slot */
-	UFUNCTION(BlueprintCallable, Category="Appearance")
+	UFUNCTION(BlueprintCallable, Category="Layout (Canvas Slot)")
 	void SetAlignment(FVector2D InAlignment);
 
 	/** Sets the z-order on the slot */
-	UFUNCTION(BlueprintCallable, Category="Appearance")
+	UFUNCTION(BlueprintCallable, Category="Layout (Canvas Slot)")
 	void SetZOrder(int32 InZOrder);
 
 	// UPanelSlot interface
