UMG - Changing Textblock back to common.

diff --git a/Engine/Source/Runtime/UMG/Private/Components/TextBlock.cpp b/Engine/Source/Runtime/UMG/Private/Components/TextBlock.cpp
--- a/Engine/Source/Runtime/UMG/Private/Components/TextBlock.cpp
+++ b/Engine/Source/Runtime/UMG/Private/Components/TextBlock.cpp
@@ -160,7 +160,7 @@ const FSlateBrush* UTextBlock::GetEditorIcon()
 
 const FText UTextBlock::GetToolboxCategory()
 {
-	return LOCTEXT("Primitive", "Primitive");
+	return LOCTEXT("Common", "Common");
 }
 
 #endif
