UMG - Marking several unfinished widgets as experimental so that they will be hidden in the palette.

diff --git a/Engine/Source/Runtime/UMG/Public/Components/ComboBox.h b/Engine/Source/Runtime/UMG/Public/Components/ComboBox.h
--- a/Engine/Source/Runtime/UMG/Public/Components/ComboBox.h
+++ b/Engine/Source/Runtime/UMG/Public/Components/ComboBox.h
@@ -7,7 +7,7 @@
 /**
  * The combobox allows you to display a list of options to the user in a dropdown menu for them to select one.
  */
-UCLASS( meta=( DisplayName="ComboBox (Object)"), ClassGroup=UserInterface)
+UCLASS(Experimental, ClassGroup=UserInterface, meta=( DisplayName="ComboBox (Object)" ))
 class UMG_API UComboBox : public UWidget
 {
 	GENERATED_UCLASS_BODY()
diff --git a/Engine/Source/Runtime/UMG/Public/Components/ListView.h b/Engine/Source/Runtime/UMG/Public/Components/ListView.h
--- a/Engine/Source/Runtime/UMG/Public/Components/ListView.h
+++ b/Engine/Source/Runtime/UMG/Public/Components/ListView.h
@@ -5,7 +5,7 @@
 #include "ListView.generated.h"
 
 /** Allows thousands of items to be displayed in a list.  Generates widgets dynamically for each item. */
-UCLASS(ClassGroup=UserInterface)
+UCLASS(Experimental, ClassGroup=UserInterface)
 class UMG_API UListView : public UTableViewBase
 {
 	GENERATED_UCLASS_BODY()
diff --git a/Engine/Source/Runtime/UMG/Public/Components/MenuAnchor.h b/Engine/Source/Runtime/UMG/Public/Components/MenuAnchor.h
--- a/Engine/Source/Runtime/UMG/Public/Components/MenuAnchor.h
+++ b/Engine/Source/Runtime/UMG/Public/Components/MenuAnchor.h
@@ -8,7 +8,7 @@
  * The Menu Anchor allows you to specify an location that a popup menu should be anchored to, and should be
  * summoned from.
  */
-UCLASS(ClassGroup=UserInterface, meta=( IsAdvanced = "True" ))
+UCLASS(Experimental)
 class UMG_API UMenuAnchor : public UContentWidget
 {
 	GENERATED_UCLASS_BODY()
diff --git a/Engine/Source/Runtime/UMG/Public/Components/TileView.h b/Engine/Source/Runtime/UMG/Public/Components/TileView.h
--- a/Engine/Source/Runtime/UMG/Public/Components/TileView.h
+++ b/Engine/Source/Runtime/UMG/Public/Components/TileView.h
@@ -5,7 +5,7 @@
 #include "TileView.generated.h"
 
 /** A flow panel that presents the contents as a set of tiles all uniformly sized. */
-UCLASS(ClassGroup=UserInterface)
+UCLASS(Experimental, ClassGroup=UserInterface)
 class UMG_API UTileView : public UTableViewBase
 {
 	GENERATED_UCLASS_BODY()
