UMG - DragDropOperation is now a blueprintable type.

diff --git a/Engine/Source/Runtime/UMG/Public/Blueprint/DragDropOperation.h b/Engine/Source/Runtime/UMG/Public/Blueprint/DragDropOperation.h
--- a/Engine/Source/Runtime/UMG/Public/Blueprint/DragDropOperation.h
+++ b/Engine/Source/Runtime/UMG/Public/Blueprint/DragDropOperation.h
@@ -10,7 +10,7 @@ DECLARE_DYNAMIC_MULTICAST_DELEGATE_OneParam(FOnDragDropMulticast, UDragDropOpera
 /**
  * This class is the base drag drop operation for UMG, extend it to add additional data and add new functionality.
  */
-UCLASS()
+UCLASS(BlueprintType)
 class UMG_API UDragDropOperation : public UObject
 {
 	GENERATED_UCLASS_BODY()
