UMG - Continuing to work on styles, exposing more logic to blueprints, and making it easier to get slots for different types of widgets, starting with a few.  Will add more in the future.

diff --git a/Engine/Source/Editor/UMGEditor/Private/Customizations/CanvasSlotCustomization.cpp b/Engine/Source/Editor/UMGEditor/Private/Customizations/CanvasSlotCustomization.cpp
--- a/Engine/Source/Editor/UMGEditor/Private/Customizations/CanvasSlotCustomization.cpp
+++ b/Engine/Source/Editor/UMGEditor/Private/Customizations/CanvasSlotCustomization.cpp
@@ -8,6 +8,8 @@
 #include "ScopedTransaction.h"
 #include "BlueprintEditorUtils.h"
 
+#include "CanvasSlotCustomization.h"
+
 #define LOCTEXT_NAMESPACE "UMG"
 
 // FCanvasSlotCustomization
diff --git a/Engine/Source/Editor/UMGEditor/Private/Customizations/DetailCustomizations.cpp b/Engine/Source/Editor/UMGEditor/Private/Customizations/DetailCustomizations.cpp
--- a/Engine/Source/Editor/UMGEditor/Private/Customizations/DetailCustomizations.cpp
+++ b/Engine/Source/Editor/UMGEditor/Private/Customizations/DetailCustomizations.cpp
@@ -318,6 +318,11 @@ FReply FBlueprintWidgetCustomization::AddOrViewEventBinding(TSharedPtr<FEdGraphS
 }
 
 void FBlueprintWidgetCustomization::CustomizeDetails( IDetailLayoutBuilder& DetailLayout )
+{
+	PerformBindingCustomization(DetailLayout);
+}
+
+void FBlueprintWidgetCustomization::PerformBindingCustomization(IDetailLayoutBuilder& DetailLayout)
 {
 	TArray< TWeakObjectPtr<UObject> > OutObjects;
 	DetailLayout.GetObjectsBeingCustomized(OutObjects);
@@ -333,13 +338,13 @@ void FBlueprintWidgetCustomization::CustomizeDetails( IDetailLayoutBuilder& Deta
 
 			if ( UDelegateProperty* DelegateProperty = Cast<UDelegateProperty>(*PropertyIt) )
 			{
-				if( DelegateProperty->GetName().EndsWith(TEXT("Delegate")) )
+				if ( DelegateProperty->GetName().EndsWith(TEXT("Delegate")) )
 				{
-					CreateDelegateCustomization( DetailLayout, DelegateProperty, Widget );
+					CreateDelegateCustomization(DetailLayout, DelegateProperty, Widget);
 				}
-				else if( DelegateProperty->GetName().EndsWith(TEXT("Event")) )
+				else if ( DelegateProperty->GetName().EndsWith(TEXT("Event")) )
 				{
-					CreateEventCustomization( DetailLayout, DelegateProperty, Widget );
+					CreateEventCustomization(DetailLayout, DelegateProperty, Widget);
 				}
 			}
 			else if ( UMulticastDelegateProperty* DelegateProperty = Cast<UMulticastDelegateProperty>(Property) )
diff --git a/Engine/Source/Editor/UMGEditor/Private/Customizations/DetailCustomizations.h b/Engine/Source/Editor/UMGEditor/Private/Customizations/DetailCustomizations.h
--- a/Engine/Source/Editor/UMGEditor/Private/Customizations/DetailCustomizations.h
+++ b/Engine/Source/Editor/UMGEditor/Private/Customizations/DetailCustomizations.h
@@ -2,6 +2,9 @@
 
 #pragma once
 
+/**
+ * Provides the customization for all UWidgets.  Bindings, style disabling...etc.
+ */
 class FBlueprintWidgetCustomization : public IDetailCustomization
 {
 public:
@@ -30,6 +33,8 @@ private:
 		UEdGraph* EdGraph;
 	};
 
+	void PerformBindingCustomization(IDetailLayoutBuilder& DetailLayout);
+
 	void RefreshBlueprintMemberCache(const UFunction* DelegateSignature);
 
 	TSharedRef<SWidget> OnGenerateDelegateMenu(UWidget* Widget, TSharedRef<IPropertyHandle> PropertyHandle, UFunction* DelegateSignature, bool bIsPure);
diff --git a/Engine/Source/Editor/UMGEditor/Private/Customizations/HorizontalAlignmentCustomization.cpp b/Engine/Source/Editor/UMGEditor/Private/Customizations/HorizontalAlignmentCustomization.cpp
--- a/Engine/Source/Editor/UMGEditor/Private/Customizations/HorizontalAlignmentCustomization.cpp
+++ b/Engine/Source/Editor/UMGEditor/Private/Customizations/HorizontalAlignmentCustomization.cpp
@@ -8,6 +8,8 @@
 #include "ScopedTransaction.h"
 #include "BlueprintEditorUtils.h"
 
+#include "HorizontalAlignmentCustomization.h"
+
 #define LOCTEXT_NAMESPACE "UMG"
 
 void FHorizontalAlignmentCustomization::CustomizeHeader(TSharedRef<IPropertyHandle> PropertyHandle, FDetailWidgetRow& HeaderRow, IPropertyTypeCustomizationUtils& CustomizationUtils)
diff --git a/Engine/Source/Editor/UMGEditor/Private/Customizations/SlateChildSizeCustomization.cpp b/Engine/Source/Editor/UMGEditor/Private/Customizations/SlateChildSizeCustomization.cpp
--- a/Engine/Source/Editor/UMGEditor/Private/Customizations/SlateChildSizeCustomization.cpp
+++ b/Engine/Source/Editor/UMGEditor/Private/Customizations/SlateChildSizeCustomization.cpp
@@ -8,6 +8,8 @@
 #include "ScopedTransaction.h"
 #include "BlueprintEditorUtils.h"
 
+#include "SlateChildSizeCustomization.h"
+
 #define LOCTEXT_NAMESPACE "UMG"
 
 void FSlateChildSizeCustomization::CustomizeHeader(TSharedRef<IPropertyHandle> PropertyHandle, FDetailWidgetRow& HeaderRow, IPropertyTypeCustomizationUtils& CustomizationUtils)
diff --git a/Engine/Source/Editor/UMGEditor/Private/Customizations/VerticalAlignmentCustomization.cpp b/Engine/Source/Editor/UMGEditor/Private/Customizations/VerticalAlignmentCustomization.cpp
--- a/Engine/Source/Editor/UMGEditor/Private/Customizations/VerticalAlignmentCustomization.cpp
+++ b/Engine/Source/Editor/UMGEditor/Private/Customizations/VerticalAlignmentCustomization.cpp
@@ -8,6 +8,8 @@
 #include "ScopedTransaction.h"
 #include "BlueprintEditorUtils.h"
 
+#include "VerticalAlignmentCustomization.h"
+
 #define LOCTEXT_NAMESPACE "UMG"
 
 void FVerticalAlignmentCustomization::CustomizeHeader(TSharedRef<IPropertyHandle> PropertyHandle, FDetailWidgetRow& HeaderRow, IPropertyTypeCustomizationUtils& CustomizationUtils)
diff --git a/Engine/Source/Editor/UMGEditor/Private/Details/SWidgetDetailsView.cpp b/Engine/Source/Editor/UMGEditor/Private/Details/SWidgetDetailsView.cpp
--- a/Engine/Source/Editor/UMGEditor/Private/Details/SWidgetDetailsView.cpp
+++ b/Engine/Source/Editor/UMGEditor/Private/Details/SWidgetDetailsView.cpp
@@ -9,6 +9,11 @@
 #include "ISequencer.h"
 #include "Animation/UMGDetailKeyframeHandler.h"
 
+#include "CanvasSlotCustomization.h"
+#include "HorizontalAlignmentCustomization.h"
+#include "VerticalAlignmentCustomization.h"
+#include "SlateChildSizeCustomization.h"
+
 #define LOCTEXT_NAMESPACE "UMG"
 
 void SWidgetDetailsView::Construct(const FArguments& InArgs, TSharedPtr<FWidgetBlueprintEditor> InBlueprintEditor)
@@ -105,6 +110,7 @@ SWidgetDetailsView::~SWidgetDetailsView()
 void SWidgetDetailsView::RegisterCustomizations()
 {
 	PropertyView->RegisterInstancedCustomPropertyLayout(UWidget::StaticClass(), FOnGetDetailCustomizationInstance::CreateStatic(&FBlueprintWidgetCustomization::MakeInstance, BlueprintEditor.Pin().ToSharedRef(), BlueprintEditor.Pin()->GetBlueprintObj()));
+
 	//PropertyView->RegisterInstancedCustomPropertyLayout(FSlateBrush::StaticStruct()->StaticClass(), FOnGetDetailCustomizationInstance::CreateStatic(&FSlateBrushStructCustomization::MakeInstance));
 
 	static FName PropertyEditor("PropertyEditor");
diff --git a/Engine/Source/Editor/UMGEditor/Private/UMGEditorPrivatePCH.h b/Engine/Source/Editor/UMGEditor/Private/UMGEditorPrivatePCH.h
--- a/Engine/Source/Editor/UMGEditor/Private/UMGEditorPrivatePCH.h
+++ b/Engine/Source/Editor/UMGEditor/Private/UMGEditorPrivatePCH.h
@@ -41,10 +41,6 @@
 #include "GridSlotExtension.h"
 
 #include "DetailCustomizations.h"
-#include "CanvasSlotCustomization.h"
-#include "HorizontalAlignmentCustomization.h"
-#include "VerticalAlignmentCustomization.h"
-#include "SlateChildSizeCustomization.h"
 
 #include "WidgetTemplate.h"
 #include "WidgetTemplateDragDropOp.h"
diff --git a/Engine/Source/Runtime/SlateCore/Public/Types/SlateEnums.h b/Engine/Source/Runtime/SlateCore/Public/Types/SlateEnums.h
--- a/Engine/Source/Runtime/SlateCore/Public/Types/SlateEnums.h
+++ b/Engine/Source/Runtime/SlateCore/Public/Types/SlateEnums.h
@@ -5,7 +5,7 @@
 #include "SlateEnums.generated.h"
 
 
-UENUM()
+UENUM(BlueprintType)
 namespace EButtonClickMethod
 {
 	/**
@@ -33,7 +33,7 @@ namespace EButtonClickMethod
 }
 
 
-UENUM()
+UENUM(BlueprintType)
 namespace EButtonTouchMethod
 {
 	/** Ways in which touch interactions trigger a "Clicked" event. */
@@ -51,7 +51,7 @@ namespace EButtonTouchMethod
 }
 
 
-UENUM()
+UENUM(BlueprintType)
 namespace EFocusMoveDirection
 {
 	/**
@@ -71,7 +71,7 @@ namespace EFocusMoveDirection
 /**
  * Enumerates horizontal alignment options, i.e. for widget slots.
  */
-UENUM()
+UENUM(BlueprintType)
 enum EHorizontalAlignment
 {
 	/** Fill the entire width. */
@@ -91,7 +91,7 @@ enum EHorizontalAlignment
 /**
  * Enumerates vertical alignment options, i.e. for widget slots.
  */
-UENUM()
+UENUM(BlueprintType)
 enum EVerticalAlignment
 {
 	/** Fill the entire height. */
@@ -111,7 +111,7 @@ enum EVerticalAlignment
 /**
  * Enumerates possible placements for pop-up menus.
  */
-UENUM()
+UENUM(BlueprintType)
 enum EMenuPlacement
 {
 	/** Place the menu immediately below the anchor */
@@ -134,7 +134,7 @@ enum EMenuPlacement
 /**
  * Enumerates widget orientations.
  */
-UENUM()
+UENUM(BlueprintType)
 enum EOrientation
 {
 	/** Orient horizontally, i.e. left to right. */
@@ -148,7 +148,7 @@ enum EOrientation
 /**
  * Enumerates scroll directions.
  */
-UENUM()
+UENUM(BlueprintType)
 enum EScrollDirection
 {
 	/** Scroll down. */
@@ -162,7 +162,7 @@ enum EScrollDirection
 /**
  * Additional information about a text committal
  */
-UENUM()
+UENUM(BlueprintType)
 namespace ETextCommit
 {
 	enum Type
@@ -182,7 +182,7 @@ namespace ETextCommit
 /**
  * Additional information about a selection event
  */
-UENUM()
+UENUM(BlueprintType)
 namespace ESelectInfo
 {
 	enum Type
diff --git a/Engine/Source/Runtime/UMG/Private/Components/Button.cpp b/Engine/Source/Runtime/UMG/Private/Components/Button.cpp
--- a/Engine/Source/Runtime/UMG/Private/Components/Button.cpp
+++ b/Engine/Source/Runtime/UMG/Private/Components/Button.cpp
@@ -161,6 +161,18 @@ const FText UButton::GetToolboxCategory()
 	return LOCTEXT("Common", "Common");
 }
 
+bool UButton::CanEditChange(const UProperty* Property) const
+{
+	static const FName ButtonStyleName(TEXT("ButtonStyle"));
+
+	if ( Property->GetFName() == ButtonStyleName )
+	{
+		return  Style == nullptr;
+	}
+
+	return true;
+}
+
 #endif
 
 /////////////////////////////////////////////////////
diff --git a/Engine/Source/Runtime/UMG/Private/Components/GridPanel.cpp b/Engine/Source/Runtime/UMG/Private/Components/GridPanel.cpp
--- a/Engine/Source/Runtime/UMG/Private/Components/GridPanel.cpp
+++ b/Engine/Source/Runtime/UMG/Private/Components/GridPanel.cpp
@@ -81,6 +81,11 @@ void UGridPanel::SynchronizeProperties()
 	}
 }
 
+UGridSlot* UGridPanel::GetGridChildSlot(UWidget* ChildWidget)
+{
+	return Cast<UGridSlot>(ChildWidget->Slot);
+}
+
 #if WITH_EDITOR
 
 const FSlateBrush* UGridPanel::GetEditorIcon()
diff --git a/Engine/Source/Runtime/UMG/Private/WidgetLayoutLibrary.cpp b/Engine/Source/Runtime/UMG/Private/WidgetLayoutLibrary.cpp
new file mode 100644
--- /dev/null
+++ b/Engine/Source/Runtime/UMG/Private/WidgetLayoutLibrary.cpp
@@ -0,0 +1,32 @@
+// Copyright 1998-2014 Epic Games, Inc. All Rights Reserved.
+
+#include "UMGPrivatePCH.h"
+#include "Slate/SlateBrushAsset.h"
+#include "WidgetLayoutLibrary.h"
+
+#define LOCTEXT_NAMESPACE "UMG"
+
+/////////////////////////////////////////////////////
+// UWidgetLayoutLibrary
+
+UWidgetLayoutLibrary::UWidgetLayoutLibrary(const FPostConstructInitializeProperties& PCIP)
+: Super(PCIP)
+{
+}
+
+UCanvasPanelSlot* UWidgetLayoutLibrary::SlotAsCanvasSlot(UWidget* ChildWidget)
+{
+	return Cast<UCanvasPanelSlot>(ChildWidget->Slot);
+}
+
+UGridSlot* UWidgetLayoutLibrary::SlotAsGridSlot(UWidget* ChildWidget)
+{
+	return Cast<UGridSlot>(ChildWidget->Slot);
+}
+
+UUniformGridSlot* UWidgetLayoutLibrary::SlotAsUniformGridSlot(UWidget* ChildWidget)
+{
+	return Cast<UUniformGridSlot>(ChildWidget->Slot);
+}
+
+#undef LOCTEXT_NAMESPACE
\ No newline at end of file
diff --git a/Engine/Source/Runtime/UMG/Public/Blueprint/UserWidget.h b/Engine/Source/Runtime/UMG/Public/Blueprint/UserWidget.h
--- a/Engine/Source/Runtime/UMG/Public/Blueprint/UserWidget.h
+++ b/Engine/Source/Runtime/UMG/Public/Blueprint/UserWidget.h
@@ -17,7 +17,7 @@ class UWidgetAnimation;
  * The state passed into OnPaint that we can expose as a single painting structure to blueprints to
  * allow script code to override OnPaint behavior.
  */
-USTRUCT()
+USTRUCT(BlueprintType)
 struct UMG_API FPaintContext
 {
 	GENERATED_USTRUCT_BODY()
diff --git a/Engine/Source/Runtime/UMG/Public/Blueprint/WidgetBlueprintLibrary.h b/Engine/Source/Runtime/UMG/Public/Blueprint/WidgetBlueprintLibrary.h
--- a/Engine/Source/Runtime/UMG/Public/Blueprint/WidgetBlueprintLibrary.h
+++ b/Engine/Source/Runtime/UMG/Public/Blueprint/WidgetBlueprintLibrary.h
@@ -84,20 +84,31 @@ class UWidgetBlueprintLibrary : public UBlueprintFunctionLibrary
 	UFUNCTION(BlueprintPure, Category="Widget|Event Reply")
 	static FEventReply EndDragDrop(UPARAM(ref) FEventReply& Reply);
 
-	/**  */
+	/**
+	 * Creates a Slate Brush from a Slate Brush Asset
+	 *
+	 * @return A new slate brush using the asset's brush.
+	 */
 	UFUNCTION(BlueprintPure, Category="Widget|Brush")
 	static FSlateBrush MakeBrushFromAsset(USlateBrushAsset* BrushAsset);
 
-   /** 
-	* Creates a Slate Brush from a Texture2D
-	*
-	* @param Width  When less than or equal to zero, the Width of the brush will default to the Width of the Texture
-	* @param Height  When less than or equal to zero, the Height of the brush will default to the Height of the Texture
-	*/
+	/** 
+	 * Creates a Slate Brush from a Texture2D
+	 *
+	 * @param Width  When less than or equal to zero, the Width of the brush will default to the Width of the Texture
+	 * @param Height  When less than or equal to zero, the Height of the brush will default to the Height of the Texture
+	 *
+	 * @return A new slate brush using the texture.
+	 */
 	UFUNCTION(BlueprintPure, Category="Widget|Brush")
 	static FSlateBrush MakeBrushFromTexture(UTexture2D* Texture, int32 Width = 0, int32 Height = 0);
 
-	/**  */
+	/**
+	 * Creates a Slate Brush from a Material.  Materials don't have an implicit size, so providing a widget and height
+	 * is required to hint slate with how large the image wants to be by default.
+	 *
+	 * @return A new slate brush using the material.
+	 */
 	UFUNCTION(BlueprintPure, Category="Widget|Brush")
-	static FSlateBrush MakeBrushFromMaterial(UMaterialInterface* Material, int32 Width, int32 Height);
+	static FSlateBrush MakeBrushFromMaterial(UMaterialInterface* Material, int32 Width = 32, int32 Height = 32);
 };
diff --git a/Engine/Source/Runtime/UMG/Public/Blueprint/WidgetLayoutLibrary.h b/Engine/Source/Runtime/UMG/Public/Blueprint/WidgetLayoutLibrary.h
new file mode 100644
--- /dev/null
+++ b/Engine/Source/Runtime/UMG/Public/Blueprint/WidgetLayoutLibrary.h
@@ -0,0 +1,32 @@
+// Copyright 1998-2014 Epic Games, Inc. All Rights Reserved.
+
+#pragma once
+
+#include "WidgetLayoutLibrary.generated.h"
+
+UCLASS(MinimalAPI)
+class UWidgetLayoutLibrary : public UBlueprintFunctionLibrary
+{
+	GENERATED_UCLASS_BODY()
+
+	/**
+	 * Gets the slot object on the child widget as a Canvas Slot, allowing you to manipulate layout information.
+	 * @param Widget The child widget of a canvas panel.
+	 */
+	UFUNCTION(BlueprintPure, Category="Layout")
+	static UCanvasPanelSlot* SlotAsCanvasSlot(UWidget* Widget);
+
+	/**
+	 * Gets the slot object on the child widget as a Grid Slot, allowing you to manipulate layout information.
+	 * @param Widget The child widget of a grid panel.
+	 */
+	UFUNCTION(BlueprintPure, Category="Layout")
+	static UGridSlot* SlotAsGridSlot(UWidget* Widget);
+
+	/**
+	 * Gets the slot object on the child widget as a Uniform Grid Slot, allowing you to manipulate layout information.
+	 * @param Widget The child widget of a uniform grid panel.
+	 */
+	UFUNCTION(BlueprintPure, Category="Layout")
+	static UUniformGridSlot* SlotAsUniformGridSlot(UWidget* Widget);
+};
diff --git a/Engine/Source/Runtime/UMG/Public/Components/BorderSlot.h b/Engine/Source/Runtime/UMG/Public/Components/BorderSlot.h
--- a/Engine/Source/Runtime/UMG/Public/Components/BorderSlot.h
+++ b/Engine/Source/Runtime/UMG/Public/Components/BorderSlot.h
@@ -16,26 +16,26 @@ class UMG_API UBorderSlot : public UPanelSlot
 
 public:
 
-	UFUNCTION(BlueprintCallable, Category="Appearance")
+	UFUNCTION(BlueprintCallable, Category="Layout (Border Slot)")
 	void SetPadding(FMargin InPadding);
 
-	UFUNCTION(BlueprintCallable, Category="Appearance")
+	UFUNCTION(BlueprintCallable, Category="Layout (Border Slot)")
 	void SetHorizontalAlignment(EHorizontalAlignment InHorizontalAlignment);
 
-	UFUNCTION(BlueprintCallable, Category="Appearance")
+	UFUNCTION(BlueprintCallable, Category="Layout (Border Slot)")
 	void SetVerticalAlignment(EVerticalAlignment InVerticalAlignment);
 
 protected:
 	/** The padding area between the slot and the content it contains. */
-	UPROPERTY(EditDefaultsOnly, Category=Layout)
+	UPROPERTY(EditDefaultsOnly, Category="Layout (Border Slot)")
 	FMargin Padding;
 
 	/** The alignment of the object horizontally. */
-	UPROPERTY(EditDefaultsOnly, Category=Layout)
+	UPROPERTY(EditDefaultsOnly, Category="Layout (Border Slot)")
 	TEnumAsByte<EHorizontalAlignment> HorizontalAlignment;
 
 	/** The alignment of the object vertically. */
-	UPROPERTY(EditDefaultsOnly, Category=Layout)
+	UPROPERTY(EditDefaultsOnly, Category="Layout (Border Slot)")
 	TEnumAsByte<EVerticalAlignment> VerticalAlignment;
 
 public:
diff --git a/Engine/Source/Runtime/UMG/Public/Components/Button.h b/Engine/Source/Runtime/UMG/Public/Components/Button.h
--- a/Engine/Source/Runtime/UMG/Public/Components/Button.h
+++ b/Engine/Source/Runtime/UMG/Public/Components/Button.h
@@ -18,11 +18,11 @@ class UMG_API UButton : public UContentWidget
 
 public:
 	/** The template style asset, used to seed the mutable instance of the style. */
-	UPROPERTY(EditDefaultsOnly, AdvancedDisplay, DisplayName="Template Style", Category=Appearance, meta=(DesignerRebuild=true))
+	UPROPERTY(EditDefaultsOnly, AdvancedDisplay, DisplayName="Template Style", Category=Appearance, meta=( DesignerRebuild=true ))
 	USlateWidgetStyleAsset* Style;
 
 	/** The button style used at runtime by the slate button */
-	UPROPERTY(VisibleAnywhere, Instanced, Category=Appearance)
+	UPROPERTY(VisibleAnywhere, Instanced, Category=Appearance )
 	TSubobjectPtr<UButtonWidgetStyle> ButtonStyle;
 	
 	/** The color multiplier for the button content */
@@ -86,6 +86,7 @@ public:
 #if WITH_EDITOR
 	virtual const FSlateBrush* GetEditorIcon() override;
 	virtual const FText GetToolboxCategory() override;
+	virtual bool CanEditChange(const UProperty* Property) const override;
 #endif
 
 	static FName StyleName;
diff --git a/Engine/Source/Runtime/UMG/Public/Components/ButtonSlot.h b/Engine/Source/Runtime/UMG/Public/Components/ButtonSlot.h
--- a/Engine/Source/Runtime/UMG/Public/Components/ButtonSlot.h
+++ b/Engine/Source/Runtime/UMG/Public/Components/ButtonSlot.h
@@ -15,24 +15,24 @@ class UMG_API UButtonSlot : public UPanelSlot
 public:
 	
 	/** The padding area between the slot and the content it contains. */
-	UPROPERTY(EditDefaultsOnly, Category=Layout)
+	UPROPERTY(EditDefaultsOnly, Category="Layout (Button Slot)")
 	FMargin Padding;
 
 	/** The alignment of the object horizontally. */
-	UPROPERTY(EditDefaultsOnly, Category=Layout)
+	UPROPERTY(EditDefaultsOnly, Category="Layout (Button Slot)")
 	TEnumAsByte<EHorizontalAlignment> HorizontalAlignment;
 
 	/** The alignment of the object vertically. */
-	UPROPERTY(EditDefaultsOnly, Category=Layout)
+	UPROPERTY(EditDefaultsOnly, Category="Layout (Button Slot)")
 	TEnumAsByte<EVerticalAlignment> VerticalAlignment;
 
-	UFUNCTION(BlueprintCallable, Category="Appearance")
+	UFUNCTION(BlueprintCallable, Category="Layout (Button Slot)")
 	void SetPadding(FMargin InPadding);
 
-	UFUNCTION(BlueprintCallable, Category="Appearance")
+	UFUNCTION(BlueprintCallable, Category="Layout (Button Slot)")
 	void SetHorizontalAlignment(EHorizontalAlignment InHorizontalAlignment);
 
-	UFUNCTION(BlueprintCallable, Category="Appearance")
+	UFUNCTION(BlueprintCallable, Category="Layout (Button Slot)")
 	void SetVerticalAlignment(EVerticalAlignment InVerticalAlignment);
 
 	// UPanelSlot interface
diff --git a/Engine/Source/Runtime/UMG/Public/Components/GridPanel.h b/Engine/Source/Runtime/UMG/Public/Components/GridPanel.h
--- a/Engine/Source/Runtime/UMG/Public/Components/GridPanel.h
+++ b/Engine/Source/Runtime/UMG/Public/Components/GridPanel.h
@@ -12,14 +12,21 @@ class UMG_API UGridPanel : public UPanelWidget
 
 public:
 
-	/** Fill */
+	/** The column fill rules */
 	UPROPERTY(EditDefaultsOnly, Category="Fill Rules")
 	TArray<float> ColumnFill;
 
-	/** Fill */
+	/** The row fill rules */
 	UPROPERTY(EditDefaultsOnly, Category="Fill Rules")
 	TArray<float> RowFill;
 
+	/**
+	 * Gets the grid slot object on the child widget inside this panel, allowing you to manipulate layout information.
+	 * @param ChildWidget The child widget.
+	 */
+	UFUNCTION(BlueprintCallable, Category="Slot")
+	static UGridSlot* GetGridChildSlot(UWidget* ChildWidget);
+
 public:
 
 	// UWidget interface
diff --git a/Engine/Source/Runtime/UMG/Public/Components/GridSlot.h b/Engine/Source/Runtime/UMG/Public/Components/GridSlot.h
--- a/Engine/Source/Runtime/UMG/Public/Components/GridSlot.h
+++ b/Engine/Source/Runtime/UMG/Public/Components/GridSlot.h
@@ -18,59 +18,59 @@ class UMG_API UGridSlot : public UPanelSlot
 public:
 
 	/** The alignment of the object horizontally. */
-	UPROPERTY(EditDefaultsOnly, Category=Layout)
+	UPROPERTY(EditDefaultsOnly, Category="Layout (Grid Slot)")
 	TEnumAsByte<EHorizontalAlignment> HorizontalAlignment;
 
 	/** The alignment of the object vertically. */
-	UPROPERTY(EditDefaultsOnly, Category=Layout)
+	UPROPERTY(EditDefaultsOnly, Category="Layout (Grid Slot)")
 	TEnumAsByte<EVerticalAlignment> VerticalAlignment;
 	
 	/** The row index of the cell this slot is in */
-	UPROPERTY(EditDefaultsOnly, meta=( UIMin = "0" ), Category=Layout)
+	UPROPERTY(EditDefaultsOnly, meta=( UIMin = "0" ), Category="Layout (Grid Slot)")
 	int32 Row;
 	
 	/**  */
-	UPROPERTY(EditDefaultsOnly, Category=Layout)
+	UPROPERTY(EditDefaultsOnly, Category="Layout (Grid Slot)")
 	int32 RowSpan;
 
 	/** The column index of the cell this slot is in */
-	UPROPERTY(EditDefaultsOnly, meta=( UIMin = "0" ), Category=Layout)
+	UPROPERTY(EditDefaultsOnly, meta=( UIMin = "0" ), Category="Layout (Grid Slot)")
 	int32 Column;
 
 	/**  */
-	UPROPERTY(EditDefaultsOnly, Category=Layout)
+	UPROPERTY(EditDefaultsOnly, Category="Layout (Grid Slot)")
 	int32 ColumnSpan;
 
 	/** Positive values offset this cell to be hit-tested and drawn on top of others. Default is 0; i.e. no offset. */
-	UPROPERTY(EditDefaultsOnly, Category=Layout)
+	UPROPERTY(EditDefaultsOnly, Category="Layout (Grid Slot)")
 	int32 Layer;
 
 	/** Offset this slot's content by some amount; positive values offset to lower right*/
-	UPROPERTY(EditDefaultsOnly, Category=Layout)
+	UPROPERTY(EditDefaultsOnly, Category="Layout (Grid Slot)")
 	FVector2D Nudge;
 
 	/** Sets the row index of the slot, this determines what cell the slot is in the panel */
-	UFUNCTION(BlueprintCallable, Category=Layout)
+	UFUNCTION(BlueprintCallable, Category="Layout (Grid Slot)")
 	void SetRow(int32 InRow);
 
 	/** How many rows this this slot spans over */
-	UFUNCTION(BlueprintCallable, Category=Layout)
+	UFUNCTION(BlueprintCallable, Category="Layout (Grid Slot)")
 	void SetRowSpan(int32 InRowSpan);
 
 	/** Sets the column index of the slot, this determines what cell the slot is in the panel */
-	UFUNCTION(BlueprintCallable, Category=Layout)
+	UFUNCTION(BlueprintCallable, Category="Layout (Grid Slot)")
 	void SetColumn(int32 InColumn);
 
 	/** How many columns this slot spans over */
-	UFUNCTION(BlueprintCallable, Category=Layout)
+	UFUNCTION(BlueprintCallable, Category="Layout (Grid Slot)")
 	void SetColumnSpan(int32 InColumnSpan);
 
 	/**  */
-	UFUNCTION(BlueprintCallable, Category=Layout)
+	UFUNCTION(BlueprintCallable, Category="Layout (Grid Slot)")
 	void SetHorizontalAlignment(EHorizontalAlignment InHorizontalAlignment);
 
 	/**  */
-	UFUNCTION(BlueprintCallable, Category=Layout)
+	UFUNCTION(BlueprintCallable, Category="Layout (Grid Slot)")
 	void SetVerticalAlignment(EVerticalAlignment InVerticalAlignment);
 
 	// UPanelSlot interface
diff --git a/Engine/Source/Runtime/UMG/Public/Components/HorizontalBoxSlot.h b/Engine/Source/Runtime/UMG/Public/Components/HorizontalBoxSlot.h
--- a/Engine/Source/Runtime/UMG/Public/Components/HorizontalBoxSlot.h
+++ b/Engine/Source/Runtime/UMG/Public/Components/HorizontalBoxSlot.h
@@ -14,29 +14,29 @@ class UMG_API UHorizontalBoxSlot : public UPanelSlot
 public:
 	
 	/** The amount of padding between the slots parent and the content. */
-	UPROPERTY(EditDefaultsOnly, Category=Appearance)
+	UPROPERTY(EditDefaultsOnly, Category="Layout (Horizontal Box Slot)")
 	FMargin Padding;
 
 	/** How much space this slot should occupy in the direction of the panel. */
-	UPROPERTY(EditDefaultsOnly, Category=Appearance)
+	UPROPERTY(EditDefaultsOnly, Category="Layout (Horizontal Box Slot)")
 	FSlateChildSize Size;
 
-	UPROPERTY(EditDefaultsOnly, Category=Appearance)
+	UPROPERTY(EditDefaultsOnly, Category="Layout (Horizontal Box Slot)")
 	TEnumAsByte<EHorizontalAlignment> HorizontalAlignment;
 
-	UPROPERTY(EditDefaultsOnly, Category=Appearance)
+	UPROPERTY(EditDefaultsOnly, Category="Layout (Horizontal Box Slot)")
 	TEnumAsByte<EVerticalAlignment> VerticalAlignment;
 
-	UFUNCTION(BlueprintCallable, Category="Appearance")
+	UFUNCTION(BlueprintCallable, Category="Layout(Horizontal Box Slot)")
 	void SetPadding(FMargin InPadding);
 
-	UFUNCTION(BlueprintCallable, Category="Appearance")
+	UFUNCTION(BlueprintCallable, Category="Layout (Horizontal Box Slot)")
 	void SetSize(FSlateChildSize InSize);
 
-	UFUNCTION(BlueprintCallable, Category="Appearance")
+	UFUNCTION(BlueprintCallable, Category="Layout (Horizontal Box Slot)")
 	void SetHorizontalAlignment(EHorizontalAlignment InHorizontalAlignment);
 
-	UFUNCTION(BlueprintCallable, Category="Appearance")
+	UFUNCTION(BlueprintCallable, Category="Layout (Horizontal Box Slot)")
 	void SetVerticalAlignment(EVerticalAlignment InVerticalAlignment);
 
 	void BuildSlot(TSharedRef<SHorizontalBox> HorizontalBox);
diff --git a/Engine/Source/Runtime/UMG/Public/Components/OverlaySlot.h b/Engine/Source/Runtime/UMG/Public/Components/OverlaySlot.h
--- a/Engine/Source/Runtime/UMG/Public/Components/OverlaySlot.h
+++ b/Engine/Source/Runtime/UMG/Public/Components/OverlaySlot.h
@@ -11,28 +11,34 @@ UCLASS()
 class UMG_API UOverlaySlot : public UPanelSlot
 {
 	GENERATED_UCLASS_BODY()
+
+public:
 	
 	/** The padding area between the slot and the content it contains. */
-	UPROPERTY(EditDefaultsOnly, Category=Layout)
+	UPROPERTY(EditDefaultsOnly, Category="Layout (Overlay Slot)")
 	FMargin Padding;
 
 	/** The alignment of the object horizontally. */
-	UPROPERTY(EditDefaultsOnly, Category=Layout)
+	UPROPERTY(EditDefaultsOnly, Category="Layout (Overlay Slot)")
 	TEnumAsByte<EHorizontalAlignment> HorizontalAlignment;
 
 	/** The alignment of the object vertically. */
-	UPROPERTY(EditDefaultsOnly, Category=Layout)
+	UPROPERTY(EditDefaultsOnly, Category="Layout (Overlay Slot)")
 	TEnumAsByte<EVerticalAlignment> VerticalAlignment;
 
-	UFUNCTION(BlueprintCallable, Category="Appearance")
+public:
+
+	UFUNCTION(BlueprintCallable, Category="Layout (Overlay Slot)")
 	void SetPadding(FMargin InPadding);
 
-	UFUNCTION(BlueprintCallable, Category="Appearance")
+	UFUNCTION(BlueprintCallable, Category="Layout (Overlay Slot)")
 	void SetHorizontalAlignment(EHorizontalAlignment InHorizontalAlignment);
 
-	UFUNCTION(BlueprintCallable, Category="Appearance")
+	UFUNCTION(BlueprintCallable, Category="Layout (Overlay Slot)")
 	void SetVerticalAlignment(EVerticalAlignment InVerticalAlignment);
 
+public:
+
 	void BuildSlot(TSharedRef<SHorizontalBox> HorizontalBox);
 
 	// UPanelSlot interface
diff --git a/Engine/Source/Runtime/UMG/Public/Components/ScaleBoxSlot.h b/Engine/Source/Runtime/UMG/Public/Components/ScaleBoxSlot.h
--- a/Engine/Source/Runtime/UMG/Public/Components/ScaleBoxSlot.h
+++ b/Engine/Source/Runtime/UMG/Public/Components/ScaleBoxSlot.h
@@ -15,26 +15,26 @@ class UMG_API UScaleBoxSlot : public UPanelSlot
 public:
 	
 	/** The padding area between the slot and the content it contains. */
-	UPROPERTY(EditDefaultsOnly, Category=Layout)
+	UPROPERTY(EditDefaultsOnly, Category="Layout (ScaleBox Slot)")
 	FMargin Padding;
 
 	/** The alignment of the object horizontally. */
-	UPROPERTY(EditDefaultsOnly, Category=Layout)
+	UPROPERTY(EditDefaultsOnly, Category="Layout (ScaleBox Slot)")
 	TEnumAsByte<EHorizontalAlignment> HorizontalAlignment;
 
 	/** The alignment of the object vertically. */
-	UPROPERTY(EditDefaultsOnly, Category=Layout)
+	UPROPERTY(EditDefaultsOnly, Category="Layout (ScaleBox Slot)")
 	TEnumAsByte<EVerticalAlignment> VerticalAlignment;
 
 public:
 
-	UFUNCTION(BlueprintCallable, Category="Appearance")
+	UFUNCTION(BlueprintCallable, Category="Layout (ScaleBox Slot)")
 	void SetPadding(FMargin InPadding);
 
-	UFUNCTION(BlueprintCallable, Category="Appearance")
+	UFUNCTION(BlueprintCallable, Category="Layout (ScaleBox Slot)")
 	void SetHorizontalAlignment(EHorizontalAlignment InHorizontalAlignment);
 
-	UFUNCTION(BlueprintCallable, Category="Appearance")
+	UFUNCTION(BlueprintCallable, Category="Layout (ScaleBox Slot)")
 	void SetVerticalAlignment(EVerticalAlignment InVerticalAlignment);
 
 	// UPanelSlot interface
diff --git a/Engine/Source/Runtime/UMG/Public/Components/ScrollBoxSlot.h b/Engine/Source/Runtime/UMG/Public/Components/ScrollBoxSlot.h
--- a/Engine/Source/Runtime/UMG/Public/Components/ScrollBoxSlot.h
+++ b/Engine/Source/Runtime/UMG/Public/Components/ScrollBoxSlot.h
@@ -15,19 +15,23 @@ class UMG_API UScrollBoxSlot : public UPanelSlot
 public:
 	
 	/** The padding area between the slot and the content it contains. */
-	UPROPERTY(EditDefaultsOnly, Category=Layout)
+	UPROPERTY(EditDefaultsOnly, Category="Layout (ScrollBox Slot)")
 	FMargin Padding;
 
 	/** The alignment of the object horizontally. */
-	UPROPERTY(EditDefaultsOnly, Category=Layout)
+	UPROPERTY(EditDefaultsOnly, Category="Layout (ScrollBox Slot)")
 	TEnumAsByte<EHorizontalAlignment> HorizontalAlignment;
 
-	UFUNCTION(BlueprintCallable, Category="Appearance")
+public:
+
+	UFUNCTION(BlueprintCallable, Category="Layout (ScrollBox Slot)")
 	void SetPadding(FMargin InPadding);
 
-	UFUNCTION(BlueprintCallable, Category="Appearance")
+	UFUNCTION(BlueprintCallable, Category="Layout (ScrollBox Slot)")
 	void SetHorizontalAlignment(EHorizontalAlignment InHorizontalAlignment);
 
+public:
+
 	// UPanelSlot interface
 	virtual void SynchronizeProperties() override;
 	// End of UPanelSlot interface
diff --git a/Engine/Source/Runtime/UMG/Public/Components/SizeBoxSlot.h b/Engine/Source/Runtime/UMG/Public/Components/SizeBoxSlot.h
--- a/Engine/Source/Runtime/UMG/Public/Components/SizeBoxSlot.h
+++ b/Engine/Source/Runtime/UMG/Public/Components/SizeBoxSlot.h
@@ -14,13 +14,13 @@ class UMG_API USizeBoxSlot : public UPanelSlot
 
 public:
 
-	UFUNCTION(BlueprintCallable, Category="Layout")
+	UFUNCTION(BlueprintCallable, Category="Layout (SizeBox Slot)")
 	void SetPadding(FMargin InPadding);
 
-	UFUNCTION(BlueprintCallable, Category="Layout")
+	UFUNCTION(BlueprintCallable, Category="Layout (SizeBox Slot)")
 	void SetHorizontalAlignment(EHorizontalAlignment InHorizontalAlignment);
 
-	UFUNCTION(BlueprintCallable, Category="Layout")
+	UFUNCTION(BlueprintCallable, Category="Layout (SizeBox Slot)")
 	void SetVerticalAlignment(EVerticalAlignment InVerticalAlignment);
 
 public:
@@ -36,15 +36,15 @@ public:
 
 protected:
 	/** The padding area between the slot and the content it contains. */
-	UPROPERTY(EditDefaultsOnly, Category=Layout)
+	UPROPERTY(EditDefaultsOnly, Category="Layout (SizeBox Slot)")
 	FMargin Padding;
 
 	/** The alignment of the object horizontally. */
-	UPROPERTY(EditDefaultsOnly, Category=Layout)
+	UPROPERTY(EditDefaultsOnly, Category="Layout (SizeBox Slot)")
 	TEnumAsByte<EHorizontalAlignment> HorizontalAlignment;
 
 	/** The alignment of the object vertically. */
-	UPROPERTY(EditDefaultsOnly, Category=Layout)
+	UPROPERTY(EditDefaultsOnly, Category="Layout (SizeBox Slot)")
 	TEnumAsByte<EVerticalAlignment> VerticalAlignment;
 
 private:
diff --git a/Engine/Source/Runtime/UMG/Public/Components/SlateWrapperTypes.h b/Engine/Source/Runtime/UMG/Public/Components/SlateWrapperTypes.h
--- a/Engine/Source/Runtime/UMG/Public/Components/SlateWrapperTypes.h
+++ b/Engine/Source/Runtime/UMG/Public/Components/SlateWrapperTypes.h
@@ -31,7 +31,7 @@ namespace ESlateVisibility
 }
 
 /** The sizing options of UWidgets */
-UENUM()
+UENUM(BlueprintType)
 namespace ESlateSizeRule
 {
 	enum Type
@@ -77,7 +77,7 @@ struct FEventReply
 };
 
 /** A struct exposing size param related properties to UMG. */
-USTRUCT()
+USTRUCT(BlueprintType)
 struct FSlateChildSize
 {
 	GENERATED_USTRUCT_BODY()
diff --git a/Engine/Source/Runtime/UMG/Public/Components/UniformGridSlot.h b/Engine/Source/Runtime/UMG/Public/Components/UniformGridSlot.h
--- a/Engine/Source/Runtime/UMG/Public/Components/UniformGridSlot.h
+++ b/Engine/Source/Runtime/UMG/Public/Components/UniformGridSlot.h
@@ -18,37 +18,41 @@ class UMG_API UUniformGridSlot : public UPanelSlot
 public:
 
 	/** The alignment of the object horizontally. */
-	UPROPERTY(EditDefaultsOnly, Category=Layout)
+	UPROPERTY(EditDefaultsOnly, Category="Layout (Uniform Grid Slot)")
 	TEnumAsByte<EHorizontalAlignment> HorizontalAlignment;
 
 	/** The alignment of the object vertically. */
-	UPROPERTY(EditDefaultsOnly, Category=Layout)
+	UPROPERTY(EditDefaultsOnly, Category="Layout (Uniform Grid Slot)")
 	TEnumAsByte<EVerticalAlignment> VerticalAlignment;
 	
 	/** The row index of the cell this slot is in */
-	UPROPERTY(EditDefaultsOnly, meta=( UIMin = "0" ), Category=Layout)
+	UPROPERTY(EditDefaultsOnly, meta=( UIMin = "0" ), Category="Layout (Uniform Grid Slot)")
 	int32 Row;
 	
 	/** The column index of the cell this slot is in */
-	UPROPERTY(EditDefaultsOnly, meta=( UIMin = "0" ), Category=Layout)
+	UPROPERTY(EditDefaultsOnly, meta=( UIMin = "0" ), Category="Layout (Uniform Grid Slot)")
 	int32 Column;
 
+public:
+
 	/** Sets the row index of the slot, this determines what cell the slot is in the panel */
-	UFUNCTION(BlueprintCallable, Category=Layout)
+	UFUNCTION(BlueprintCallable, Category="Layout (Uniform Grid Slot)")
 	void SetRow(int32 InRow);
 
 	/** Sets the column index of the slot, this determines what cell the slot is in the panel */
-	UFUNCTION(BlueprintCallable, Category=Layout)
+	UFUNCTION(BlueprintCallable, Category="Layout (Uniform Grid Slot)")
 	void SetColumn(int32 InColumn);
 
 	/**  */
-	UFUNCTION(BlueprintCallable, Category=Layout)
+	UFUNCTION(BlueprintCallable, Category="Layout (Uniform Grid Slot)")
 	void SetHorizontalAlignment(EHorizontalAlignment InHorizontalAlignment);
 
 	/**  */
-	UFUNCTION(BlueprintCallable, Category=Layout)
+	UFUNCTION(BlueprintCallable, Category="Layout (Uniform Grid Slot)")
 	void SetVerticalAlignment(EVerticalAlignment InVerticalAlignment);
 
+public:
+
 	// UPanelSlot interface
 	virtual void SynchronizeProperties() override;
 	// End of UPanelSlot interface
diff --git a/Engine/Source/Runtime/UMG/Public/Components/VerticalBoxSlot.h b/Engine/Source/Runtime/UMG/Public/Components/VerticalBoxSlot.h
--- a/Engine/Source/Runtime/UMG/Public/Components/VerticalBoxSlot.h
+++ b/Engine/Source/Runtime/UMG/Public/Components/VerticalBoxSlot.h
@@ -15,34 +15,36 @@ class UMG_API UVerticalBoxSlot : public UPanelSlot
 public:
 	
 	/** The padding area between the slot and the content it contains. */
-	UPROPERTY(EditDefaultsOnly, Category=Layout)
+	UPROPERTY(EditDefaultsOnly, Category="Layout (Vertical Box Slot)")
 	FMargin Padding;
 
 	/** How much space this slot should occupy in the direction of the panel. */
-	UPROPERTY(EditDefaultsOnly, Category=Layout)
+	UPROPERTY(EditDefaultsOnly, Category="Layout (Vertical Box Slot)")
 	FSlateChildSize Size;
 
 	/** The alignment of the object horizontally. */
-	UPROPERTY(EditDefaultsOnly, Category=Layout)
+	UPROPERTY(EditDefaultsOnly, Category="Layout (Vertical Box Slot)")
 	TEnumAsByte<EHorizontalAlignment> HorizontalAlignment;
 
 	/** The alignment of the object vertically. */
-	UPROPERTY(EditDefaultsOnly, Category=Layout)
+	UPROPERTY(EditDefaultsOnly, Category="Layout (Vertical Box Slot)")
 	TEnumAsByte<EVerticalAlignment> VerticalAlignment;
 
 public:
 
-	UFUNCTION(BlueprintCallable, Category="Appearance")
-	void SetPadding(FMargin InPadding);
+	UFUNCTION(BlueprintCallable, Category="Layout (Vertical Box Slot)")
+		void SetPadding(FMargin InPadding);
 
-	UFUNCTION(BlueprintCallable, Category="Appearance")
-	void SetSize(FSlateChildSize InSize);
+	UFUNCTION(BlueprintCallable, Category="Layout (Vertical Box Slot)")
+		void SetSize(FSlateChildSize InSize);
 
-	UFUNCTION(BlueprintCallable, Category="Appearance")
-	void SetHorizontalAlignment(EHorizontalAlignment InHorizontalAlignment);
+	UFUNCTION(BlueprintCallable, Category="Layout (Vertical Box Slot)")
+		void SetHorizontalAlignment(EHorizontalAlignment InHorizontalAlignment);
 
-	UFUNCTION(BlueprintCallable, Category="Appearance")
-	void SetVerticalAlignment(EVerticalAlignment InVerticalAlignment);
+	UFUNCTION(BlueprintCallable, Category="Layout (Vertical Box Slot)")
+		void SetVerticalAlignment(EVerticalAlignment InVerticalAlignment);
+
+public:
 
 	// UPanelSlot interface
 	virtual void SynchronizeProperties() override;
diff --git a/Engine/Source/Runtime/UMG/Public/Components/WidgetSwitcherSlot.h b/Engine/Source/Runtime/UMG/Public/Components/WidgetSwitcherSlot.h
--- a/Engine/Source/Runtime/UMG/Public/Components/WidgetSwitcherSlot.h
+++ b/Engine/Source/Runtime/UMG/Public/Components/WidgetSwitcherSlot.h
@@ -11,28 +11,34 @@ UCLASS()
 class UMG_API UWidgetSwitcherSlot : public UPanelSlot
 {
 	GENERATED_UCLASS_BODY()
+
+public:
 	
 	/** The padding area between the slot and the content it contains. */
-	UPROPERTY(EditDefaultsOnly, Category=Layout)
+	UPROPERTY(EditDefaultsOnly, Category="Layout (Widget Switcher Slot)")
 	FMargin Padding;
 
 	/** The alignment of the object horizontally. */
-	UPROPERTY(EditDefaultsOnly, Category=Layout)
+	UPROPERTY(EditDefaultsOnly, Category="Layout (Widget Switcher Slot)")
 	TEnumAsByte<EHorizontalAlignment> HorizontalAlignment;
 
 	/** The alignment of the object vertically. */
-	UPROPERTY(EditDefaultsOnly, Category=Layout)
+	UPROPERTY(EditDefaultsOnly, Category="Layout (Widget Switcher Slot)")
 	TEnumAsByte<EVerticalAlignment> VerticalAlignment;
 
-	UFUNCTION(BlueprintCallable, Category="Appearance")
+public:
+
+	UFUNCTION(BlueprintCallable, Category="Layout (Widget Switcher Slot)")
 	void SetPadding(FMargin InPadding);
 
-	UFUNCTION(BlueprintCallable, Category="Appearance")
+	UFUNCTION(BlueprintCallable, Category="Layout (Widget Switcher Slot)")
 	void SetHorizontalAlignment(EHorizontalAlignment InHorizontalAlignment);
 
-	UFUNCTION(BlueprintCallable, Category="Appearance")
+	UFUNCTION(BlueprintCallable, Category="Layout (Widget Switcher Slot)")
 	void SetVerticalAlignment(EVerticalAlignment InVerticalAlignment);
 
+public:
+
 	// UPanelSlot interface
 	virtual void SynchronizeProperties() override;
 	// End of UPanelSlot interface
