UMG - Making it easier to set position in the viewport, taking into account the current UI scale curve inverse so that it shows up at the correct screen location.  Adding ways to query the current scale curve to the Widget layout library.  bAbsoluteLayout has been removed, we no longer need it, as soon as you change position, it automatically converts to Abs Layout.

diff --git a/Engine/Source/Editor/UMGEditor/Private/Designer/SDesignerView.cpp b/Engine/Source/Editor/UMGEditor/Private/Designer/SDesignerView.cpp
--- a/Engine/Source/Editor/UMGEditor/Private/Designer/SDesignerView.cpp
+++ b/Engine/Source/Editor/UMGEditor/Private/Designer/SDesignerView.cpp
@@ -910,11 +910,10 @@ void SDesignerView::UpdatePreviewWidget()
 		PreviewWidget = LatestPreviewWidget;
 		if ( PreviewWidget )
 		{
-			const bool bAbsoluteLayout = false;
 			const bool bModal = false;
 			const bool bShowCursor = false;
 			TSharedPtr<SWidget> OutUserWidget;
-			TSharedRef<SWidget> CurrentWidget = PreviewWidget->MakeViewportWidget(bAbsoluteLayout, bModal, bShowCursor, OutUserWidget);
+			TSharedRef<SWidget> CurrentWidget = PreviewWidget->MakeViewportWidget(bModal, bShowCursor, OutUserWidget);
 			CurrentWidget->SlatePrepass();
 
 			if ( CurrentWidget != PreviewSlateWidget.Pin() )
diff --git a/Engine/Source/Runtime/UMG/Private/UserWidget.cpp b/Engine/Source/Runtime/UMG/Private/UserWidget.cpp
--- a/Engine/Source/Runtime/UMG/Private/UserWidget.cpp
+++ b/Engine/Source/Runtime/UMG/Private/UserWidget.cpp
@@ -5,7 +5,9 @@
 #include "UMGSequencePlayer.h"
 #include "SceneViewport.h"
 #include "WidgetAnimation.h"
+
 #include "WidgetBlueprintLibrary.h"
+#include "WidgetLayoutLibrary.h"
 
 #define LOCTEXT_NAMESPACE "UMG"
 
@@ -122,9 +124,7 @@ protected:
 UUserWidget::UUserWidget(const FPostConstructInitializeProperties& PCIP)
 	: Super(PCIP)
 {
-	HorizontalAlignment = HAlign_Fill;
-	VerticalAlignment = VAlign_Fill;
-
+	ViewportAnchors = FAnchors(0, 0, 1, 1);
 	Visiblity = ESlateVisibility::SelfHitTestInvisible;
 
 	bInitialized = false;
@@ -180,27 +180,25 @@ void UUserWidget::PostInitProperties()
 
 UWorld* UUserWidget::GetWorld() const
 {
+	// Use the Player Context's world.
 	if ( PlayerContext.IsValid() )
 	{
-		return PlayerContext.GetWorld();
+		if ( UWorld* World = PlayerContext.GetWorld() )
+		{
+			return World;
+		}
 	}
 
-	if ( CachedWorld == NULL )
+	// If the current player context doesn't have a world or isn't valid, return the game instance's world.
+	if (  UGameInstance* GameInstance = Cast<UGameInstance>(GetOuter()) )
 	{
-		UObject* Outer = GetOuter();
-		while ( Outer != NULL )
+		if ( UWorld* World = GameInstance->GetWorld() )
 		{
-			CachedWorld = Outer->GetWorld();
-			if ( CachedWorld )
-			{
-				return CachedWorld;
-			}
-
-			Outer = Outer->GetOuter();
+			return World;
 		}
 	}
-	
-	return CachedWorld;
+
+	return nullptr;
 }
 
 void UUserWidget::Construct_Implementation()
@@ -499,41 +497,20 @@ void UUserWidget::NativeTick(const FGeometry& MyGeometry, float InDeltaTime )
 	Tick( MyGeometry, InDeltaTime );
 }
 
-TSharedRef<SWidget> UUserWidget::MakeViewportWidget(bool bAbsoluteLayout, bool bModal, bool bShowCursor, TSharedPtr<SWidget>& UserSlateWidget)
+TSharedRef<SWidget> UUserWidget::MakeViewportWidget(bool bModal, bool bShowCursor, TSharedPtr<SWidget>& UserSlateWidget)
 {
 	UserSlateWidget = TakeWidget();
 
-	if ( bAbsoluteLayout )
-	{
-		//BIND_UOBJECT_ATTRIBUTE
-
-		return SNew(SConstraintCanvas)
-
-			+ SConstraintCanvas::Slot()
-			.Offset(BIND_UOBJECT_ATTRIBUTE(FMargin, GetFullScreenOffset))
-			.Anchors(BIND_UOBJECT_ATTRIBUTE(FAnchors, GetViewportAnchors))
-			.Alignment(BIND_UOBJECT_ATTRIBUTE(FVector2D, GetFullScreenAlignment))
-			.ZOrder(BIND_UOBJECT_ATTRIBUTE(int32, GetFullScreenZOrder))
-			[
-				UserSlateWidget.ToSharedRef()
-			];
-	}
-	else
-	{
-		TSharedRef<SVerticalBox> VerticalBox = SNew(SVerticalBox);
-
-		auto& NewSlot = VerticalBox->AddSlot()
-			.Padding(Padding)
-			.HAlign(HorizontalAlignment)
-			.VAlign(VerticalAlignment)
-			[
-				UserSlateWidget.ToSharedRef()
-			];
-
-		NewSlot.SizeParam = UWidget::ConvertSerializedSizeParamToRuntime(Size);
+	return SNew(SConstraintCanvas)
 
-		return VerticalBox;
-	}
+		+ SConstraintCanvas::Slot()
+		.Offset(BIND_UOBJECT_ATTRIBUTE(FMargin, GetFullScreenOffset))
+		.Anchors(BIND_UOBJECT_ATTRIBUTE(FAnchors, GetViewportAnchors))
+		.Alignment(BIND_UOBJECT_ATTRIBUTE(FVector2D, GetFullScreenAlignment))
+		.ZOrder(BIND_UOBJECT_ATTRIBUTE(int32, GetFullScreenZOrder))
+		[
+			UserSlateWidget.ToSharedRef()
+		];
 }
 
 UWidget* UUserWidget::GetRootWidgetComponent()
@@ -546,12 +523,12 @@ UWidget* UUserWidget::GetRootWidgetComponent()
 	return NULL;
 }
 
-void UUserWidget::AddToViewport(bool bAbsoluteLayout, bool bModal, bool bShowCursor)
+void UUserWidget::AddToViewport(bool bModal, bool bShowCursor)
 {
 	if ( !FullScreenWidget.IsValid() )
 	{
 		TSharedPtr<SWidget> OutUserSlateWidget;
-		TSharedRef<SWidget> RootWidget = MakeViewportWidget(bAbsoluteLayout, bModal, bShowCursor, OutUserSlateWidget);
+		TSharedRef<SWidget> RootWidget = MakeViewportWidget(bModal, bShowCursor, OutUserSlateWidget);
 
 		TSharedRef<SViewportWidgetHost> WidgetHost = SNew(SViewportWidgetHost, OutUserSlateWidget, bModal)
 			[
@@ -564,21 +541,25 @@ void UUserWidget::AddToViewport(bool bAbsoluteLayout, bool bModal, bool bShowCur
 		UWorld* World = GetWorld();
 		if ( World && World->IsGameWorld() )
 		{
-			if ( UGameViewportClient* Viewport = World->GetGameViewport() )
+			if ( UGameViewportClient* ViewportClient = World->GetGameViewport() )
 			{
-				Viewport->AddViewportWidgetContent(WidgetHost);
+				ViewportClient->AddViewportWidgetContent(WidgetHost);
 
 				//TODO UMG this isn't what should manage focus, a higher level window controller, probably the viewport needs to understand
 				// the Widget stack, and the dialog stack.
 				if ( bModal )
 				{
-					TWeakPtr<SViewport> GameViewportWidget = Viewport->GetGameViewport()->GetViewportWidget();
-					if ( GameViewportWidget.IsValid() )
+					if ( FSceneViewport* SceneViewport = ViewportClient->GetGameViewport() )
 					{
-						GameViewportWidget.Pin()->SetWidgetToFocusOnActivate(OutUserSlateWidget);
-						
-						FSlateApplication::Get().SetFocusToGameViewport();
-						FSlateApplication::Get().SetJoystickCaptorToGameViewport();
+						TWeakPtr<SViewport> GameViewportWidget = SceneViewport->GetViewportWidget();
+
+						if ( GameViewportWidget.IsValid() )
+						{
+							GameViewportWidget.Pin()->SetWidgetToFocusOnActivate(OutUserSlateWidget);
+
+							FSlateApplication::Get().SetFocusToGameViewport();
+							FSlateApplication::Get().SetJoystickCaptorToGameViewport();
+						}
 					}
 				}
 			}
@@ -596,22 +577,26 @@ void UUserWidget::RemoveFromViewport()
 		UWorld* World = GetWorld();
 		if ( World && World->IsGameWorld() )
 		{
-			if ( UGameViewportClient* Viewport = World->GetGameViewport() )
+			if ( UGameViewportClient* ViewportClient = World->GetGameViewport() )
 			{
-				Viewport->RemoveViewportWidgetContent(WidgetHost.ToSharedRef());
+				ViewportClient->RemoveViewportWidgetContent(WidgetHost.ToSharedRef());
 
 				if ( WidgetHost->IsModal() )
 				{
-					TWeakPtr<SViewport> GameViewportWidget = Viewport->GetGameViewport()->GetViewportWidget();
-					if ( GameViewportWidget.IsValid() )
+					if ( FSceneViewport* SceneViewport = ViewportClient->GetGameViewport() )
 					{
-						//TODO UMG this isn't what should manage focus, a higher level window controller, probably the viewport needs to understand
-						// the Widget stack, and the dialog stack.
-						GameViewportWidget.Pin()->ClearWidgetToFocusOnActivate();
-						FSlateApplication::Get().ClearKeyboardFocus(EKeyboardFocusCause::SetDirectly);
-
-						FSlateApplication::Get().SetFocusToGameViewport();
-						FSlateApplication::Get().SetJoystickCaptorToGameViewport();
+						TWeakPtr<SViewport> GameViewportWidget = SceneViewport->GetViewportWidget();
+
+						if ( GameViewportWidget.IsValid() )
+						{
+							//TODO UMG this isn't what should manage focus, a higher level window controller, probably the viewport needs to understand
+							// the Widget stack, and the dialog stack.
+							GameViewportWidget.Pin()->ClearWidgetToFocusOnActivate();
+							FSlateApplication::Get().ClearKeyboardFocus(EKeyboardFocusCause::SetDirectly);
+
+							FSlateApplication::Get().SetFocusToGameViewport();
+							FSlateApplication::Get().SetJoystickCaptorToGameViewport();
+						}
 					}
 				}
 			}
@@ -657,19 +642,25 @@ APlayerController* UUserWidget::GetOwningPlayer() const
 	return PlayerContext.IsValid() ? PlayerContext.GetPlayerController() : NULL;
 }
 
-void UUserWidget::SetOffsetInViewport(FVector2D DesiredOffset)
+void UUserWidget::SetPositionInViewport(FVector2D Position)
 {
-	ViewportOffsets.Left = DesiredOffset.X;
-	ViewportOffsets.Top = DesiredOffset.Y;
+	float Scale = UWidgetLayoutLibrary::GetViewportScale(this);
+
+	ViewportOffsets.Left = Position.X * ( 1.0f / Scale );
+	ViewportOffsets.Top = Position.Y * ( 1.0f / Scale );
+
+	ViewportAnchors = FAnchors(0, 0);
 }
 
 void UUserWidget::SetDesiredSizeInViewport(FVector2D DesiredSize)
 {
 	ViewportOffsets.Right = DesiredSize.X;
 	ViewportOffsets.Bottom = DesiredSize.Y;
+
+	ViewportAnchors = FAnchors(0, 0);
 }
 
-void UUserWidget::SetAnchorsInViewport(FVector2D Anchors)
+void UUserWidget::SetAnchorsInViewport(FAnchors Anchors)
 {
 	ViewportAnchors = Anchors;
 }
@@ -701,7 +692,7 @@ FMargin UUserWidget::GetFullScreenOffset() const
 
 FAnchors UUserWidget::GetViewportAnchors() const
 {
-	return FAnchors(ViewportAnchors.X, ViewportAnchors.Y, ViewportAnchors.X, ViewportAnchors.Y);
+	return ViewportAnchors;
 }
 
 FVector2D UUserWidget::GetFullScreenAlignment() const
diff --git a/Engine/Source/Runtime/UMG/Private/WidgetLayoutLibrary.cpp b/Engine/Source/Runtime/UMG/Private/WidgetLayoutLibrary.cpp
--- a/Engine/Source/Runtime/UMG/Private/WidgetLayoutLibrary.cpp
+++ b/Engine/Source/Runtime/UMG/Private/WidgetLayoutLibrary.cpp
@@ -14,6 +14,38 @@ UWidgetLayoutLibrary::UWidgetLayoutLibrary(const FPostConstructInitializePropert
 {
 }
 
+float UWidgetLayoutLibrary::GetViewportScale(UObject* WorldContextObject)
+{
+	UWorld* World = GEngine->GetWorldFromContextObject(WorldContextObject);
+	if ( World && World->IsGameWorld() )
+	{
+		if ( UGameViewportClient* ViewportClient = World->GetGameViewport() )
+		{
+			FVector2D ViewportSize;
+			ViewportClient->GetViewportSize(ViewportSize);
+			return GetDefault<URendererSettings>(URendererSettings::StaticClass())->GetDPIScaleBasedOnSize(FIntPoint(ViewportSize.X, ViewportSize.Y));
+		}
+	}
+
+	return 1;
+}
+
+FVector2D UWidgetLayoutLibrary::GetViewportSize(UObject* WorldContextObject)
+{
+	UWorld* World = GEngine->GetWorldFromContextObject(WorldContextObject);
+	if ( World && World->IsGameWorld() )
+	{
+		if ( UGameViewportClient* ViewportClient = World->GetGameViewport() )
+		{
+			FVector2D ViewportSize;
+			ViewportClient->GetViewportSize(ViewportSize);
+			return ViewportSize;
+		}
+	}
+
+	return FVector2D(1, 1);
+}
+
 UCanvasPanelSlot* UWidgetLayoutLibrary::SlotAsCanvasSlot(UWidget* ChildWidget)
 {
 	return Cast<UCanvasPanelSlot>(ChildWidget->Slot);
diff --git a/Engine/Source/Runtime/UMG/Public/Blueprint/UserWidget.h b/Engine/Source/Runtime/UMG/Public/Blueprint/UserWidget.h
--- a/Engine/Source/Runtime/UMG/Public/Blueprint/UserWidget.h
+++ b/Engine/Source/Runtime/UMG/Public/Blueprint/UserWidget.h
@@ -99,23 +99,25 @@ public:
 
 	/*  */
 	UFUNCTION(BlueprintCallable, Category="User Interface|Viewport")
-	void AddToViewport(bool bAbsoluteLayout = false, bool bModal = false, bool bShowCursor = false);
+	void AddToViewport(bool bModal = false, bool bShowCursor = false);
 
 	/*  */
 	UFUNCTION(BlueprintCallable, Category="User Interface|Viewport")
 	void RemoveFromViewport();
 
-	/*  */
+	/**
+	 * Sets the widgets position in the viewport.
+	 */
 	UFUNCTION(BlueprintCallable, Category="User Interface|Viewport")
-	void SetOffsetInViewport(FVector2D DesiredOffset);
+	void SetPositionInViewport(FVector2D Position);
 
 	/*  */
 	UFUNCTION(BlueprintCallable, Category="User Interface|Viewport")
-	void SetDesiredSizeInViewport(FVector2D DesiredSize);
+	void SetDesiredSizeInViewport(FVector2D Size);
 
 	/*  */
 	UFUNCTION(BlueprintCallable, Category="User Interface|Viewport")
-	void SetAnchorsInViewport(FVector2D Anchors);
+	void SetAnchorsInViewport(FAnchors Anchors);
 
 	/*  */
 	UFUNCTION(BlueprintCallable, Category="User Interface|Viewport")
@@ -252,7 +254,7 @@ public:
 	UWidget* GetWidgetHandle(TSharedRef<SWidget> InWidget);
 
 	/** Creates a fullscreen host widget, that wraps this widget. */
-	TSharedRef<SWidget> MakeViewportWidget(bool bAbsoluteLayout, bool bModal, bool bShowCursor, TSharedPtr<SWidget>& UserSlateWidget);
+	TSharedRef<SWidget> MakeViewportWidget(bool bModal, bool bShowCursor, TSharedPtr<SWidget>& UserSlateWidget);
 
 	/** @returns The root UObject widget wrapper */
 	UWidget* GetRootWidgetComponent();
@@ -279,44 +281,6 @@ public:
 	UPROPERTY(EditAnywhere, BlueprintReadWrite, AdvancedDisplay, Category="Behavior")
 	bool bSupportsKeyboardFocus;
 
-	UPROPERTY()
-	FMargin Padding;
-
-	/** How much space this slot should occupy in the direction of the panel. */
-	UPROPERTY()
-	FSlateChildSize Size;
-
-	/**
-	* Horizontal pivot position
-	*  Given a top aligned slot, where '+' represents the
-	*  anchor point defined by PositionAttr.
-	*
-	*   Left				Center				Right
-	+ _ _ _ _            _ _ + _ _          _ _ _ _ +
-	|		  |		   | 		   |	  |		    |
-	| _ _ _ _ |        | _ _ _ _ _ |	  | _ _ _ _ |
-	*
-	*  Note: FILL is NOT supported in absolute layout
-	*/
-	UPROPERTY()
-	TEnumAsByte<EHorizontalAlignment> HorizontalAlignment;
-
-	/**
-	* Vertical pivot position
-	*   Given a left aligned slot, where '+' represents the
-	*   anchor point defined by PositionAttr.
-	*
-	*   Top					Center			  Bottom
-	*	+_ _ _ _ _		 _ _ _ _ _		 _ _ _ _ _
-	*	|         |		| 		  |		|		  |
-	*	|         |     +		  |		|		  |
-	*	| _ _ _ _ |		| _ _ _ _ |		+ _ _ _ _ |
-	*
-	*  Note: FILL is NOT supported in absolute layout
-	*/
-	UPROPERTY()
-	TEnumAsByte<EVerticalAlignment> VerticalAlignment;
-
 	/** The components contained in this user widget. */
 	UPROPERTY(Transient)
 	TArray<UWidget*> Components;
@@ -342,7 +306,7 @@ protected:
 	int32 GetFullScreenZOrder() const;
 
 private:
-	FVector2D ViewportAnchors;
+	FAnchors ViewportAnchors;
 	FMargin ViewportOffsets;
 	FVector2D ViewportAlignment;
 	int32 ViewportZOrder;
diff --git a/Engine/Source/Runtime/UMG/Public/Blueprint/WidgetLayoutLibrary.h b/Engine/Source/Runtime/UMG/Public/Blueprint/WidgetLayoutLibrary.h
--- a/Engine/Source/Runtime/UMG/Public/Blueprint/WidgetLayoutLibrary.h
+++ b/Engine/Source/Runtime/UMG/Public/Blueprint/WidgetLayoutLibrary.h
@@ -9,24 +9,32 @@ class UWidgetLayoutLibrary : public UBlueprintFunctionLibrary
 {
 	GENERATED_UCLASS_BODY()
 
+	/**  */
+	UFUNCTION(BlueprintPure, Category="Viewport", meta=( HidePin="WorldContextObject", DefaultToSelf="WorldContextObject" ))
+	static float GetViewportScale(UObject* WorldContextObject);
+
+	/**  */
+	UFUNCTION(BlueprintPure, Category="Viewport", meta=( HidePin="WorldContextObject", DefaultToSelf="WorldContextObject" ))
+	static FVector2D GetViewportSize(UObject* WorldContextObject);
+
 	/**
 	 * Gets the slot object on the child widget as a Canvas Slot, allowing you to manipulate layout information.
 	 * @param Widget The child widget of a canvas panel.
 	 */
-	UFUNCTION(BlueprintPure, Category="Layout")
+	UFUNCTION(BlueprintPure, Category="Slot")
 	static UCanvasPanelSlot* SlotAsCanvasSlot(UWidget* Widget);
 
 	/**
 	 * Gets the slot object on the child widget as a Grid Slot, allowing you to manipulate layout information.
 	 * @param Widget The child widget of a grid panel.
 	 */
-	UFUNCTION(BlueprintPure, Category="Layout")
+	UFUNCTION(BlueprintPure, Category="Slot")
 	static UGridSlot* SlotAsGridSlot(UWidget* Widget);
 
 	/**
 	 * Gets the slot object on the child widget as a Uniform Grid Slot, allowing you to manipulate layout information.
 	 * @param Widget The child widget of a uniform grid panel.
 	 */
-	UFUNCTION(BlueprintPure, Category="Layout")
+	UFUNCTION(BlueprintPure, Category="Slot")
 	static UUniformGridSlot* SlotAsUniformGridSlot(UWidget* Widget);
 };
diff --git a/Engine/Source/Runtime/UMG/Public/Slate/Anchors.h b/Engine/Source/Runtime/UMG/Public/Slate/Anchors.h
--- a/Engine/Source/Runtime/UMG/Public/Slate/Anchors.h
+++ b/Engine/Source/Runtime/UMG/Public/Slate/Anchors.h
@@ -45,9 +45,9 @@ public:
 	{ }
 
 	/** Construct Anchors where the spacing on each side is individually specified. */
-	FAnchors(float InLeft, float InTop, float InRight, float InBottom)
-		: Minimum(InLeft, InTop)
-		, Maximum(InRight, InBottom)
+	FAnchors(float MinX, float MinY, float MaxX, float MaxY)
+		: Minimum(MinX, MinY)
+		, Maximum(MaxX, MaxY)
 	{ }
 
 	/** Returns true if the anchors represent a stretch along the vertical axis */
