UMG - Adding a shortcut to get the owning player pawn to the user widget.

diff --git a/Engine/Source/Runtime/UMG/Private/UserWidget.cpp b/Engine/Source/Runtime/UMG/Private/UserWidget.cpp
--- a/Engine/Source/Runtime/UMG/Private/UserWidget.cpp
+++ b/Engine/Source/Runtime/UMG/Private/UserWidget.cpp
@@ -633,13 +633,23 @@ const FLocalPlayerContext& UUserWidget::GetPlayerContext() const
 
 ULocalPlayer* UUserWidget::GetOwningLocalPlayer() const
 {
-	APlayerController* PC = PlayerContext.IsValid() ? PlayerContext.GetPlayerController() : NULL;
-	return PC ? Cast<ULocalPlayer>(PC->Player) : NULL;
+	APlayerController* PC = PlayerContext.IsValid() ? PlayerContext.GetPlayerController() : nullptr;
+	return PC ? Cast<ULocalPlayer>(PC->Player) : nullptr;
 }
 
 APlayerController* UUserWidget::GetOwningPlayer() const
 {
-	return PlayerContext.IsValid() ? PlayerContext.GetPlayerController() : NULL;
+	return PlayerContext.IsValid() ? PlayerContext.GetPlayerController() : nullptr;
+}
+
+class APawn* UUserWidget::GetOwningPlayerPawn() const
+{
+	if ( APlayerController* PC = GetOwningPlayer() )
+	{
+		return PC->GetPawn();
+	}
+
+	return nullptr;
 }
 
 void UUserWidget::SetPositionInViewport(FVector2D Position)
diff --git a/Engine/Source/Runtime/UMG/Public/Blueprint/UserWidget.h b/Engine/Source/Runtime/UMG/Public/Blueprint/UserWidget.h
--- a/Engine/Source/Runtime/UMG/Public/Blueprint/UserWidget.h
+++ b/Engine/Source/Runtime/UMG/Public/Blueprint/UserWidget.h
@@ -148,6 +148,10 @@ public:
 	UFUNCTION(BlueprintCallable, Category="Player")
 	class APlayerController* GetOwningPlayer() const;
 
+	/** Gets the player pawn associated with this UI. */
+	UFUNCTION(BlueprintCallable, Category="Player")
+	class APawn* GetOwningPlayerPawn() const;
+
 	/** Called when the widget is constructed */
 	UFUNCTION(BlueprintNativeEvent, Category="User Interface")
 	void Construct();
