UMG - Cleaning up some category games, and unused functions.

diff --git a/Engine/Source/Runtime/UMG/Private/Components/GridPanel.cpp b/Engine/Source/Runtime/UMG/Private/Components/GridPanel.cpp
--- a/Engine/Source/Runtime/UMG/Private/Components/GridPanel.cpp
+++ b/Engine/Source/Runtime/UMG/Private/Components/GridPanel.cpp
@@ -81,11 +81,6 @@ void UGridPanel::SynchronizeProperties()
 	}
 }
 
-UGridSlot* UGridPanel::GetGridChildSlot(UWidget* ChildWidget)
-{
-	return Cast<UGridSlot>(ChildWidget->Slot);
-}
-
 #if WITH_EDITOR
 
 const FSlateBrush* UGridPanel::GetEditorIcon()
diff --git a/Engine/Source/Runtime/UMG/Public/Components/Button.h b/Engine/Source/Runtime/UMG/Public/Components/Button.h
--- a/Engine/Source/Runtime/UMG/Public/Components/Button.h
+++ b/Engine/Source/Runtime/UMG/Public/Components/Button.h
@@ -67,7 +67,11 @@ public:
 	UFUNCTION(BlueprintCallable, Category="Button|Appearance")
 	void SetBackgroundColor(FLinearColor InBackgroundColor);
 
-	/** @return true if the user is actively pressing the button otherwise false.  For detecting Clicks, use the OnClicked event. */
+	/**
+	 * Returns true if the user is actively pressing the button.  Do not use this for detecting 'Clicks', use the OnClicked event instead.
+	 *
+	 * @return true if the user is actively pressing the button otherwise false.
+	 */
 	UFUNCTION(BlueprintCallable, Category="Button")
 	bool IsPressed() const;
 
diff --git a/Engine/Source/Runtime/UMG/Public/Components/GridPanel.h b/Engine/Source/Runtime/UMG/Public/Components/GridPanel.h
--- a/Engine/Source/Runtime/UMG/Public/Components/GridPanel.h
+++ b/Engine/Source/Runtime/UMG/Public/Components/GridPanel.h
@@ -20,13 +20,6 @@ public:
 	UPROPERTY(EditDefaultsOnly, Category="Fill Rules")
 	TArray<float> RowFill;
 
-	/**
-	 * Gets the grid slot object on the child widget inside this panel, allowing you to manipulate layout information.
-	 * @param ChildWidget The child widget.
-	 */
-	UFUNCTION(BlueprintCallable, Category="Slot")
-	static UGridSlot* GetGridChildSlot(UWidget* ChildWidget);
-
 public:
 
 	// UWidget interface
diff --git a/Engine/Source/Runtime/UMG/Public/Components/SizeBox.h b/Engine/Source/Runtime/UMG/Public/Components/SizeBox.h
--- a/Engine/Source/Runtime/UMG/Public/Components/SizeBox.h
+++ b/Engine/Source/Runtime/UMG/Public/Components/SizeBox.h
@@ -102,27 +102,27 @@ protected:
 
 
 	/** When specified, ignore the content's desired size and report the WidthOverride as the Box's desired width. */
-	UPROPERTY(EditDefaultsOnly, Category=Layout, meta=( editcondition="bOverride_WidthOverride" ))
+	UPROPERTY(EditDefaultsOnly, Category="Child Layout", meta=( editcondition="bOverride_WidthOverride" ))
 	float WidthOverride;
 
 	/** When specified, ignore the content's desired size and report the HeightOverride as the Box's desired height. */
-	UPROPERTY(EditDefaultsOnly, Category=Layout, meta=( editcondition="bOverride_HeightOverride" ))
+	UPROPERTY(EditDefaultsOnly, Category="Child Layout", meta=( editcondition="bOverride_HeightOverride" ))
 	float HeightOverride;
 
 	/** When specified, will report the MinDesiredWidth if larger than the content's desired width. */
-	UPROPERTY(EditDefaultsOnly, Category=Layout, meta=( editcondition="bOverride_MinDesiredWidth" ))
+	UPROPERTY(EditDefaultsOnly, Category="Child Layout", meta=( editcondition="bOverride_MinDesiredWidth" ))
 	float MinDesiredWidth;
 
 	/** When specified, will report the MinDesiredHeight if larger than the content's desired height. */
-	UPROPERTY(EditDefaultsOnly, Category=Layout, meta=( editcondition="bOverride_MinDesiredHeight" ))
+	UPROPERTY(EditDefaultsOnly, Category="Child Layout", meta=( editcondition="bOverride_MinDesiredHeight" ))
 	float MinDesiredHeight;
 
 	/** When specified, will report the MaxDesiredWidth if smaller than the content's desired width. */
-	UPROPERTY(EditDefaultsOnly, Category=Layout, meta=( editcondition="bOverride_MaxDesiredWidth" ))
+	UPROPERTY(EditDefaultsOnly, Category="Child Layout", meta=( editcondition="bOverride_MaxDesiredWidth" ))
 	float MaxDesiredWidth;
 
 	/** When specified, will report the MaxDesiredHeight if smaller than the content's desired height. */
-	UPROPERTY(EditDefaultsOnly, Category=Layout, meta=( editcondition="bOverride_MaxDesiredHeight" ))
+	UPROPERTY(EditDefaultsOnly, Category="Child Layout", meta=( editcondition="bOverride_MaxDesiredHeight" ))
 	float MaxDesiredHeight;
 
 protected:
diff --git a/Engine/Source/Runtime/UMG/Public/Components/UniformGridPanel.h b/Engine/Source/Runtime/UMG/Public/Components/UniformGridPanel.h
--- a/Engine/Source/Runtime/UMG/Public/Components/UniformGridPanel.h
+++ b/Engine/Source/Runtime/UMG/Public/Components/UniformGridPanel.h
@@ -13,15 +13,15 @@ class UMG_API UUniformGridPanel : public UPanelWidget
 public:
 
 	/** Padding given to each slot */
-	UPROPERTY(EditDefaultsOnly, Category=Layout)
+	UPROPERTY(EditDefaultsOnly, Category="Child Layout")
 	FMargin SlotPadding;
 
 	/** The minimum desired width of the slots */
-	UPROPERTY(EditDefaultsOnly, Category=Layout)
+	UPROPERTY(EditDefaultsOnly, Category="Child Layout")
 	float MinDesiredSlotWidth;
 
 	/** The minimum desired height of the slots */
-	UPROPERTY(EditDefaultsOnly, Category=Layout)
+	UPROPERTY(EditDefaultsOnly, Category="Child Layout")
 	float MinDesiredSlotHeight;
 
 public:
