UMG - Simplifying the conversion to native FReplys in UMG, no longer need to call ToReply, since the native reply has already been constructed, we just return it.

diff --git a/Engine/Source/Runtime/UMG/Private/Components/Border.cpp b/Engine/Source/Runtime/UMG/Private/Components/Border.cpp
--- a/Engine/Source/Runtime/UMG/Private/Components/Border.cpp
+++ b/Engine/Source/Runtime/UMG/Private/Components/Border.cpp
@@ -145,7 +145,7 @@ FReply UBorder::HandleMouseButtonDown(const FGeometry& Geometry, const FPointerE
 {
 	if ( OnMouseButtonDownEvent.IsBound() )
 	{
-		return OnMouseButtonDownEvent.Execute(Geometry, MouseEvent).ToReply( MyBorder.ToSharedRef() );
+		return OnMouseButtonDownEvent.Execute(Geometry, MouseEvent).NativeReply;
 	}
 
 	return FReply::Unhandled();
@@ -155,7 +155,7 @@ FReply UBorder::HandleMouseButtonUp(const FGeometry& Geometry, const FPointerEve
 {
 	if ( OnMouseButtonUpEvent.IsBound() )
 	{
-		return OnMouseButtonUpEvent.Execute(Geometry, MouseEvent).ToReply( MyBorder.ToSharedRef() );
+		return OnMouseButtonUpEvent.Execute(Geometry, MouseEvent).NativeReply;
 	}
 
 	return FReply::Unhandled();
@@ -165,7 +165,7 @@ FReply UBorder::HandleMouseMove(const FGeometry& Geometry, const FPointerEvent&
 {
 	if ( OnMouseMoveEvent.IsBound() )
 	{
-		return OnMouseMoveEvent.Execute(Geometry, MouseEvent).ToReply( MyBorder.ToSharedRef() );
+		return OnMouseMoveEvent.Execute(Geometry, MouseEvent).NativeReply;
 	}
 
 	return FReply::Unhandled();
@@ -175,7 +175,7 @@ FReply UBorder::HandleMouseDoubleClick(const FGeometry& Geometry, const FPointer
 {
 	if ( OnMouseDoubleClickEvent.IsBound() )
 	{
-		return OnMouseDoubleClickEvent.Execute(Geometry, MouseEvent).ToReply( MyBorder.ToSharedRef() );
+		return OnMouseDoubleClickEvent.Execute(Geometry, MouseEvent).NativeReply;
 	}
 
 	return FReply::Unhandled();
diff --git a/Engine/Source/Runtime/UMG/Private/Components/Image.cpp b/Engine/Source/Runtime/UMG/Private/Components/Image.cpp
--- a/Engine/Source/Runtime/UMG/Private/Components/Image.cpp
+++ b/Engine/Source/Runtime/UMG/Private/Components/Image.cpp
@@ -142,7 +142,7 @@ FReply UImage::HandleMouseButtonDown(const FGeometry& Geometry, const FPointerEv
 {
 	if ( OnMouseButtonDownEvent.IsBound() )
 	{
-		return OnMouseButtonDownEvent.Execute(Geometry, MouseEvent).ToReply(MyImage.ToSharedRef());
+		return OnMouseButtonDownEvent.Execute(Geometry, MouseEvent).NativeReply;
 	}
 
 	return FReply::Unhandled();
diff --git a/Engine/Source/Runtime/UMG/Private/Slate/SObjectWidget.cpp b/Engine/Source/Runtime/UMG/Private/Slate/SObjectWidget.cpp
--- a/Engine/Source/Runtime/UMG/Private/Slate/SObjectWidget.cpp
+++ b/Engine/Source/Runtime/UMG/Private/Slate/SObjectWidget.cpp
@@ -72,11 +72,7 @@ FReply SObjectWidget::OnKeyboardFocusReceived(const FGeometry& MyGeometry, const
 {
 	if ( WidgetObject && !WidgetObject->IsDesignTime() )
 	{
-		TSharedPtr<SWidget> SlateWidget = WidgetObject->GetCachedWidget();
-		if ( SlateWidget.IsValid() )
-		{
-			return WidgetObject->OnKeyboardFocusReceived(MyGeometry, InKeyboardFocusEvent).ToReply(SlateWidget.ToSharedRef());
-		}
+		return WidgetObject->OnKeyboardFocusReceived(MyGeometry, InKeyboardFocusEvent).NativeReply;
 	}
 
 	return FReply::Unhandled();
@@ -97,10 +93,9 @@ void SObjectWidget::OnKeyboardFocusChanging(const FWeakWidgetPath& PreviousFocus
 
 FReply SObjectWidget::OnKeyChar(const FGeometry& MyGeometry, const FCharacterEvent& InCharacterEvent)
 {
-	TSharedPtr<SWidget> SlateWidget = WidgetObject->GetCachedWidget();
-	if ( SlateWidget.IsValid() )
+	if ( WidgetObject && !WidgetObject->IsDesignTime() )
 	{
-		return WidgetObject->OnKeyChar(MyGeometry, InCharacterEvent).ToReply(SlateWidget.ToSharedRef());
+		return WidgetObject->OnKeyChar(MyGeometry, InCharacterEvent).NativeReply;
 	}
 
 	return FReply::Unhandled();
@@ -108,10 +103,9 @@ FReply SObjectWidget::OnKeyChar(const FGeometry& MyGeometry, const FCharacterEve
 
 FReply SObjectWidget::OnPreviewKeyDown(const FGeometry& MyGeometry, const FKeyboardEvent& InKeyboardEvent)
 {
-	TSharedPtr<SWidget> SlateWidget = WidgetObject->GetCachedWidget();
-	if ( SlateWidget.IsValid() )
+	if ( WidgetObject && !WidgetObject->IsDesignTime() )
 	{
-		return WidgetObject->OnPreviewKeyDown(MyGeometry, InKeyboardEvent).ToReply(SlateWidget.ToSharedRef());
+		return WidgetObject->OnPreviewKeyDown(MyGeometry, InKeyboardEvent).NativeReply;
 	}
 
 	return FReply::Unhandled();
@@ -119,10 +113,9 @@ FReply SObjectWidget::OnPreviewKeyDown(const FGeometry& MyGeometry, const FKeybo
 
 FReply SObjectWidget::OnKeyDown(const FGeometry& MyGeometry, const FKeyboardEvent& InKeyboardEvent)
 {
-	TSharedPtr<SWidget> SlateWidget = WidgetObject->GetCachedWidget();
-	if ( SlateWidget.IsValid() )
+	if ( WidgetObject && !WidgetObject->IsDesignTime() )
 	{
-		return WidgetObject->OnKeyDown(MyGeometry, InKeyboardEvent).ToReply(SlateWidget.ToSharedRef());
+		return WidgetObject->OnKeyDown(MyGeometry, InKeyboardEvent).NativeReply;
 	}
 
 	return FReply::Unhandled();
@@ -130,10 +123,9 @@ FReply SObjectWidget::OnKeyDown(const FGeometry& MyGeometry, const FKeyboardEven
 
 FReply SObjectWidget::OnKeyUp(const FGeometry& MyGeometry, const FKeyboardEvent& InKeyboardEvent)
 {
-	TSharedPtr<SWidget> SlateWidget = WidgetObject->GetCachedWidget();
-	if ( SlateWidget.IsValid() )
+	if ( WidgetObject && !WidgetObject->IsDesignTime() )
 	{
-		return WidgetObject->OnKeyUp(MyGeometry, InKeyboardEvent).ToReply(SlateWidget.ToSharedRef());
+		return WidgetObject->OnKeyUp(MyGeometry, InKeyboardEvent).NativeReply;
 	}
 
 	return FReply::Unhandled();
@@ -141,10 +133,9 @@ FReply SObjectWidget::OnKeyUp(const FGeometry& MyGeometry, const FKeyboardEvent&
 
 FReply SObjectWidget::OnMouseButtonDown(const FGeometry& MyGeometry, const FPointerEvent& MouseEvent)
 {
-	TSharedPtr<SWidget> SlateWidget = WidgetObject->GetCachedWidget();
-	if ( SlateWidget.IsValid() )
+	if ( WidgetObject && !WidgetObject->IsDesignTime() )
 	{
-		return WidgetObject->OnMouseButtonDown(MyGeometry, MouseEvent).ToReply(SlateWidget.ToSharedRef());
+		return WidgetObject->OnMouseButtonDown(MyGeometry, MouseEvent).NativeReply;
 	}
 
 	return FReply::Unhandled();
@@ -152,10 +143,9 @@ FReply SObjectWidget::OnMouseButtonDown(const FGeometry& MyGeometry, const FPoin
 
 FReply SObjectWidget::OnPreviewMouseButtonDown(const FGeometry& MyGeometry, const FPointerEvent& MouseEvent)
 {
-	TSharedPtr<SWidget> SlateWidget = WidgetObject->GetCachedWidget();
-	if ( SlateWidget.IsValid() )
+	if ( WidgetObject && !WidgetObject->IsDesignTime() )
 	{
-		return WidgetObject->OnPreviewMouseButtonDown(MyGeometry, MouseEvent).ToReply(SlateWidget.ToSharedRef());
+		return WidgetObject->OnPreviewMouseButtonDown(MyGeometry, MouseEvent).NativeReply;
 	}
 
 	return FReply::Unhandled();
@@ -163,10 +153,9 @@ FReply SObjectWidget::OnPreviewMouseButtonDown(const FGeometry& MyGeometry, cons
 
 FReply SObjectWidget::OnMouseButtonUp(const FGeometry& MyGeometry, const FPointerEvent& MouseEvent)
 {
-	TSharedPtr<SWidget> SlateWidget = WidgetObject->GetCachedWidget();
-	if ( SlateWidget.IsValid() )
+	if ( WidgetObject && !WidgetObject->IsDesignTime() )
 	{
-		return WidgetObject->OnMouseButtonUp(MyGeometry, MouseEvent).ToReply(SlateWidget.ToSharedRef());
+		return WidgetObject->OnMouseButtonUp(MyGeometry, MouseEvent).NativeReply;
 	}
 
 	return FReply::Unhandled();
@@ -174,10 +163,9 @@ FReply SObjectWidget::OnMouseButtonUp(const FGeometry& MyGeometry, const FPointe
 
 FReply SObjectWidget::OnMouseMove(const FGeometry& MyGeometry, const FPointerEvent& MouseEvent)
 {
-	TSharedPtr<SWidget> SlateWidget = WidgetObject->GetCachedWidget();
-	if ( SlateWidget.IsValid() )
+	if ( WidgetObject && !WidgetObject->IsDesignTime() )
 	{
-		return WidgetObject->OnMouseMove(MyGeometry, MouseEvent).ToReply(SlateWidget.ToSharedRef());
+		return WidgetObject->OnMouseMove(MyGeometry, MouseEvent).NativeReply;
 	}
 
 	return FReply::Unhandled();
@@ -185,20 +173,25 @@ FReply SObjectWidget::OnMouseMove(const FGeometry& MyGeometry, const FPointerEve
 
 void SObjectWidget::OnMouseEnter(const FGeometry& MyGeometry, const FPointerEvent& MouseEvent)
 {
-	WidgetObject->OnMouseEnter(MyGeometry, MouseEvent);
+	if ( WidgetObject && !WidgetObject->IsDesignTime() )
+	{
+		WidgetObject->OnMouseEnter(MyGeometry, MouseEvent);
+	}
 }
 
 void SObjectWidget::OnMouseLeave(const FPointerEvent& MouseEvent)
 {
-	WidgetObject->OnMouseLeave(MouseEvent);
+	if ( WidgetObject && !WidgetObject->IsDesignTime() )
+	{
+		WidgetObject->OnMouseLeave(MouseEvent);
+	}
 }
 
 FReply SObjectWidget::OnMouseWheel(const FGeometry& MyGeometry, const FPointerEvent& MouseEvent)
 {
-	TSharedPtr<SWidget> SlateWidget = WidgetObject->GetCachedWidget();
-	if ( SlateWidget.IsValid() )
+	if ( WidgetObject && !WidgetObject->IsDesignTime() )
 	{
-		return WidgetObject->OnMouseWheel(MyGeometry, MouseEvent).ToReply(SlateWidget.ToSharedRef());
+		return WidgetObject->OnMouseWheel(MyGeometry, MouseEvent).NativeReply;
 	}
 
 	return FReply::Unhandled();
@@ -212,10 +205,9 @@ FCursorReply SObjectWidget::OnCursorQuery(const FGeometry& MyGeometry, const FPo
 
 FReply SObjectWidget::OnMouseButtonDoubleClick(const FGeometry& MyGeometry, const FPointerEvent& MouseEvent)
 {
-	TSharedPtr<SWidget> SlateWidget = WidgetObject->GetCachedWidget();
-	if ( SlateWidget.IsValid() )
+	if ( WidgetObject && !WidgetObject->IsDesignTime() )
 	{
-		return WidgetObject->OnMouseButtonDoubleClick(MyGeometry, MouseEvent).ToReply(SlateWidget.ToSharedRef());
+		return WidgetObject->OnMouseButtonDoubleClick(MyGeometry, MouseEvent).NativeReply;
 	}
 
 	return FReply::Unhandled();
@@ -294,10 +286,9 @@ void SObjectWidget::OnDragCancelled(const FDragDropEvent& DragDropEvent, UDragDr
 
 FReply SObjectWidget::OnControllerButtonPressed(const FGeometry& MyGeometry, const FControllerEvent& ControllerEvent)
 {
-	TSharedPtr<SWidget> SlateWidget = WidgetObject->GetCachedWidget();
-	if ( SlateWidget.IsValid() )
+	if ( WidgetObject && !WidgetObject->IsDesignTime() )
 	{
-		return WidgetObject->OnControllerButtonPressed(MyGeometry, ControllerEvent).ToReply(SlateWidget.ToSharedRef());
+		return WidgetObject->OnControllerButtonPressed(MyGeometry, ControllerEvent).NativeReply;
 	}
 
 	return FReply::Unhandled();
@@ -305,10 +296,9 @@ FReply SObjectWidget::OnControllerButtonPressed(const FGeometry& MyGeometry, con
 
 FReply SObjectWidget::OnControllerButtonReleased(const FGeometry& MyGeometry, const FControllerEvent& ControllerEvent)
 {
-	TSharedPtr<SWidget> SlateWidget = WidgetObject->GetCachedWidget();
-	if ( SlateWidget.IsValid() )
+	if ( WidgetObject && !WidgetObject->IsDesignTime() )
 	{
-		return WidgetObject->OnControllerButtonReleased(MyGeometry, ControllerEvent).ToReply(SlateWidget.ToSharedRef());
+		return WidgetObject->OnControllerButtonReleased(MyGeometry, ControllerEvent).NativeReply;
 	}
 
 	return FReply::Unhandled();
@@ -316,10 +306,9 @@ FReply SObjectWidget::OnControllerButtonReleased(const FGeometry& MyGeometry, co
 
 FReply SObjectWidget::OnControllerAnalogValueChanged(const FGeometry& MyGeometry, const FControllerEvent& ControllerEvent)
 {
-	TSharedPtr<SWidget> SlateWidget = WidgetObject->GetCachedWidget();
-	if ( SlateWidget.IsValid() )
+	if ( WidgetObject && !WidgetObject->IsDesignTime() )
 	{
-		return WidgetObject->OnControllerAnalogValueChanged(MyGeometry, ControllerEvent).ToReply(SlateWidget.ToSharedRef());
+		return WidgetObject->OnControllerAnalogValueChanged(MyGeometry, ControllerEvent).NativeReply;
 	}
 
 	return FReply::Unhandled();
@@ -327,10 +316,9 @@ FReply SObjectWidget::OnControllerAnalogValueChanged(const FGeometry& MyGeometry
 
 FReply SObjectWidget::OnTouchGesture(const FGeometry& MyGeometry, const FPointerEvent& GestureEvent)
 {
-	TSharedPtr<SWidget> SlateWidget = WidgetObject->GetCachedWidget();
-	if ( SlateWidget.IsValid() )
+	if ( WidgetObject && !WidgetObject->IsDesignTime() )
 	{
-		return WidgetObject->OnTouchGesture(MyGeometry, GestureEvent).ToReply(SlateWidget.ToSharedRef());
+		return WidgetObject->OnTouchGesture(MyGeometry, GestureEvent).NativeReply;
 	}
 
 	return FReply::Unhandled();
@@ -338,10 +326,9 @@ FReply SObjectWidget::OnTouchGesture(const FGeometry& MyGeometry, const FPointer
 
 FReply SObjectWidget::OnTouchStarted(const FGeometry& MyGeometry, const FPointerEvent& InTouchEvent)
 {
-	TSharedPtr<SWidget> SlateWidget = WidgetObject->GetCachedWidget();
-	if ( SlateWidget.IsValid() )
+	if ( WidgetObject && !WidgetObject->IsDesignTime() )
 	{
-		return WidgetObject->OnTouchStarted(MyGeometry, InTouchEvent).ToReply(SlateWidget.ToSharedRef());
+		return WidgetObject->OnTouchStarted(MyGeometry, InTouchEvent).NativeReply;
 	}
 
 	return FReply::Unhandled();
@@ -349,10 +336,9 @@ FReply SObjectWidget::OnTouchStarted(const FGeometry& MyGeometry, const FPointer
 
 FReply SObjectWidget::OnTouchMoved(const FGeometry& MyGeometry, const FPointerEvent& InTouchEvent)
 {
-	TSharedPtr<SWidget> SlateWidget = WidgetObject->GetCachedWidget();
-	if ( SlateWidget.IsValid() )
+	if ( WidgetObject && !WidgetObject->IsDesignTime() )
 	{
-		return WidgetObject->OnTouchMoved(MyGeometry, InTouchEvent).ToReply(SlateWidget.ToSharedRef());
+		return WidgetObject->OnTouchMoved(MyGeometry, InTouchEvent).NativeReply;
 	}
 
 	return FReply::Unhandled();
@@ -360,10 +346,9 @@ FReply SObjectWidget::OnTouchMoved(const FGeometry& MyGeometry, const FPointerEv
 
 FReply SObjectWidget::OnTouchEnded(const FGeometry& MyGeometry, const FPointerEvent& InTouchEvent)
 {
-	TSharedPtr<SWidget> SlateWidget = WidgetObject->GetCachedWidget();
-	if ( SlateWidget.IsValid() )
+	if ( WidgetObject && !WidgetObject->IsDesignTime() )
 	{
-		return WidgetObject->OnTouchEnded(MyGeometry, InTouchEvent).ToReply(SlateWidget.ToSharedRef());
+		return WidgetObject->OnTouchEnded(MyGeometry, InTouchEvent).NativeReply;
 	}
 
 	return FReply::Unhandled();
@@ -371,10 +356,9 @@ FReply SObjectWidget::OnTouchEnded(const FGeometry& MyGeometry, const FPointerEv
 
 FReply SObjectWidget::OnMotionDetected(const FGeometry& MyGeometry, const FMotionEvent& InMotionEvent)
 {
-	TSharedPtr<SWidget> SlateWidget = WidgetObject->GetCachedWidget();
-	if ( SlateWidget.IsValid() )
+	if ( WidgetObject && !WidgetObject->IsDesignTime() )
 	{
-		return WidgetObject->OnMotionDetected(MyGeometry, InMotionEvent).ToReply(SlateWidget.ToSharedRef());
+		return WidgetObject->OnMotionDetected(MyGeometry, InMotionEvent).NativeReply;
 	}
 
 	return FReply::Unhandled();
diff --git a/Engine/Source/Runtime/UMG/Public/Components/SlateWrapperTypes.h b/Engine/Source/Runtime/UMG/Public/Components/SlateWrapperTypes.h
--- a/Engine/Source/Runtime/UMG/Public/Components/SlateWrapperTypes.h
+++ b/Engine/Source/Runtime/UMG/Public/Components/SlateWrapperTypes.h
@@ -51,28 +51,13 @@ struct FEventReply
 {
 	GENERATED_USTRUCT_BODY()
 
+public:
+
 	FEventReply(bool IsHandled = false)
 	: NativeReply(IsHandled ? FReply::Handled() : FReply::Unhandled())
 	{
 	}
 
-	FReply ToReply(TSharedRef<SWidget> Widget) const
-	{
-		//FReply Reply = FReply::Unhandled();
-
-		//if ( bIsHandled )
-		//{
-		//	Reply = FReply::Handled();
-		//}
-
-		//if ( bCaptureMouse )
-		//{
-		//	Reply = Reply.CaptureMouse(Widget);
-		//}
-
-		return NativeReply;
-	}
-
 	FReply NativeReply;
 };
 
@@ -83,7 +68,7 @@ struct FSlateChildSize
 	GENERATED_USTRUCT_BODY()
 
 	/** The parameter of the size rule. */
-	UPROPERTY(EditAnywhere, BlueprintReadWrite, Category=Appearance, meta=( ClampMin="0", ClampMax="1" ))
+	UPROPERTY(EditAnywhere, BlueprintReadWrite, Category=Appearance, meta=( UIMin="0", UIMax="1" ))
 	float Value;
 
 	/** The sizing rule of the content. */
