UMG - Changing the category on several widgets.

diff --git a/Engine/Source/Runtime/UMG/Private/Components/CircularThrobber.cpp b/Engine/Source/Runtime/UMG/Private/Components/CircularThrobber.cpp
--- a/Engine/Source/Runtime/UMG/Private/Components/CircularThrobber.cpp
+++ b/Engine/Source/Runtime/UMG/Private/Components/CircularThrobber.cpp
@@ -103,7 +103,7 @@ const FSlateBrush* UCircularThrobber::GetEditorIcon()
 
 const FText UCircularThrobber::GetToolboxCategory()
 {
-	return LOCTEXT("Misc", "Misc");
+	return LOCTEXT("Primitive", "Primitive");
 }
 
 #endif
diff --git a/Engine/Source/Runtime/UMG/Private/Components/NativeWidgetHost.cpp b/Engine/Source/Runtime/UMG/Private/Components/NativeWidgetHost.cpp
--- a/Engine/Source/Runtime/UMG/Private/Components/NativeWidgetHost.cpp
+++ b/Engine/Source/Runtime/UMG/Private/Components/NativeWidgetHost.cpp
@@ -54,7 +54,7 @@ const FSlateBrush* UNativeWidgetHost::GetEditorIcon()
 
 const FText UNativeWidgetHost::GetToolboxCategory()
 {
-	return LOCTEXT("Advanced", "Advanced");
+	return LOCTEXT("Primitive", "Primitive");
 }
 
 #endif
diff --git a/Engine/Source/Runtime/UMG/Private/Components/ScrollBar.cpp b/Engine/Source/Runtime/UMG/Private/Components/ScrollBar.cpp
--- a/Engine/Source/Runtime/UMG/Private/Components/ScrollBar.cpp
+++ b/Engine/Source/Runtime/UMG/Private/Components/ScrollBar.cpp
@@ -59,7 +59,7 @@ const FSlateBrush* UScrollBar::GetEditorIcon()
 
 const FText UScrollBar::GetToolboxCategory()
 {
-	return LOCTEXT("Advanced", "Advanced");
+	return LOCTEXT("Primitive", "Primitive");
 }
 
 #endif
diff --git a/Engine/Source/Runtime/UMG/Private/Components/Throbber.cpp b/Engine/Source/Runtime/UMG/Private/Components/Throbber.cpp
--- a/Engine/Source/Runtime/UMG/Private/Components/Throbber.cpp
+++ b/Engine/Source/Runtime/UMG/Private/Components/Throbber.cpp
@@ -114,7 +114,7 @@ const FSlateBrush* UThrobber::GetEditorIcon()
 
 const FText UThrobber::GetToolboxCategory()
 {
-	return LOCTEXT("Misc", "Misc");
+	return LOCTEXT("Primitive", "Primitive");
 }
 
 #endif
diff --git a/Engine/Source/Runtime/UMG/Private/Components/Viewport.cpp b/Engine/Source/Runtime/UMG/Private/Components/Viewport.cpp
--- a/Engine/Source/Runtime/UMG/Private/Components/Viewport.cpp
+++ b/Engine/Source/Runtime/UMG/Private/Components/Viewport.cpp
@@ -628,7 +628,7 @@ const FSlateBrush* UViewport::GetEditorIcon()
 
 const FText UViewport::GetToolboxCategory()
 {
-	return LOCTEXT("Advanced", "Advanced");
+	return LOCTEXT("Primitive", "Primitive");
 }
 
 #endif
diff --git a/Engine/Source/Runtime/UMG/Public/Components/Viewport.h b/Engine/Source/Runtime/UMG/Public/Components/Viewport.h
--- a/Engine/Source/Runtime/UMG/Public/Components/Viewport.h
+++ b/Engine/Source/Runtime/UMG/Public/Components/Viewport.h
@@ -238,7 +238,7 @@ protected:
 /**
  * 
  */
-UCLASS(ClassGroup=UserInterface)
+UCLASS(Experimental, ClassGroup=UserInterface)
 class UMG_API UViewport : public UContentWidget
 {
 	GENERATED_UCLASS_BODY()
