UMG - Moving combobox string to input cat.

diff --git a/Engine/Source/Runtime/UMG/Private/Components/ComboBoxString.cpp b/Engine/Source/Runtime/UMG/Private/Components/ComboBoxString.cpp
--- a/Engine/Source/Runtime/UMG/Private/Components/ComboBoxString.cpp
+++ b/Engine/Source/Runtime/UMG/Private/Components/ComboBoxString.cpp
@@ -152,7 +152,7 @@ const FSlateBrush* UComboBoxString::GetEditorIcon()
 
 const FText UComboBoxString::GetToolboxCategory()
 {
-	return LOCTEXT("Misc", "Misc");
+	return LOCTEXT("Input", "Input");
 }
 
 #endif
