UMG - Fixing a bug on the auto sizing fullscreen code, we no longer always apply the auto sizing logic if the anchor points are stretching.

diff --git a/Engine/Source/Runtime/UMG/Private/UserWidget.cpp b/Engine/Source/Runtime/UMG/Private/UserWidget.cpp
--- a/Engine/Source/Runtime/UMG/Private/UserWidget.cpp
+++ b/Engine/Source/Runtime/UMG/Private/UserWidget.cpp
@@ -687,8 +687,9 @@ void UUserWidget::SetZOrderInViewport(int32 ZOrder)
 
 FMargin UUserWidget::GetFullScreenOffset() const
 {
+	// If the size is zero, and we're not stretched, then use the desired size.
 	FVector2D FinalSize = FVector2D(ViewportOffsets.Right, ViewportOffsets.Bottom);
-	if ( FinalSize.IsZero() )
+	if ( FinalSize.IsZero() && !ViewportAnchors.IsStretchedVertical() && !ViewportAnchors.IsStretchedHorizontal() )
 	{
 		TSharedPtr<SWidget> CachedWidget = GetCachedWidget();
 		if ( CachedWidget.IsValid() )
