UMG - Fixing a crash on copy/paste of borders.

diff --git a/Engine/Source/Runtime/UMG/Private/Components/BorderSlot.cpp b/Engine/Source/Runtime/UMG/Private/Components/BorderSlot.cpp
--- a/Engine/Source/Runtime/UMG/Private/Components/BorderSlot.cpp
+++ b/Engine/Source/Runtime/UMG/Private/Components/BorderSlot.cpp
@@ -50,9 +50,12 @@ void UBorderSlot::SetVerticalAlignment(EVerticalAlignment InVerticalAlignment)
 
 void UBorderSlot::SynchronizeProperties()
 {
-	SetPadding(Padding);
-	SetHorizontalAlignment(HorizontalAlignment);
-	SetVerticalAlignment(VerticalAlignment);
+	if ( Border.IsValid() )
+	{
+		SetPadding(Padding);
+		SetHorizontalAlignment(HorizontalAlignment);
+		SetVerticalAlignment(VerticalAlignment);
+	}
 }
 
 #if WITH_EDITOR
