UMG - Adding a way to return a 'null' brush to draw nothing.

diff --git a/Engine/Source/Runtime/UMG/Public/Blueprint/WidgetBlueprintLibrary.h b/Engine/Source/Runtime/UMG/Public/Blueprint/WidgetBlueprintLibrary.h
--- a/Engine/Source/Runtime/UMG/Public/Blueprint/WidgetBlueprintLibrary.h
+++ b/Engine/Source/Runtime/UMG/Public/Blueprint/WidgetBlueprintLibrary.h
@@ -111,4 +111,12 @@ class UWidgetBlueprintLibrary : public UBlueprintFunctionLibrary
 	 */
 	UFUNCTION(BlueprintPure, Category="Widget|Brush")
 	static FSlateBrush MakeBrushFromMaterial(UMaterialInterface* Material, int32 Width = 32, int32 Height = 32);
+
+	/**
+	 * Creates a Slate Brush that wont draw anything, the "Null Brush".
+	 *
+	 * @return A new slate brush that wont draw anything.
+	 */
+	UFUNCTION(BlueprintPure, Category="Widget|Brush")
+	static FSlateBrush NoResourceBrush();
 };
