UMG - Adding a way to return a 'null' brush to draw nothing.

diff --git a/Engine/Source/Runtime/UMG/Private/WidgetBlueprintLibrary.cpp b/Engine/Source/Runtime/UMG/Private/WidgetBlueprintLibrary.cpp
--- a/Engine/Source/Runtime/UMG/Private/WidgetBlueprintLibrary.cpp
+++ b/Engine/Source/Runtime/UMG/Private/WidgetBlueprintLibrary.cpp
@@ -227,4 +227,9 @@ FSlateBrush UWidgetBlueprintLibrary::MakeBrushFromMaterial(UMaterialInterface* M
 	return FSlateNoResource();
 }
 
+FSlateBrush UWidgetBlueprintLibrary::NoResourceBrush()
+{
+	return FSlateNoResource();
+}
+
 #undef LOCTEXT_NAMESPACE
\ No newline at end of file
