UMG - Deprecated style template on buttons

diff --git a/Engine/Source/Runtime/Core/Public/UObject/ObjectVersion.h b/Engine/Source/Runtime/Core/Public/UObject/ObjectVersion.h
--- a/Engine/Source/Runtime/Core/Public/UObject/ObjectVersion.h
+++ b/Engine/Source/Runtime/Core/Public/UObject/ObjectVersion.h
@@ -606,6 +606,8 @@ enum EUnrealEngineObjectUE4Version
 	VER_UE4_LIGHTMAP_MESH_BUILD_SETTINGS,
 	// Rename SM3 to ES3_1 and updates featurelevel material node selector
 	VER_UE4_RENAME_SM3_TO_ES3_1,
+	// Deprecated separate style assets for use in UMG
+	VER_UE4_DEPRECATE_UMG_STYLE_ASSETS,
 
 	// -----<new versions can be added before this line>-------------------------------------------------
 	// - this needs to be the last line (see note below)
diff --git a/Engine/Source/Runtime/UMG/Private/Components/Button.cpp b/Engine/Source/Runtime/UMG/Private/Components/Button.cpp
--- a/Engine/Source/Runtime/UMG/Private/Components/Button.cpp
+++ b/Engine/Source/Runtime/UMG/Private/Components/Button.cpp
@@ -35,12 +35,6 @@ void UButton::ReleaseNativeWidget()
 
 TSharedRef<SWidget> UButton::RebuildWidget()
 {
-	const FButtonStyle* StylePtr = ( Style != nullptr ) ? Style->GetStyle<FButtonStyle>() : nullptr;
-	if ( StylePtr != nullptr )
-	{
-		ButtonStyle->ButtonStyle = *StylePtr;
-	}
-
 	MyButton = SNew(SButton)
 		.ButtonStyle(&ButtonStyle->ButtonStyle)
 		.ClickMethod(ClickMethod)
@@ -140,6 +134,17 @@ void UButton::PostLoad()
 			}
 		}
 	}
+
+	if( GetLinkerUE4Version() < VER_UE4_DEPRECATE_UMG_STYLE_ASSETS && Style_DEPRECATED != nullptr )
+	{
+		const FButtonStyle* StylePtr = Style_DEPRECATED->GetStyle<FButtonStyle>();
+		if(StylePtr != nullptr)
+		{
+			ButtonStyle->ButtonStyle = *StylePtr;
+		}
+
+		Style_DEPRECATED = nullptr;
+	}
 }
 
 FReply UButton::SlateHandleClicked()
@@ -161,18 +166,6 @@ const FText UButton::GetToolboxCategory()
 	return LOCTEXT("Common", "Common");
 }
 
-bool UButton::CanEditChange(const UProperty* Property) const
-{
-	static const FName ButtonStyleName(TEXT("ButtonStyle"));
-
-	if ( Property->GetFName() == ButtonStyleName )
-	{
-		return  Style == nullptr;
-	}
-
-	return true;
-}
-
 #endif
 
 /////////////////////////////////////////////////////
diff --git a/Engine/Source/Runtime/UMG/Public/Components/Button.h b/Engine/Source/Runtime/UMG/Public/Components/Button.h
--- a/Engine/Source/Runtime/UMG/Public/Components/Button.h
+++ b/Engine/Source/Runtime/UMG/Public/Components/Button.h
@@ -18,8 +18,8 @@ class UMG_API UButton : public UContentWidget
 
 public:
 	/** The template style asset, used to seed the mutable instance of the style. */
-	UPROPERTY(EditDefaultsOnly, AdvancedDisplay, DisplayName="Template Style", Category=Appearance, meta=( DesignerRebuild=true ))
-	USlateWidgetStyleAsset* Style;
+	UPROPERTY()
+	USlateWidgetStyleAsset* Style_DEPRECATED;
 
 	/** The button style used at runtime by the slate button */
 	UPROPERTY(VisibleAnywhere, Instanced, Category=Appearance )
@@ -90,7 +90,6 @@ public:
 #if WITH_EDITOR
 	virtual const FSlateBrush* GetEditorIcon() override;
 	virtual const FText GetToolboxCategory() override;
-	virtual bool CanEditChange(const UProperty* Property) const override;
 #endif
 
 	static FName StyleName;
