UMG - Adding some better doc to Widget.h

diff --git a/Engine/Source/Runtime/UMG/Public/Components/Widget.h b/Engine/Source/Runtime/UMG/Public/Components/Widget.h
--- a/Engine/Source/Runtime/UMG/Public/Components/Widget.h
+++ b/Engine/Source/Runtime/UMG/Public/Components/Widget.h
@@ -226,7 +226,10 @@ public:
 	/** Mark this object as modified, also mark the slot as modified. */
 	virtual bool Modify(bool bAlwaysMarkDirty = true);
 
-	/** @return true if this widget is a child of the PossibleParent */
+	/**
+	 * Recurses up the list of parents and returns true if this widget is a descendant of the PossibleParent
+	 * @return true if this widget is a child of the PossibleParent
+	 */
 	bool IsChildOf(UWidget* PossibleParent);
 	
 #if WITH_EDITOR
