UMG - Moving the style name to a local static in button.

diff --git a/Engine/Source/Runtime/UMG/Private/Components/Button.cpp b/Engine/Source/Runtime/UMG/Private/Components/Button.cpp
--- a/Engine/Source/Runtime/UMG/Private/Components/Button.cpp
+++ b/Engine/Source/Runtime/UMG/Private/Components/Button.cpp
@@ -7,11 +7,10 @@
 /////////////////////////////////////////////////////
 // UButton
 
-FName UButton::StyleName(TEXT("Style"));
-
 UButton::UButton(const FPostConstructInitializeProperties& PCIP)
 	: Super(PCIP)
 {
+	static const FName StyleName(TEXT("Style"));
 	ButtonStyle = PCIP.CreateDefaultSubobject<UButtonWidgetStyle>(this, StyleName);
 
 	SButton::FArguments ButtonDefaults;
diff --git a/Engine/Source/Runtime/UMG/Public/Components/Button.h b/Engine/Source/Runtime/UMG/Public/Components/Button.h
--- a/Engine/Source/Runtime/UMG/Public/Components/Button.h
+++ b/Engine/Source/Runtime/UMG/Public/Components/Button.h
@@ -92,8 +92,6 @@ public:
 	virtual const FText GetToolboxCategory() override;
 #endif
 
-	static FName StyleName;
-
 protected:
 
 	// UPanelWidget
