UMG - Progress bar has been updated to use the 'new method' of styling widgets.

diff --git a/Engine/Source/Runtime/Slate/Public/Framework/Styling/ButtonWidgetStyle.h b/Engine/Source/Runtime/Slate/Public/Framework/Styling/ButtonWidgetStyle.h
--- a/Engine/Source/Runtime/Slate/Public/Framework/Styling/ButtonWidgetStyle.h
+++ b/Engine/Source/Runtime/Slate/Public/Framework/Styling/ButtonWidgetStyle.h
@@ -7,14 +7,14 @@
 
 /**
  */
-UCLASS(hidecategories=Object, MinimalAPI)
+UCLASS(BlueprintType, hidecategories=Object, MinimalAPI)
 class UButtonWidgetStyle : public USlateWidgetStyleContainerBase
 {
 	GENERATED_UCLASS_BODY()
 
 public:
 	/** The actual data describing the button's appearance. */
-	UPROPERTY(Category=Appearance, EditAnywhere, meta=( ShowOnlyInnerProperties ))
+	UPROPERTY(Category=Appearance, EditAnywhere, BlueprintReadWrite, meta=( ShowOnlyInnerProperties ))
 	FButtonStyle ButtonStyle;
 
 	virtual const struct FSlateWidgetStyle* const GetStyle() const override
diff --git a/Engine/Source/Runtime/Slate/Public/Framework/Styling/ProgressWidgetStyle.h b/Engine/Source/Runtime/Slate/Public/Framework/Styling/ProgressWidgetStyle.h
--- a/Engine/Source/Runtime/Slate/Public/Framework/Styling/ProgressWidgetStyle.h
+++ b/Engine/Source/Runtime/Slate/Public/Framework/Styling/ProgressWidgetStyle.h
@@ -7,14 +7,14 @@
 
 /**
  */
-UCLASS(hidecategories=Object, MinimalAPI)
+UCLASS(BlueprintType, hidecategories=Object, MinimalAPI)
 class UProgressWidgetStyle : public USlateWidgetStyleContainerBase
 {
 	GENERATED_UCLASS_BODY()
 
 public:
 	/** The actual data describing the button's appearance. */
-	UPROPERTY(Category=Appearance, EditAnywhere, meta=(ShowOnlyInnerProperties))
+	UPROPERTY(Category="Style", EditAnywhere, BlueprintReadWrite, meta=(ShowOnlyInnerProperties))
 	FProgressBarStyle ProgressBarStyle;
 
 	virtual const struct FSlateWidgetStyle* const GetStyle() const override
diff --git a/Engine/Source/Runtime/UMG/Private/Components/ProgressBar.cpp b/Engine/Source/Runtime/UMG/Private/Components/ProgressBar.cpp
--- a/Engine/Source/Runtime/UMG/Private/Components/ProgressBar.cpp
+++ b/Engine/Source/Runtime/UMG/Private/Components/ProgressBar.cpp
@@ -10,39 +10,49 @@
 UProgressBar::UProgressBar(const FPostConstructInitializeProperties& PCIP)
 	: Super(PCIP)
 {
+	static const FName StyleName(TEXT("Style"));
+	ProgressStyle = PCIP.CreateDefaultSubobject<UProgressWidgetStyle>(this, StyleName);
+
 	BarFillType = EProgressBarFillType::LeftToRight;
 	bIsMarquee = false;
 	Percent = 0;
 	FillColorAndOpacity = FLinearColor::White;
-	BorderPadding = FVector2D(1, 0);
+}
+
+void UProgressBar::ReleaseNativeWidget()
+{
+	Super::ReleaseNativeWidget();
+
+	MyProgressBar.Reset();
 }
 
 TSharedRef<SWidget> UProgressBar::RebuildWidget()
 {
-	MyProgressBar = SNew(SProgressBar);
+	MyProgressBar = SNew(SProgressBar)
+		.BorderPadding(FVector2D(0.0f, 0.0f));
 
 	return MyProgressBar.ToSharedRef();
 }
 
 void UProgressBar::SynchronizeProperties()
 {
-	const FProgressBarStyle* StylePtr = ( Style != NULL ) ? Style->GetStyle<FProgressBarStyle>() : NULL;
-	if ( StylePtr )
-	{
-		MyProgressBar->SetStyle(StylePtr);
-	}
-
 	TAttribute< TOptional<float> > PercentBinding = OPTIONAL_BINDING_CONVERT(float, Percent, TOptional<float>, ConvertFloatToOptionalFloat);
 
+	MyProgressBar->SetStyle(&ProgressStyle->ProgressBarStyle);
+
 	MyProgressBar->SetPercent(bIsMarquee ? TOptional<float>() : PercentBinding);
 	
 	MyProgressBar->SetBarFillType(BarFillType);
 	MyProgressBar->SetFillColorAndOpacity(FillColorAndOpacity);
-	MyProgressBar->SetBorderPadding(BorderPadding);
-	
-	MyProgressBar->SetBackgroundImage(BackgroundImage ? &BackgroundImage->Brush : nullptr);
-	MyProgressBar->SetFillImage(FillImage ? &FillImage->Brush : nullptr);
-	MyProgressBar->SetMarqueeImage(MarqueeImage ? &MarqueeImage->Brush : nullptr);
+}
+
+void UProgressBar::SetIsMarquee(bool InbIsMarquee)
+{
+	bIsMarquee = InbIsMarquee;
+	if ( MyProgressBar.IsValid() )
+	{
+		MyProgressBar->SetPercent(bIsMarquee ? TOptional<float>() : Percent);
+	}
 }
 
 void UProgressBar::SetPercent(float InPercent)
@@ -54,6 +64,43 @@ void UProgressBar::SetPercent(float InPercent)
 	}
 }
 
+void UProgressBar::PostLoad()
+{
+	Super::PostLoad();
+
+	if ( GetLinkerUE4Version() < VER_UE4_DEPRECATE_UMG_STYLE_ASSETS )
+	{
+		if ( Style_DEPRECATED != nullptr )
+		{
+			const FProgressBarStyle* StylePtr = Style_DEPRECATED->GetStyle<FProgressBarStyle>();
+			if ( StylePtr != nullptr )
+			{
+				ProgressStyle->ProgressBarStyle = *StylePtr;
+			}
+
+			Style_DEPRECATED = nullptr;
+		}
+
+		if ( BackgroundImage_DEPRECATED != nullptr )
+		{
+			ProgressStyle->ProgressBarStyle.BackgroundImage = BackgroundImage_DEPRECATED->Brush;
+			BackgroundImage_DEPRECATED = nullptr;
+		}
+
+		if ( FillImage_DEPRECATED != nullptr )
+		{
+			ProgressStyle->ProgressBarStyle.FillImage = FillImage_DEPRECATED->Brush;
+			FillImage_DEPRECATED = nullptr;
+		}
+
+		if ( MarqueeImage_DEPRECATED != nullptr )
+		{
+			ProgressStyle->ProgressBarStyle.MarqueeImage = MarqueeImage_DEPRECATED->Brush;
+			MarqueeImage_DEPRECATED = nullptr;
+		}
+	}
+}
+
 #if WITH_EDITOR
 
 const FSlateBrush* UProgressBar::GetEditorIcon()
diff --git a/Engine/Source/Runtime/UMG/Public/Components/ProgressBar.h b/Engine/Source/Runtime/UMG/Public/Components/ProgressBar.h
--- a/Engine/Source/Runtime/UMG/Public/Components/ProgressBar.h
+++ b/Engine/Source/Runtime/UMG/Public/Components/ProgressBar.h
@@ -2,6 +2,8 @@
 
 #pragma once
 
+#include "ProgressWidgetStyle.h"
+
 #include "ProgressBar.generated.h"
 
 /** ProgressBar widget */
@@ -12,21 +14,25 @@ class UMG_API UProgressBar : public UWidget
 	
 public:
 
+	/** The progress bar style */
+	UPROPERTY(VisibleAnywhere, Instanced, Category="Style")
+	TSubobjectPtr<UProgressWidgetStyle> ProgressStyle;
+
 	/** Style used for the progress bar */
-	UPROPERTY(EditDefaultsOnly, Category="Style", meta=( DisplayThumbnail = "true" ))
-	USlateWidgetStyleAsset* Style;
+	UPROPERTY()
+	USlateWidgetStyleAsset* Style_DEPRECATED;
 
 	/** The brush to use as the background of the progress bar */
-	UPROPERTY(EditDefaultsOnly, Category="Style", meta=( DisplayThumbnail = "true" ), AdvancedDisplay)
-	USlateBrushAsset* BackgroundImage;
+	UPROPERTY()
+	USlateBrushAsset* BackgroundImage_DEPRECATED;
 	
 	/** The brush to use as the fill image */
-	UPROPERTY(EditDefaultsOnly, Category="Style", meta=( DisplayThumbnail = "true" ), AdvancedDisplay)
-	USlateBrushAsset* FillImage;
+	UPROPERTY()
+	USlateBrushAsset* FillImage_DEPRECATED;
 	
 	/** The brush to use as the marquee image */
-	UPROPERTY(EditDefaultsOnly, Category="Style", meta=( DisplayThumbnail = "true" ), AdvancedDisplay)
-	USlateBrushAsset* MarqueeImage;
+	UPROPERTY()
+	USlateBrushAsset* MarqueeImage_DEPRECATED;
 
 	/** Defines if this progress bar fills Left to right or right to left */
 	UPROPERTY(EditDefaultsOnly, Category=Appearance)
@@ -36,7 +42,7 @@ public:
 	bool bIsMarquee;
 
 	/** Used to determine the fill position of the progress bar ranging 0..1 */
-	UPROPERTY(EditDefaultsOnly, Category=Appearance, meta=( UIMin = "0", UIMax = "1", editcondition = "bOverride_Percent" ))
+	UPROPERTY(EditDefaultsOnly, Category=Appearance, meta=( UIMin = "0", UIMax = "1" ))
 	float Percent;
 
 	/** A bindable delegate to allow logic to drive the text of the widget */
@@ -47,20 +53,28 @@ public:
 	UPROPERTY(EditDefaultsOnly, Category=Appearance)
 	FLinearColor FillColorAndOpacity;
 
-	/** Border Padding around fill bar */
-	UPROPERTY(EditDefaultsOnly, Category=Appearance)
-	FVector2D BorderPadding;
-
 public:
 	
 	/** Sets the current value of the ProgressBar. */
 	UFUNCTION(BlueprintCallable, Category="Behavior")
 	void SetPercent(float InPercent);
+
+	/** Sets the progress bar to show as a marquee. */
+	UFUNCTION(BlueprintCallable, Category="Behavior")
+	void SetIsMarquee(bool InbIsMarquee);
 	
 	// UWidget interface
 	virtual void SynchronizeProperties() override;
 	// End of UWidget interface
 
+	// UVisual interface
+	virtual void ReleaseNativeWidget() override;
+	// End of UVisual interface
+
+	// Begin UObject interface
+	virtual void PostLoad() override;
+	// End of UObject interface
+
 #if WITH_EDITOR
 	// UWidget interface
 	virtual const FSlateBrush* GetEditorIcon() override;
