UMG - Setting the progress bar to the default style.

diff --git a/Engine/Source/Runtime/UMG/Private/Components/ProgressBar.cpp b/Engine/Source/Runtime/UMG/Private/Components/ProgressBar.cpp
--- a/Engine/Source/Runtime/UMG/Private/Components/ProgressBar.cpp
+++ b/Engine/Source/Runtime/UMG/Private/Components/ProgressBar.cpp
@@ -13,6 +13,9 @@ UProgressBar::UProgressBar(const FPostConstructInitializeProperties& PCIP)
 	static const FName StyleName(TEXT("Style"));
 	ProgressStyle = PCIP.CreateDefaultSubobject<UProgressWidgetStyle>(this, StyleName);
 
+	SProgressBar::FArguments SlateDefaults;
+	ProgressStyle->ProgressBarStyle = *SlateDefaults._Style;
+
 	BarFillType = EProgressBarFillType::LeftToRight;
 	bIsMarquee = false;
 	Percent = 0;
