UMG - Throbber and Circular Throbber now use the new method of styling, no assets, raw brushes.

diff --git a/Engine/Source/Runtime/UMG/Private/Components/CircularThrobber.cpp b/Engine/Source/Runtime/UMG/Private/Components/CircularThrobber.cpp
--- a/Engine/Source/Runtime/UMG/Private/Components/CircularThrobber.cpp
+++ b/Engine/Source/Runtime/UMG/Private/Components/CircularThrobber.cpp
@@ -12,6 +12,7 @@ UCircularThrobber::UCircularThrobber(const FPostConstructInitializeProperties& P
 	: Super(PCIP)
 {
 	SCircularThrobber::FArguments DefaultArgs;
+	Image = *DefaultArgs._PieceImage;
 
 	NumberOfPieces = DefaultArgs._NumPieces;
 	Period = DefaultArgs._Period;
@@ -30,7 +31,7 @@ TSharedRef<SWidget> UCircularThrobber::RebuildWidget()
 	SCircularThrobber::FArguments DefaultArgs;
 
 	MyCircularThrobber = SNew(SCircularThrobber)
-		.PieceImage(GetPieceBrush())
+		.PieceImage(&Image)
 		.NumPieces(NumberOfPieces)
 		.Period(Period)
 		.Radius(Radius);
@@ -42,22 +43,11 @@ void UCircularThrobber::SynchronizeProperties()
 {
 	Super::SynchronizeProperties();
 
-	MyCircularThrobber->SetPieceImage(GetPieceBrush());
 	MyCircularThrobber->SetNumPieces(NumberOfPieces);
 	MyCircularThrobber->SetPeriod(Period);
 	MyCircularThrobber->SetRadius(Radius);
 }
 
-const FSlateBrush* UCircularThrobber::GetPieceBrush() const
-{
-	if (PieceImage == NULL)
-	{
-		SCircularThrobber::FArguments DefaultArgs;
-		return DefaultArgs._PieceImage;
-	}
-	return &PieceImage->Brush;
-}
-
 void UCircularThrobber::SetNumberOfPieces(int32 InNumberOfPieces)
 {
 	NumberOfPieces = InNumberOfPieces;
@@ -85,12 +75,17 @@ void UCircularThrobber::SetRadius(float InRadius)
 	}
 }
 
-void UCircularThrobber::SetPieceImage(USlateBrushAsset* InPieceImage)
+void UCircularThrobber::PostLoad()
 {
-	PieceImage = InPieceImage;
-	if (MyCircularThrobber.IsValid())
+	Super::PostLoad();
+
+	if ( GetLinkerUE4Version() < VER_UE4_DEPRECATE_UMG_STYLE_ASSETS )
 	{
-		MyCircularThrobber->SetPieceImage(GetPieceBrush());
+		if ( PieceImage_DEPRECATED != nullptr )
+		{
+			Image = PieceImage_DEPRECATED->Brush;
+			PieceImage_DEPRECATED = nullptr;
+		}
 	}
 }
 
diff --git a/Engine/Source/Runtime/UMG/Private/Components/Throbber.cpp b/Engine/Source/Runtime/UMG/Private/Components/Throbber.cpp
--- a/Engine/Source/Runtime/UMG/Private/Components/Throbber.cpp
+++ b/Engine/Source/Runtime/UMG/Private/Components/Throbber.cpp
@@ -12,6 +12,7 @@ UThrobber::UThrobber(const FPostConstructInitializeProperties& PCIP)
 	: Super(PCIP)
 {
 	SThrobber::FArguments DefaultArgs;
+	Image = *DefaultArgs._PieceImage;
 
 	NumberOfPieces = DefaultArgs._NumPieces;
 
@@ -20,12 +21,19 @@ UThrobber::UThrobber(const FPostConstructInitializeProperties& PCIP)
 	bAnimateOpacity = (DefaultArgs._Animate & SThrobber::Opacity) != 0;
 }
 
+void UThrobber::ReleaseNativeWidget()
+{
+	Super::ReleaseNativeWidget();
+
+	MyThrobber.Reset();
+}
+
 TSharedRef<SWidget> UThrobber::RebuildWidget()
 {
 	SThrobber::FArguments DefaultArgs;
 
 	MyThrobber = SNew(SThrobber)
-		.PieceImage(GetPieceBrush())
+		.PieceImage(&Image)
 		.NumPieces(NumberOfPieces)
 		.Animate(GetAnimation());
 
@@ -36,7 +44,6 @@ void UThrobber::SynchronizeProperties()
 {
 	Super::SynchronizeProperties();
 
-	MyThrobber->SetPieceImage(GetPieceBrush());
 	MyThrobber->SetNumPieces(NumberOfPieces);
 	MyThrobber->SetAnimate(GetAnimation());
 }
@@ -50,16 +57,6 @@ SThrobber::EAnimation UThrobber::GetAnimation() const
 	return static_cast<SThrobber::EAnimation>(AnimationParams);
 }
 
-const FSlateBrush* UThrobber::GetPieceBrush() const
-{
-	if (PieceImage == NULL)
-	{
-		SThrobber::FArguments DefaultArgs;
-		return DefaultArgs._PieceImage;
-	}
-	return &PieceImage->Brush;
-}
-
 void UThrobber::SetNumberOfPieces(int32 InNumberOfPieces)
 {
 	NumberOfPieces = InNumberOfPieces;
@@ -69,7 +66,7 @@ void UThrobber::SetNumberOfPieces(int32 InNumberOfPieces)
 	}
 }
 
-void UThrobber:: SetAnimateHorizontally(bool bInAnimateHorizontally)
+void UThrobber::SetAnimateHorizontally(bool bInAnimateHorizontally)
 {
 	bAnimateHorizontally = bInAnimateHorizontally;
 	if (MyThrobber.IsValid())
@@ -78,7 +75,7 @@ void UThrobber:: SetAnimateHorizontally(bool bInAnimateHorizontally)
 	}
 }
 
-void UThrobber:: SetAnimateVertically(bool bInAnimateVertically)
+void UThrobber::SetAnimateVertically(bool bInAnimateVertically)
 {
 	bAnimateVertically = bInAnimateVertically;
 	if (MyThrobber.IsValid())
@@ -96,12 +93,17 @@ void UThrobber:: SetAnimateOpacity(bool bInAnimateOpacity)
 	}
 }
 
-void UThrobber::SetPieceImage(USlateBrushAsset* InPieceImage)
+void UThrobber::PostLoad()
 {
-	PieceImage = InPieceImage;
-	if (MyThrobber.IsValid())
+	Super::PostLoad();
+
+	if ( GetLinkerUE4Version() < VER_UE4_DEPRECATE_UMG_STYLE_ASSETS )
 	{
-		MyThrobber->SetPieceImage(GetPieceBrush());
+		if ( PieceImage_DEPRECATED != nullptr )
+		{
+			Image = PieceImage_DEPRECATED->Brush;
+			PieceImage_DEPRECATED = nullptr;
+		}
 	}
 }
 
diff --git a/Engine/Source/Runtime/UMG/Public/Components/CircularThrobber.h b/Engine/Source/Runtime/UMG/Public/Components/CircularThrobber.h
--- a/Engine/Source/Runtime/UMG/Public/Components/CircularThrobber.h
+++ b/Engine/Source/Runtime/UMG/Public/Components/CircularThrobber.h
@@ -26,15 +26,17 @@ public:
 	UFUNCTION(BlueprintCallable, Category = "Appearance")
 	void SetRadius(float InRadius);
 
-	/** Sets image to use for each segment of the throbber. */
-	UFUNCTION(BlueprintCallable, Category = "Appearance")
-	void SetPieceImage(USlateBrushAsset* InPieceImage);
-
 	// UWidget interface
 	virtual void SynchronizeProperties() override;
 	// End of UWidget interface
 
+	// UVisual interface
 	virtual void ReleaseNativeWidget() override;
+	// End of UVisual interface
+
+	// Begin UObject interface
+	virtual void PostLoad() override;
+	// End of UObject interface
 
 #if WITH_EDITOR
 	virtual const FSlateBrush* GetEditorIcon() override;
@@ -56,19 +58,17 @@ protected:
 	float Radius;
 
 	/** Image to use for each segment of the throbber */
-	UPROPERTY(EditDefaultsOnly, Category=Appearance, meta=( DisplayThumbnail = "true" ))
-	USlateBrushAsset* PieceImage;
+	UPROPERTY()
+	USlateBrushAsset* PieceImage_DEPRECATED;
+
+	UPROPERTY(EditDefaultsOnly, Category=Appearance)
+	FSlateBrush Image;
 
 protected:
 	// UWidget interface
 	virtual TSharedRef<SWidget> RebuildWidget() override;
 	// End of UWidget interface
 
-private:
-
-	/** Gets the brush used to draw the pieces of the throbber. */
-	const FSlateBrush* GetPieceBrush() const;
-
 private:
 	/** The CircularThrobber widget managed by this object. */
 	TSharedPtr<SCircularThrobber> MyCircularThrobber;
diff --git a/Engine/Source/Runtime/UMG/Public/Components/Throbber.h b/Engine/Source/Runtime/UMG/Public/Components/Throbber.h
--- a/Engine/Source/Runtime/UMG/Public/Components/Throbber.h
+++ b/Engine/Source/Runtime/UMG/Public/Components/Throbber.h
@@ -30,14 +30,18 @@ public:
 	UFUNCTION(BlueprintCallable, Category = "Appearance")
 	void SetAnimateOpacity(bool bInAnimateOpacity);
 
-	/** Sets the image to use for each segment of the throbber. */
-	UFUNCTION(BlueprintCallable, Category = "Appearance")
-	void SetPieceImage(USlateBrushAsset* InPieceImage);
-
 	// UWidget interface
 	virtual void SynchronizeProperties() override;
 	// End of UWidget interface
 
+	// UVisual interface
+	virtual void ReleaseNativeWidget() override;
+	// End of UVisual interface
+
+	// Begin UObject interface
+	virtual void PostLoad() override;
+	// End of UObject interface
+
 #if WITH_EDITOR
 	virtual const FSlateBrush* GetEditorIcon() override;
 	virtual const FText GetToolboxCategory() override;
@@ -62,8 +66,11 @@ protected:
 	bool bAnimateOpacity;
 
 	/** Image to use for each segment of the throbber */
-	UPROPERTY(EditDefaultsOnly, Category=Appearance, meta=( DisplayThumbnail = "true" ))
-	USlateBrushAsset* PieceImage;
+	UPROPERTY()
+	USlateBrushAsset* PieceImage_DEPRECATED;
+
+	UPROPERTY(EditDefaultsOnly, Category=Appearance)
+	FSlateBrush Image;
 
 protected:
 	// UWidget interface
@@ -75,9 +82,6 @@ private:
 	/** Gets the combined value of the animation properties as a single SThrobber::EAnimation value. */
 	SThrobber::EAnimation GetAnimation() const;
 
-	/** Gets the brush which should be used to draw the trobber pieces. */
-	const FSlateBrush* GetPieceBrush() const;
-
 private:
 	/** The Throbber widget managed by this object. */
 	TSharedPtr<SThrobber> MyThrobber;
