UMG - Checkboxes have been converted to use the new styling method of widgets.

diff --git a/Engine/Source/Runtime/UMG/Private/Components/CheckBox.cpp b/Engine/Source/Runtime/UMG/Private/Components/CheckBox.cpp
--- a/Engine/Source/Runtime/UMG/Private/Components/CheckBox.cpp
+++ b/Engine/Source/Runtime/UMG/Private/Components/CheckBox.cpp
@@ -10,14 +10,17 @@
 UCheckBox::UCheckBox(const FPostConstructInitializeProperties& PCIP)
 	: Super(PCIP)
 {
-	SCheckBox::FArguments CheckBoxDefaults;
+	static const FName StyleName(TEXT("Style"));
+	CheckBoxStyle = PCIP.CreateDefaultSubobject<UCheckBoxWidgetStyle>(this, StyleName);
+
+	SCheckBox::FArguments SlateDefaults;
+	CheckBoxStyle->CheckBoxStyle = *SlateDefaults._Style;
 
 	CheckedState = ESlateCheckBoxState::Unchecked;
 
-	HorizontalAlignment = CheckBoxDefaults._HAlign;
-	Padding = CheckBoxDefaults._Padding.Get();
+	HorizontalAlignment = SlateDefaults._HAlign;
+	Padding = SlateDefaults._Padding.Get();
 
-	ForegroundColor = FLinearColor::White;
 	BorderBackgroundColor = FLinearColor::White;
 }
 
@@ -30,33 +33,12 @@ void UCheckBox::ReleaseNativeWidget()
 
 TSharedRef<SWidget> UCheckBox::RebuildWidget()
 {
-	TOptional<FSlateSound> OptionalCheckedSound;
-	if ( CheckedSound.GetResourceObject() )
-	{
-		OptionalCheckedSound = CheckedSound;
-	}
-
-	TOptional<FSlateSound> OptionalUncheckedSound;
-	if ( UncheckedSound.GetResourceObject() )
-	{
-		OptionalUncheckedSound = UncheckedSound;
-	}
-
-	TOptional<FSlateSound> OptionalHoveredSound;
-	if ( HoveredSound.GetResourceObject() )
-	{
-		OptionalHoveredSound = HoveredSound;
-	}
-
 	MyCheckbox = SNew(SCheckBox)
 		.OnCheckStateChanged( BIND_UOBJECT_DELEGATE(FOnCheckStateChanged, SlateOnCheckStateChangedCallback) )
+		.Style(&CheckBoxStyle->CheckBoxStyle)
 		.HAlign( HorizontalAlignment )
 		.Padding( Padding )
-		.ForegroundColor( ForegroundColor )
 		.BorderBackgroundColor( BorderBackgroundColor )
-		.CheckedSoundOverride(OptionalCheckedSound)
-		.UncheckedSoundOverride(OptionalUncheckedSound)
-		.HoveredSoundOverride(OptionalHoveredSound)
 		;
 
 	if ( GetChildrenCount() > 0 )
@@ -71,28 +53,7 @@ void UCheckBox::SynchronizeProperties()
 {
 	Super::SynchronizeProperties();
 
-	const FCheckBoxStyle* StylePtr = ( Style != NULL ) ? Style->GetStyle<FCheckBoxStyle>() : NULL;
-	if ( StylePtr == NULL )
-	{
-		SCheckBox::FArguments Defaults;
-		StylePtr = Defaults._Style;
-	}
-
-	MyCheckbox->SetStyle(StylePtr);
-
 	MyCheckbox->SetIsChecked( OPTIONAL_BINDING(ESlateCheckBoxState::Type, CheckedState) );
-	
-	MyCheckbox->SetUncheckedImage(UncheckedImage ? &UncheckedImage->Brush : nullptr);
-	MyCheckbox->SetUncheckedHoveredImage(UncheckedHoveredImage ? &UncheckedHoveredImage->Brush : nullptr);
-	MyCheckbox->SetUncheckedPressedImage(UncheckedPressedImage ? &UncheckedPressedImage->Brush : nullptr);
-	
-	MyCheckbox->SetCheckedImage(CheckedImage ? &CheckedImage->Brush : nullptr);
-	MyCheckbox->SetCheckedHoveredImage(CheckedHoveredImage ? &CheckedHoveredImage->Brush : nullptr);
-	MyCheckbox->SetCheckedPressedImage(CheckedPressedImage ? &CheckedPressedImage->Brush : nullptr);
-	
-	MyCheckbox->SetUndeterminedImage(UndeterminedImage ? &UndeterminedImage->Brush : nullptr);
-	MyCheckbox->SetUndeterminedHoveredImage(UndeterminedHoveredImage ? &UndeterminedHoveredImage->Brush : nullptr);
-	MyCheckbox->SetUndeterminedPressedImage(UndeterminedPressedImage ? &UndeterminedPressedImage->Brush : nullptr);
 }
 
 void UCheckBox::OnSlotAdded(UPanelSlot* Slot)
@@ -176,6 +137,79 @@ void UCheckBox::SlateOnCheckStateChangedCallback(ESlateCheckBoxState::Type NewSt
 	}
 }
 
+void UCheckBox::PostLoad()
+{
+	Super::PostLoad();
+
+	if ( GetLinkerUE4Version() < VER_UE4_DEPRECATE_UMG_STYLE_ASSETS )
+	{
+		if ( Style_DEPRECATED != nullptr )
+		{
+			const FCheckBoxStyle* StylePtr = Style_DEPRECATED->GetStyle<FCheckBoxStyle>();
+			if ( StylePtr != nullptr )
+			{
+				CheckBoxStyle->CheckBoxStyle = *StylePtr;
+			}
+
+			Style_DEPRECATED = nullptr;
+		}
+
+		if ( UncheckedImage_DEPRECATED != nullptr )
+		{
+			CheckBoxStyle->CheckBoxStyle.UncheckedImage = UncheckedImage_DEPRECATED->Brush;
+			UncheckedImage_DEPRECATED = nullptr;
+		}
+
+		if ( UncheckedHoveredImage_DEPRECATED != nullptr )
+		{
+			CheckBoxStyle->CheckBoxStyle.UncheckedHoveredImage = UncheckedHoveredImage_DEPRECATED->Brush;
+			UncheckedHoveredImage_DEPRECATED = nullptr;
+		}
+
+		if ( UncheckedPressedImage_DEPRECATED != nullptr )
+		{
+			CheckBoxStyle->CheckBoxStyle.UncheckedPressedImage = UncheckedPressedImage_DEPRECATED->Brush;
+			UncheckedPressedImage_DEPRECATED = nullptr;
+		}
+
+		if ( CheckedImage_DEPRECATED != nullptr )
+		{
+			CheckBoxStyle->CheckBoxStyle.CheckedImage = CheckedImage_DEPRECATED->Brush;
+			CheckedImage_DEPRECATED = nullptr;
+		}
+
+		if ( CheckedHoveredImage_DEPRECATED != nullptr )
+		{
+			CheckBoxStyle->CheckBoxStyle.CheckedHoveredImage = CheckedHoveredImage_DEPRECATED->Brush;
+			CheckedHoveredImage_DEPRECATED = nullptr;
+		}
+
+		if ( CheckedPressedImage_DEPRECATED != nullptr )
+		{
+			CheckBoxStyle->CheckBoxStyle.CheckedPressedImage = CheckedPressedImage_DEPRECATED->Brush;
+			CheckedPressedImage_DEPRECATED = nullptr;
+		}
+
+		if ( UndeterminedImage_DEPRECATED != nullptr )
+		{
+			CheckBoxStyle->CheckBoxStyle.UndeterminedImage = UndeterminedImage_DEPRECATED->Brush;
+			UndeterminedImage_DEPRECATED = nullptr;
+		}
+
+		if ( UndeterminedHoveredImage_DEPRECATED != nullptr )
+		{
+			CheckBoxStyle->CheckBoxStyle.UndeterminedHoveredImage = UndeterminedHoveredImage_DEPRECATED->Brush;
+			UndeterminedHoveredImage_DEPRECATED = nullptr;
+		}
+
+		if ( UndeterminedPressedImage_DEPRECATED != nullptr )
+		{
+			CheckBoxStyle->CheckBoxStyle.UndeterminedPressedImage = UndeterminedPressedImage_DEPRECATED->Brush;
+			UndeterminedPressedImage_DEPRECATED = nullptr;
+		}
+	}
+}
+
 #if WITH_EDITOR
 
 const FSlateBrush* UCheckBox::GetEditorIcon()
diff --git a/Engine/Source/Runtime/UMG/Public/Components/CheckBox.h b/Engine/Source/Runtime/UMG/Public/Components/CheckBox.h
--- a/Engine/Source/Runtime/UMG/Public/Components/CheckBox.h
+++ b/Engine/Source/Runtime/UMG/Public/Components/CheckBox.h
@@ -2,6 +2,8 @@
 
 #pragma once
 
+#include "CheckBoxWidgetStyle.h"
+
 #include "CheckBox.generated.h"
 
 DECLARE_DYNAMIC_MULTICAST_DELEGATE_OneParam( FOnCheckBoxComponentStateChanged, bool, bIsChecked );
@@ -15,45 +17,49 @@ class UMG_API UCheckBox : public UContentWidget
 	GENERATED_UCLASS_BODY()
 
 public:
+	/** The checkbox bar style */
+	UPROPERTY(VisibleAnywhere, Instanced, Category="Style")
+	TSubobjectPtr<UCheckBoxWidgetStyle> CheckBoxStyle;
+
 	/** Style of the check box */
-	UPROPERTY(EditDefaultsOnly, Category=Style, meta=( DisplayThumbnail = "true" ))
-	USlateWidgetStyleAsset* Style;
+	UPROPERTY()
+	USlateWidgetStyleAsset* Style_DEPRECATED;
 
 	/** Image to use when the checkbox is unchecked */
-	UPROPERTY(EditDefaultsOnly, Category="Style", meta=( DisplayThumbnail = "true" ), AdvancedDisplay)
-	USlateBrushAsset* UncheckedImage;
+	UPROPERTY()
+	USlateBrushAsset* UncheckedImage_DEPRECATED;
 	
 	/** Image to use when the checkbox is unchecked and hovered */
-	UPROPERTY(EditDefaultsOnly, Category="Style", meta=( DisplayThumbnail = "true" ), AdvancedDisplay)
-	USlateBrushAsset* UncheckedHoveredImage;
+	UPROPERTY()
+	USlateBrushAsset* UncheckedHoveredImage_DEPRECATED;
 	
 	/** Image to use when the checkbox is unchecked and pressed */
-	UPROPERTY(EditDefaultsOnly, Category="Style", meta=( DisplayThumbnail = "true" ), AdvancedDisplay)
-	USlateBrushAsset* UncheckedPressedImage;
+	UPROPERTY()
+	USlateBrushAsset* UncheckedPressedImage_DEPRECATED;
 	
 	/** Image to use when the checkbox is checked */
-	UPROPERTY(EditDefaultsOnly, Category="Style", meta=( DisplayThumbnail = "true" ), AdvancedDisplay)
-	USlateBrushAsset* CheckedImage;
+	UPROPERTY()
+	USlateBrushAsset* CheckedImage_DEPRECATED;
 	
 	/** Image to use when the checkbox is checked and hovered */
-	UPROPERTY(EditDefaultsOnly, Category="Style", meta=( DisplayThumbnail = "true" ), AdvancedDisplay)
-	USlateBrushAsset* CheckedHoveredImage;
+	UPROPERTY()
+	USlateBrushAsset* CheckedHoveredImage_DEPRECATED;
 	
 	/** Image to use when the checkbox is checked and pressed */
-	UPROPERTY(EditDefaultsOnly, Category="Style", meta=( DisplayThumbnail = "true" ), AdvancedDisplay)
-	USlateBrushAsset* CheckedPressedImage;
+	UPROPERTY()
+	USlateBrushAsset* CheckedPressedImage_DEPRECATED;
 	
 	/** Image to use when the checkbox is in an ambiguous state and hovered */
-	UPROPERTY(EditDefaultsOnly, Category="Style", meta=( DisplayThumbnail = "true" ), AdvancedDisplay)
-	USlateBrushAsset* UndeterminedImage;
+	UPROPERTY()
+	USlateBrushAsset* UndeterminedImage_DEPRECATED;
 	
 	/** Image to use when the checkbox is checked and hovered */
-	UPROPERTY(EditDefaultsOnly, Category="Style", meta=( DisplayThumbnail = "true" ), AdvancedDisplay)
-	USlateBrushAsset* UndeterminedHoveredImage;
+	UPROPERTY()
+	USlateBrushAsset* UndeterminedHoveredImage_DEPRECATED;
 	
 	/** Image to use when the checkbox is in an ambiguous state and pressed */
-	UPROPERTY(EditDefaultsOnly, Category="Style", meta=( DisplayThumbnail = "true" ), AdvancedDisplay)
-	USlateBrushAsset* UndeterminedPressedImage;
+	UPROPERTY()
+	USlateBrushAsset* UndeterminedPressedImage_DEPRECATED;
 
 	/** Whether the check box is currently in a checked state */
 	UPROPERTY(EditDefaultsOnly, Category=Appearance)
@@ -71,23 +77,10 @@ public:
 	UPROPERTY(EditDefaultsOnly, Category=Appearance)
 	FMargin Padding;
 
-	/** Foreground color for the checkbox's content and parts */
-	UPROPERTY(EditDefaultsOnly, Category=Appearance)
-	FSlateColor ForegroundColor;
-
 	/** The color of the background border */
 	UPROPERTY(EditDefaultsOnly, Category=Appearance)
 	FSlateColor BorderBackgroundColor;
 
-	UPROPERTY(EditDefaultsOnly, Category="Sound")
-	FSlateSound CheckedSound;
-
-	UPROPERTY(EditDefaultsOnly, Category="Sound")
-	FSlateSound UncheckedSound;
-
-	UPROPERTY(EditDefaultsOnly, Category="Sound")
-	FSlateSound HoveredSound;
-
 public:
 
 	/** Called when the checked state has changed */
@@ -118,7 +111,13 @@ public:
 	virtual void SynchronizeProperties() override;
 	// End of UWidget interface
 
+	// UVisual interface
 	virtual void ReleaseNativeWidget() override;
+	// End of UVisual interface
+
+	// Begin UObject interface
+	virtual void PostLoad() override;
+	// End of UObject interface
 
 #if WITH_EDITOR
 	virtual const FSlateBrush* GetEditorIcon() override;
