UMG - Combobox is now stylable.

diff --git a/Engine/Source/Runtime/UMG/Private/Components/ComboBoxString.cpp b/Engine/Source/Runtime/UMG/Private/Components/ComboBoxString.cpp
--- a/Engine/Source/Runtime/UMG/Private/Components/ComboBoxString.cpp
+++ b/Engine/Source/Runtime/UMG/Private/Components/ComboBoxString.cpp
@@ -10,6 +10,12 @@
 UComboBoxString::UComboBoxString(const FPostConstructInitializeProperties& PCIP)
 	: Super(PCIP)
 {
+	static const FName StyleName(TEXT("Style"));
+	ComboBoxStyle = PCIP.CreateDefaultSubobject<UComboBoxWidgetStyle>(this, StyleName);
+
+	SComboBox< TSharedPtr<FString> >::FArguments SlateDefaults;
+	ComboBoxStyle->ComboBoxStyle = *SlateDefaults._ComboBoxStyle;
+
 	ContentPadding = FMargin(4.0, 2.0);
 	MaxListHeight = 450.0f;
 	HasDownArrow = true;
@@ -30,17 +36,18 @@ TSharedRef<SWidget> UComboBoxString::RebuildWidget()
 		AddOption(DefaultOptions);
 	}
 
-	TSharedPtr<FString> InitiallySelectedItem;
-	int32 InitialIndex = FindOptionIndex(InitiallySelectedOption);
+	TSharedPtr<FString> SelectedOptionPtr;
+	int32 InitialIndex = FindOptionIndex(SelectedOption);
 	if ( InitialIndex != -1 )
 	{
-		InitiallySelectedItem = Options[InitialIndex];
+		SelectedOptionPtr = Options[InitialIndex];
 	}
 
 	MyComboBox =
 		SNew(SComboBox< TSharedPtr<FString> >)
+		.ComboBoxStyle(&ComboBoxStyle->ComboBoxStyle)
 		.OptionsSource(&Options)
-		.InitiallySelectedItem(InitiallySelectedItem)
+		.InitiallySelectedItem(SelectedOptionPtr)
 		.ContentPadding(ContentPadding)
 		.MaxListHeight(MaxListHeight)
 		.HasDownArrow(HasDownArrow)
@@ -55,22 +62,9 @@ TSharedRef<SWidget> UComboBoxString::RebuildWidget()
 	if ( InitialIndex != -1 )
 	{
 		// Generate the widget for the initially selected widget if needed
-		ComoboBoxContent->SetContent(HandleGenerateWidget(InitiallySelectedItem));
+		ComoboBoxContent->SetContent(HandleGenerateWidget(SelectedOptionPtr));
 	}
 
-	//TODO UMG Expose the rest of these options
-
-	//SLATE_STYLE_ARGUMENT(FComboBoxStyle, ComboBoxStyle)
-
-	//	/** The visual style of the button part of the combo box (overrides ComboBoxStyle) */
-	//	SLATE_STYLE_ARGUMENT(FButtonStyle, ButtonStyle)
-
-	//	/** The sound to play when the button is pressed (overrides ComboBoxStyle) */
-	//	SLATE_ARGUMENT(TOptional<FSlateSound>, PressedSoundOverride)
-
-	//	/** The sound to play when the selection changes (overrides ComboBoxStyle) */
-	//	SLATE_ARGUMENT(TOptional<FSlateSound>, SelectionChangeSoundOverride)
-
 	return MyComboBox.ToSharedRef();
 }
 
diff --git a/Engine/Source/Runtime/UMG/Public/Components/ComboBoxString.h b/Engine/Source/Runtime/UMG/Public/Components/ComboBoxString.h
--- a/Engine/Source/Runtime/UMG/Public/Components/ComboBoxString.h
+++ b/Engine/Source/Runtime/UMG/Public/Components/ComboBoxString.h
@@ -2,6 +2,8 @@
 
 #pragma once
 
+#include "ComboBoxWidgetStyle.h"
+
 #include "ComboBoxString.generated.h"
 
 /**
@@ -17,13 +19,17 @@ class UMG_API UComboBoxString : public UWidget
 
 public:
 
+	/** The style */
+	UPROPERTY(VisibleAnywhere, Instanced, Category=Appearance )
+	TSubobjectPtr<UComboBoxWidgetStyle> ComboBoxStyle;
+
 	/** The list of items to be displayed on the combobox. */
 	UPROPERTY(EditDefaultsOnly, Category=Content)
 	TArray<FString> DefaultOptions;
 
 	/** The item in the combobox to select by default */
 	UPROPERTY(EditDefaultsOnly, Category=Content)
-	FString InitiallySelectedOption;
+	FString SelectedOption;
 
 	UPROPERTY(EditDefaultsOnly, Category=Content)
 	FMargin ContentPadding;
