UMG - Renaming the style property on several widgets using the new styling method.  Adding the newer styling method to EditableText.

diff --git a/Engine/Source/Runtime/UMG/Private/Components/Button.cpp b/Engine/Source/Runtime/UMG/Private/Components/Button.cpp
--- a/Engine/Source/Runtime/UMG/Private/Components/Button.cpp
+++ b/Engine/Source/Runtime/UMG/Private/Components/Button.cpp
@@ -11,10 +11,10 @@ UButton::UButton(const FPostConstructInitializeProperties& PCIP)
 	: Super(PCIP)
 {
 	static const FName StyleName(TEXT("Style"));
-	ButtonStyle = PCIP.CreateDefaultSubobject<UButtonWidgetStyle>(this, StyleName);
+	WidgetStyle = PCIP.CreateDefaultSubobject<UButtonWidgetStyle>(this, StyleName);
 
 	SButton::FArguments ButtonDefaults;
-	ButtonStyle->ButtonStyle = *ButtonDefaults._ButtonStyle;
+	WidgetStyle->ButtonStyle = *ButtonDefaults._ButtonStyle;
 
 	ColorAndOpacity = FLinearColor::White;
 	BackgroundColor = FLinearColor::White;
@@ -35,7 +35,7 @@ void UButton::ReleaseNativeWidget()
 TSharedRef<SWidget> UButton::RebuildWidget()
 {
 	MyButton = SNew(SButton)
-		.ButtonStyle(&ButtonStyle->ButtonStyle)
+		.ButtonStyle(&WidgetStyle->ButtonStyle)
 		.ClickMethod(ClickMethod)
 		.TouchMethod(TouchMethod)
 		.IsFocusable(IsFocusable);
@@ -83,12 +83,12 @@ void UButton::OnSlotRemoved(UPanelSlot* Slot)
 
 void UButton::SetButtonStyle(FButtonStyle InButtonStyle)
 {
-	ButtonStyle->ButtonStyle = InButtonStyle;
+	WidgetStyle->ButtonStyle = InButtonStyle;
 }
 
 FButtonStyle UButton::GetButtonStyle()
 {
-	return ButtonStyle->ButtonStyle;
+	return WidgetStyle->ButtonStyle;
 }
 
 void UButton::SetColorAndOpacity(FLinearColor Color)
@@ -139,7 +139,7 @@ void UButton::PostLoad()
 		const FButtonStyle* StylePtr = Style_DEPRECATED->GetStyle<FButtonStyle>();
 		if(StylePtr != nullptr)
 		{
-			ButtonStyle->ButtonStyle = *StylePtr;
+			WidgetStyle->ButtonStyle = *StylePtr;
 		}
 
 		Style_DEPRECATED = nullptr;
diff --git a/Engine/Source/Runtime/UMG/Private/Components/CheckBox.cpp b/Engine/Source/Runtime/UMG/Private/Components/CheckBox.cpp
--- a/Engine/Source/Runtime/UMG/Private/Components/CheckBox.cpp
+++ b/Engine/Source/Runtime/UMG/Private/Components/CheckBox.cpp
@@ -11,10 +11,10 @@ UCheckBox::UCheckBox(const FPostConstructInitializeProperties& PCIP)
 	: Super(PCIP)
 {
 	static const FName StyleName(TEXT("Style"));
-	CheckBoxStyle = PCIP.CreateDefaultSubobject<UCheckBoxWidgetStyle>(this, StyleName);
+	WidgetStyle = PCIP.CreateDefaultSubobject<UCheckBoxWidgetStyle>(this, StyleName);
 
 	SCheckBox::FArguments SlateDefaults;
-	CheckBoxStyle->CheckBoxStyle = *SlateDefaults._Style;
+	WidgetStyle->CheckBoxStyle = *SlateDefaults._Style;
 
 	CheckedState = ESlateCheckBoxState::Unchecked;
 
@@ -35,7 +35,7 @@ TSharedRef<SWidget> UCheckBox::RebuildWidget()
 {
 	MyCheckbox = SNew(SCheckBox)
 		.OnCheckStateChanged( BIND_UOBJECT_DELEGATE(FOnCheckStateChanged, SlateOnCheckStateChangedCallback) )
-		.Style(&CheckBoxStyle->CheckBoxStyle)
+		.Style(&WidgetStyle->CheckBoxStyle)
 		.HAlign( HorizontalAlignment )
 		.Padding( Padding )
 		.BorderBackgroundColor( BorderBackgroundColor )
@@ -148,7 +148,7 @@ void UCheckBox::PostLoad()
 			const FCheckBoxStyle* StylePtr = Style_DEPRECATED->GetStyle<FCheckBoxStyle>();
 			if ( StylePtr != nullptr )
 			{
-				CheckBoxStyle->CheckBoxStyle = *StylePtr;
+				WidgetStyle->CheckBoxStyle = *StylePtr;
 			}
 
 			Style_DEPRECATED = nullptr;
@@ -156,55 +156,55 @@ void UCheckBox::PostLoad()
 
 		if ( UncheckedImage_DEPRECATED != nullptr )
 		{
-			CheckBoxStyle->CheckBoxStyle.UncheckedImage = UncheckedImage_DEPRECATED->Brush;
+			WidgetStyle->CheckBoxStyle.UncheckedImage = UncheckedImage_DEPRECATED->Brush;
 			UncheckedImage_DEPRECATED = nullptr;
 		}
 
 		if ( UncheckedHoveredImage_DEPRECATED != nullptr )
 		{
-			CheckBoxStyle->CheckBoxStyle.UncheckedHoveredImage = UncheckedHoveredImage_DEPRECATED->Brush;
+			WidgetStyle->CheckBoxStyle.UncheckedHoveredImage = UncheckedHoveredImage_DEPRECATED->Brush;
 			UncheckedHoveredImage_DEPRECATED = nullptr;
 		}
 
 		if ( UncheckedPressedImage_DEPRECATED != nullptr )
 		{
-			CheckBoxStyle->CheckBoxStyle.UncheckedPressedImage = UncheckedPressedImage_DEPRECATED->Brush;
+			WidgetStyle->CheckBoxStyle.UncheckedPressedImage = UncheckedPressedImage_DEPRECATED->Brush;
 			UncheckedPressedImage_DEPRECATED = nullptr;
 		}
 
 		if ( CheckedImage_DEPRECATED != nullptr )
 		{
-			CheckBoxStyle->CheckBoxStyle.CheckedImage = CheckedImage_DEPRECATED->Brush;
+			WidgetStyle->CheckBoxStyle.CheckedImage = CheckedImage_DEPRECATED->Brush;
 			CheckedImage_DEPRECATED = nullptr;
 		}
 
 		if ( CheckedHoveredImage_DEPRECATED != nullptr )
 		{
-			CheckBoxStyle->CheckBoxStyle.CheckedHoveredImage = CheckedHoveredImage_DEPRECATED->Brush;
+			WidgetStyle->CheckBoxStyle.CheckedHoveredImage = CheckedHoveredImage_DEPRECATED->Brush;
 			CheckedHoveredImage_DEPRECATED = nullptr;
 		}
 
 		if ( CheckedPressedImage_DEPRECATED != nullptr )
 		{
-			CheckBoxStyle->CheckBoxStyle.CheckedPressedImage = CheckedPressedImage_DEPRECATED->Brush;
+			WidgetStyle->CheckBoxStyle.CheckedPressedImage = CheckedPressedImage_DEPRECATED->Brush;
 			CheckedPressedImage_DEPRECATED = nullptr;
 		}
 
 		if ( UndeterminedImage_DEPRECATED != nullptr )
 		{
-			CheckBoxStyle->CheckBoxStyle.UndeterminedImage = UndeterminedImage_DEPRECATED->Brush;
+			WidgetStyle->CheckBoxStyle.UndeterminedImage = UndeterminedImage_DEPRECATED->Brush;
 			UndeterminedImage_DEPRECATED = nullptr;
 		}
 
 		if ( UndeterminedHoveredImage_DEPRECATED != nullptr )
 		{
-			CheckBoxStyle->CheckBoxStyle.UndeterminedHoveredImage = UndeterminedHoveredImage_DEPRECATED->Brush;
+			WidgetStyle->CheckBoxStyle.UndeterminedHoveredImage = UndeterminedHoveredImage_DEPRECATED->Brush;
 			UndeterminedHoveredImage_DEPRECATED = nullptr;
 		}
 
 		if ( UndeterminedPressedImage_DEPRECATED != nullptr )
 		{
-			CheckBoxStyle->CheckBoxStyle.UndeterminedPressedImage = UndeterminedPressedImage_DEPRECATED->Brush;
+			WidgetStyle->CheckBoxStyle.UndeterminedPressedImage = UndeterminedPressedImage_DEPRECATED->Brush;
 			UndeterminedPressedImage_DEPRECATED = nullptr;
 		}
 	}
diff --git a/Engine/Source/Runtime/UMG/Private/Components/ComboBoxString.cpp b/Engine/Source/Runtime/UMG/Private/Components/ComboBoxString.cpp
--- a/Engine/Source/Runtime/UMG/Private/Components/ComboBoxString.cpp
+++ b/Engine/Source/Runtime/UMG/Private/Components/ComboBoxString.cpp
@@ -11,10 +11,10 @@ UComboBoxString::UComboBoxString(const FPostConstructInitializeProperties& PCIP)
 	: Super(PCIP)
 {
 	static const FName StyleName(TEXT("Style"));
-	ComboBoxStyle = PCIP.CreateDefaultSubobject<UComboBoxWidgetStyle>(this, StyleName);
+	WidgetStyle = PCIP.CreateDefaultSubobject<UComboBoxWidgetStyle>(this, StyleName);
 
 	SComboBox< TSharedPtr<FString> >::FArguments SlateDefaults;
-	ComboBoxStyle->ComboBoxStyle = *SlateDefaults._ComboBoxStyle;
+	WidgetStyle->ComboBoxStyle = *SlateDefaults._ComboBoxStyle;
 
 	ContentPadding = FMargin(4.0, 2.0);
 	MaxListHeight = 450.0f;
@@ -45,7 +45,7 @@ TSharedRef<SWidget> UComboBoxString::RebuildWidget()
 
 	MyComboBox =
 		SNew(SComboBox< TSharedPtr<FString> >)
-		.ComboBoxStyle(&ComboBoxStyle->ComboBoxStyle)
+		.ComboBoxStyle(&WidgetStyle->ComboBoxStyle)
 		.OptionsSource(&Options)
 		.InitiallySelectedItem(SelectedOptionPtr)
 		.ContentPadding(ContentPadding)
diff --git a/Engine/Source/Runtime/UMG/Private/Components/EditableText.cpp b/Engine/Source/Runtime/UMG/Private/Components/EditableText.cpp
--- a/Engine/Source/Runtime/UMG/Private/Components/EditableText.cpp
+++ b/Engine/Source/Runtime/UMG/Private/Components/EditableText.cpp
@@ -10,13 +10,18 @@
 UEditableText::UEditableText(const FPostConstructInitializeProperties& PCIP)
 	: Super(PCIP)
 {
+	static const FName StyleName(TEXT("Style"));
+	WidgetStyle = PCIP.CreateDefaultSubobject<UEditableTextWidgetStyle>(this, StyleName);
+
+	SEditableText::FArguments Defaults;
+	WidgetStyle->EditableTextStyle = *Defaults._Style;
+
 	ColorAndOpacity = FLinearColor::Black;
 
 	// HACK Special font initialization hack since there are no font assets yet for slate.
 	Font = FSlateFontInfo(TEXT("Slate/Fonts/Roboto-Bold.ttf"), 12);
 
 	// Grab other defaults from slate arguments.
-	SEditableTextBox::FArguments Defaults;
 	IsReadOnly = Defaults._IsReadOnly.Get();
 	IsPassword = Defaults._IsReadOnly.Get();
 	MinimumDesiredWidth = Defaults._MinDesiredWidth.Get();
@@ -43,16 +48,8 @@ TSharedRef<SWidget> UEditableText::RebuildWidget()
 		FontPath = FPaths::EngineContentDir() / Font.FontName.ToString();
 	}
 	
-	SEditableText::FArguments Defaults;
-	
-	const FEditableTextStyle* StylePtr = ( Style != NULL ) ? Style->GetStyle<FEditableTextStyle>() : NULL;
-	if ( StylePtr == NULL )
-	{
-		StylePtr = Defaults._Style;
-	}
-	
 	MyEditableText = SNew(SEditableText)
-	.Style(StylePtr)
+	.Style(&WidgetStyle->EditableTextStyle)
 	.Font(FSlateFontInfo(FontPath, Font.Size))
 	.MinDesiredWidth(MinimumDesiredWidth)
 	.IsCaretMovedWhenGainFocus(IsCaretMovedWhenGainFocus)
@@ -60,10 +57,6 @@ TSharedRef<SWidget> UEditableText::RebuildWidget()
 	.RevertTextOnEscape(RevertTextOnEscape)
 	.ClearKeyboardFocusOnCommit(ClearKeyboardFocusOnCommit)
 	.SelectAllTextOnCommit(SelectAllTextOnCommit)
-	.BackgroundImageSelected(BackgroundImageSelected ? TAttribute<const FSlateBrush*>(&BackgroundImageSelected->Brush) : TAttribute<const FSlateBrush*>())
-	.BackgroundImageSelectionTarget(BackgroundImageSelectionTarget ? TAttribute<const FSlateBrush*>(&BackgroundImageSelectionTarget->Brush) : TAttribute<const FSlateBrush*>())
-	.BackgroundImageComposing(BackgroundImageComposing ? TAttribute<const FSlateBrush*>(&BackgroundImageComposing->Brush) : TAttribute<const FSlateBrush*>())
-	.CaretImage(CaretImage ? TAttribute<const FSlateBrush*>(&CaretImage->Brush) : TAttribute<const FSlateBrush*>())
 	.OnTextChanged(BIND_UOBJECT_DELEGATE(FOnTextChanged, HandleOnTextChanged))
 	.OnTextCommitted(BIND_UOBJECT_DELEGATE(FOnTextCommitted, HandleOnTextCommitted))
 	;
@@ -106,6 +99,33 @@ void UEditableText::SetText(FText InText)
 	}
 }
 
+void UEditableText::SetIsPassword(bool InbIsPassword)
+{
+	IsPassword = InbIsPassword;
+	if ( MyEditableText.IsValid() )
+	{
+		MyEditableText->SetIsPassword(IsPassword);
+	}
+}
+
+void UEditableText::SetHintText(FText InHintText)
+{
+	HintText = InHintText;
+	if ( MyEditableText.IsValid() )
+	{
+		MyEditableText->SetHintText(HintText);
+	}
+}
+
+void UEditableText::SetIsReadOnly(bool InbIsReadyOnly)
+{
+	IsReadOnly = InbIsReadyOnly;
+	if ( MyEditableText.IsValid() )
+	{
+		MyEditableText->SetIsReadOnly(IsReadOnly);
+	}
+}
+
 void UEditableText::HandleOnTextChanged(const FText& Text)
 {
 	OnTextChanged.Broadcast(Text);
@@ -116,6 +136,49 @@ void UEditableText::HandleOnTextCommitted(const FText& Text, ETextCommit::Type C
 	OnTextCommitted.Broadcast(Text, CommitMethod);
 }
 
+void UEditableText::PostLoad()
+{
+	Super::PostLoad();
+
+	if ( GetLinkerUE4Version() < VER_UE4_DEPRECATE_UMG_STYLE_ASSETS )
+	{
+		if ( Style_DEPRECATED != nullptr )
+		{
+			const FEditableTextStyle* StylePtr = Style_DEPRECATED->GetStyle<FEditableTextStyle>();
+			if ( StylePtr != nullptr )
+			{
+				WidgetStyle->EditableTextStyle = *StylePtr;
+			}
+
+			Style_DEPRECATED = nullptr;
+		}
+
+		if ( BackgroundImageSelected_DEPRECATED != nullptr )
+		{
+			WidgetStyle->EditableTextStyle.BackgroundImageSelected = BackgroundImageSelected_DEPRECATED->Brush;
+			BackgroundImageSelected_DEPRECATED = nullptr;
+		}
+
+		if ( BackgroundImageSelectionTarget_DEPRECATED != nullptr )
+		{
+			WidgetStyle->EditableTextStyle.BackgroundImageSelectionTarget = BackgroundImageSelectionTarget_DEPRECATED->Brush;
+			BackgroundImageSelectionTarget_DEPRECATED = nullptr;
+		}
+
+		if ( BackgroundImageComposing_DEPRECATED != nullptr )
+		{
+			WidgetStyle->EditableTextStyle.BackgroundImageComposing = BackgroundImageComposing_DEPRECATED->Brush;
+			BackgroundImageComposing_DEPRECATED = nullptr;
+		}
+
+		if ( CaretImage_DEPRECATED != nullptr )
+		{
+			WidgetStyle->EditableTextStyle.CaretImage = CaretImage_DEPRECATED->Brush;
+			CaretImage_DEPRECATED = nullptr;
+		}
+	}
+}
+
 #if WITH_EDITOR
 
 const FSlateBrush* UEditableText::GetEditorIcon()
diff --git a/Engine/Source/Runtime/UMG/Public/Components/Button.h b/Engine/Source/Runtime/UMG/Public/Components/Button.h
--- a/Engine/Source/Runtime/UMG/Public/Components/Button.h
+++ b/Engine/Source/Runtime/UMG/Public/Components/Button.h
@@ -22,8 +22,8 @@ public:
 	USlateWidgetStyleAsset* Style_DEPRECATED;
 
 	/** The button style used at runtime by the slate button */
-	UPROPERTY(VisibleAnywhere, Instanced, Category=Appearance )
-	TSubobjectPtr<UButtonWidgetStyle> ButtonStyle;
+	UPROPERTY(VisibleAnywhere, Instanced, Category=Appearance, meta=( DisplayName="Style" ))
+	TSubobjectPtr<UButtonWidgetStyle> WidgetStyle;
 	
 	/** The color multiplier for the button content */
 	UPROPERTY(EditDefaultsOnly, Category=Appearance )
diff --git a/Engine/Source/Runtime/UMG/Public/Components/CheckBox.h b/Engine/Source/Runtime/UMG/Public/Components/CheckBox.h
--- a/Engine/Source/Runtime/UMG/Public/Components/CheckBox.h
+++ b/Engine/Source/Runtime/UMG/Public/Components/CheckBox.h
@@ -18,8 +18,8 @@ class UMG_API UCheckBox : public UContentWidget
 
 public:
 	/** The checkbox bar style */
-	UPROPERTY(VisibleAnywhere, Instanced, Category="Style")
-	TSubobjectPtr<UCheckBoxWidgetStyle> CheckBoxStyle;
+	UPROPERTY(VisibleAnywhere, Instanced, Category=Appearance, meta=( DisplayName="Style" ))
+	TSubobjectPtr<UCheckBoxWidgetStyle> WidgetStyle;
 
 	/** Style of the check box */
 	UPROPERTY()
diff --git a/Engine/Source/Runtime/UMG/Public/Components/ComboBoxString.h b/Engine/Source/Runtime/UMG/Public/Components/ComboBoxString.h
--- a/Engine/Source/Runtime/UMG/Public/Components/ComboBoxString.h
+++ b/Engine/Source/Runtime/UMG/Public/Components/ComboBoxString.h
@@ -20,8 +20,8 @@ class UMG_API UComboBoxString : public UWidget
 public:
 
 	/** The style */
-	UPROPERTY(VisibleAnywhere, Instanced, Category=Appearance )
-	TSubobjectPtr<UComboBoxWidgetStyle> ComboBoxStyle;
+	UPROPERTY(VisibleAnywhere, Instanced, Category=Appearance, meta=( DisplayName="Style" ))
+	TSubobjectPtr<UComboBoxWidgetStyle> WidgetStyle;
 
 	/** The list of items to be displayed on the combobox. */
 	UPROPERTY(EditDefaultsOnly, Category=Content)
diff --git a/Engine/Source/Runtime/UMG/Public/Components/EditableText.h b/Engine/Source/Runtime/UMG/Public/Components/EditableText.h
--- a/Engine/Source/Runtime/UMG/Public/Components/EditableText.h
+++ b/Engine/Source/Runtime/UMG/Public/Components/EditableText.h
@@ -2,6 +2,8 @@
 
 #pragma once
 
+#include "EditableTextWidgetStyle.h"
+
 #include "EditableText.generated.h"
 
 /** Editable text box widget */
@@ -33,25 +35,29 @@ public:
 	UPROPERTY()
 	FGetText HintTextDelegate;
 
+	/** The button style used at runtime by the slate button */
+	UPROPERTY(VisibleAnywhere, Instanced, Category=Appearance, meta=(DisplayName="Style") )
+	TSubobjectPtr<UEditableTextWidgetStyle> WidgetStyle;
+
 	/** Text style */
-	UPROPERTY(EditDefaultsOnly, Category=Style, meta=( DisplayThumbnail = "true" ))
-	USlateWidgetStyleAsset* Style;
+	UPROPERTY()
+	USlateWidgetStyleAsset* Style_DEPRECATED;
 
 	/** Background image for the selected text (overrides Style) */
-	UPROPERTY(EditDefaultsOnly, Category="Style", meta=( DisplayThumbnail = "true" ), AdvancedDisplay)
-	USlateBrushAsset* BackgroundImageSelected;
+	UPROPERTY()
+	USlateBrushAsset* BackgroundImageSelected_DEPRECATED;
 
 	/** Background image for the selection targeting effect (overrides Style) */
-	UPROPERTY(EditDefaultsOnly, Category="Style", meta=( DisplayThumbnail = "true" ), AdvancedDisplay)
-	USlateBrushAsset* BackgroundImageSelectionTarget;
+	UPROPERTY()
+	USlateBrushAsset* BackgroundImageSelectionTarget_DEPRECATED;
 
 	/** Background image for the composing text (overrides Style) */
-	UPROPERTY(EditDefaultsOnly, Category="Style", meta=( DisplayThumbnail = "true" ), AdvancedDisplay)
-	USlateBrushAsset* BackgroundImageComposing;
+	UPROPERTY()
+	USlateBrushAsset* BackgroundImageComposing_DEPRECATED;
 
 	/** Image brush used for the caret (overrides Style) */
-	UPROPERTY(EditDefaultsOnly, Category="Style", meta=( DisplayThumbnail = "true" ), AdvancedDisplay)
-	USlateBrushAsset* CaretImage;
+	UPROPERTY()
+	USlateBrushAsset* CaretImage_DEPRECATED;
 
 	/** Font color and opacity (overrides Style) */
 	UPROPERTY(EditDefaultsOnly, Category=Appearance)
@@ -104,25 +110,42 @@ public:
 	FOnEditableTextCommittedEvent OnTextCommitted;
 
 	/**
-	Gets the widget text
-	@return The widget text
-	*/
+	 * Gets the widget text
+	 * @return The widget text
+	 */
 	UFUNCTION(BlueprintCallable, Category = "Widget")
-		FText GetText() const;
+	FText GetText() const;
 
 	/**
-	Directly sets the widget text.
-	Warning: This will wipe any binding created for the Text property!
-	@param InText The text to assign to the widget
-	*/
+	 * Directly sets the widget text.
+	 * Warning: This will wipe any binding created for the Text property!
+	 * @param InText The text to assign to the widget
+	 */
 	UFUNCTION(BlueprintCallable, Category = "Widget")
-		void SetText(FText InText);
-	
-	virtual void ReleaseNativeWidget() override;
+	void SetText(FText InText);
+
+	UFUNCTION(BlueprintCallable, Category = "Widget")
+	void SetIsPassword(bool InbIsPassword);
+
+	UFUNCTION(BlueprintCallable, Category = "Widget")
+	void SetHintText(FText InHintText);
+
+	UFUNCTION(BlueprintCallable, Category = "Widget")
+	void SetIsReadOnly(bool InbIsReadyOnly);
 
 public:
+	
 	// UWidget interface
 	virtual void SynchronizeProperties() override;
+	// End of UWidget interface
+
+	// UVisual interface
+	virtual void ReleaseNativeWidget() override;
+	// End of UVisual interface
+
+	// Begin UObject interface
+	virtual void PostLoad() override;
+	// End of UObject interface
 
 #if WITH_EDITOR
 	virtual const FSlateBrush* GetEditorIcon() override;
