UMG - Using the slate struct directly for styles instead of the wrapper UObject.

diff --git a/Engine/Source/Runtime/UMG/Private/Components/Button.cpp b/Engine/Source/Runtime/UMG/Private/Components/Button.cpp
--- a/Engine/Source/Runtime/UMG/Private/Components/Button.cpp
+++ b/Engine/Source/Runtime/UMG/Private/Components/Button.cpp
@@ -10,11 +10,8 @@
 UButton::UButton(const FPostConstructInitializeProperties& PCIP)
 	: Super(PCIP)
 {
-	static const FName StyleName(TEXT("Style"));
-	WidgetStyle = PCIP.CreateDefaultSubobject<UButtonWidgetStyle>(this, StyleName);
-
 	SButton::FArguments ButtonDefaults;
-	WidgetStyle->ButtonStyle = *ButtonDefaults._ButtonStyle;
+	WidgetStyle = *ButtonDefaults._ButtonStyle;
 
 	ColorAndOpacity = FLinearColor::White;
 	BackgroundColor = FLinearColor::White;
@@ -35,7 +32,7 @@ void UButton::ReleaseNativeWidget()
 TSharedRef<SWidget> UButton::RebuildWidget()
 {
 	MyButton = SNew(SButton)
-		.ButtonStyle(&WidgetStyle->ButtonStyle)
+		.ButtonStyle(&WidgetStyle)
 		.ClickMethod(ClickMethod)
 		.TouchMethod(TouchMethod)
 		.IsFocusable(IsFocusable);
@@ -81,16 +78,6 @@ void UButton::OnSlotRemoved(UPanelSlot* Slot)
 	}
 }
 
-void UButton::SetButtonStyle(FButtonStyle InButtonStyle)
-{
-	WidgetStyle->ButtonStyle = InButtonStyle;
-}
-
-FButtonStyle UButton::GetButtonStyle()
-{
-	return WidgetStyle->ButtonStyle;
-}
-
 void UButton::SetColorAndOpacity(FLinearColor Color)
 {
 	ColorAndOpacity = Color;
@@ -139,7 +126,7 @@ void UButton::PostLoad()
 		const FButtonStyle* StylePtr = Style_DEPRECATED->GetStyle<FButtonStyle>();
 		if(StylePtr != nullptr)
 		{
-			WidgetStyle->ButtonStyle = *StylePtr;
+			WidgetStyle = *StylePtr;
 		}
 
 		Style_DEPRECATED = nullptr;
diff --git a/Engine/Source/Runtime/UMG/Private/Components/CheckBox.cpp b/Engine/Source/Runtime/UMG/Private/Components/CheckBox.cpp
--- a/Engine/Source/Runtime/UMG/Private/Components/CheckBox.cpp
+++ b/Engine/Source/Runtime/UMG/Private/Components/CheckBox.cpp
@@ -10,11 +10,8 @@
 UCheckBox::UCheckBox(const FPostConstructInitializeProperties& PCIP)
 	: Super(PCIP)
 {
-	static const FName StyleName(TEXT("Style"));
-	WidgetStyle = PCIP.CreateDefaultSubobject<UCheckBoxWidgetStyle>(this, StyleName);
-
 	SCheckBox::FArguments SlateDefaults;
-	WidgetStyle->CheckBoxStyle = *SlateDefaults._Style;
+	WidgetStyle = *SlateDefaults._Style;
 
 	CheckedState = ESlateCheckBoxState::Unchecked;
 
@@ -35,7 +32,7 @@ TSharedRef<SWidget> UCheckBox::RebuildWidget()
 {
 	MyCheckbox = SNew(SCheckBox)
 		.OnCheckStateChanged( BIND_UOBJECT_DELEGATE(FOnCheckStateChanged, SlateOnCheckStateChangedCallback) )
-		.Style(&WidgetStyle->CheckBoxStyle)
+		.Style(&WidgetStyle)
 		.HAlign( HorizontalAlignment )
 		.Padding( Padding )
 		.BorderBackgroundColor( BorderBackgroundColor )
@@ -148,7 +145,7 @@ void UCheckBox::PostLoad()
 			const FCheckBoxStyle* StylePtr = Style_DEPRECATED->GetStyle<FCheckBoxStyle>();
 			if ( StylePtr != nullptr )
 			{
-				WidgetStyle->CheckBoxStyle = *StylePtr;
+				WidgetStyle = *StylePtr;
 			}
 
 			Style_DEPRECATED = nullptr;
@@ -156,55 +153,55 @@ void UCheckBox::PostLoad()
 
 		if ( UncheckedImage_DEPRECATED != nullptr )
 		{
-			WidgetStyle->CheckBoxStyle.UncheckedImage = UncheckedImage_DEPRECATED->Brush;
+			WidgetStyle.UncheckedImage = UncheckedImage_DEPRECATED->Brush;
 			UncheckedImage_DEPRECATED = nullptr;
 		}
 
 		if ( UncheckedHoveredImage_DEPRECATED != nullptr )
 		{
-			WidgetStyle->CheckBoxStyle.UncheckedHoveredImage = UncheckedHoveredImage_DEPRECATED->Brush;
+			WidgetStyle.UncheckedHoveredImage = UncheckedHoveredImage_DEPRECATED->Brush;
 			UncheckedHoveredImage_DEPRECATED = nullptr;
 		}
 
 		if ( UncheckedPressedImage_DEPRECATED != nullptr )
 		{
-			WidgetStyle->CheckBoxStyle.UncheckedPressedImage = UncheckedPressedImage_DEPRECATED->Brush;
+			WidgetStyle.UncheckedPressedImage = UncheckedPressedImage_DEPRECATED->Brush;
 			UncheckedPressedImage_DEPRECATED = nullptr;
 		}
 
 		if ( CheckedImage_DEPRECATED != nullptr )
 		{
-			WidgetStyle->CheckBoxStyle.CheckedImage = CheckedImage_DEPRECATED->Brush;
+			WidgetStyle.CheckedImage = CheckedImage_DEPRECATED->Brush;
 			CheckedImage_DEPRECATED = nullptr;
 		}
 
 		if ( CheckedHoveredImage_DEPRECATED != nullptr )
 		{
-			WidgetStyle->CheckBoxStyle.CheckedHoveredImage = CheckedHoveredImage_DEPRECATED->Brush;
+			WidgetStyle.CheckedHoveredImage = CheckedHoveredImage_DEPRECATED->Brush;
 			CheckedHoveredImage_DEPRECATED = nullptr;
 		}
 
 		if ( CheckedPressedImage_DEPRECATED != nullptr )
 		{
-			WidgetStyle->CheckBoxStyle.CheckedPressedImage = CheckedPressedImage_DEPRECATED->Brush;
+			WidgetStyle.CheckedPressedImage = CheckedPressedImage_DEPRECATED->Brush;
 			CheckedPressedImage_DEPRECATED = nullptr;
 		}
 
 		if ( UndeterminedImage_DEPRECATED != nullptr )
 		{
-			WidgetStyle->CheckBoxStyle.UndeterminedImage = UndeterminedImage_DEPRECATED->Brush;
+			WidgetStyle.UndeterminedImage = UndeterminedImage_DEPRECATED->Brush;
 			UndeterminedImage_DEPRECATED = nullptr;
 		}
 
 		if ( UndeterminedHoveredImage_DEPRECATED != nullptr )
 		{
-			WidgetStyle->CheckBoxStyle.UndeterminedHoveredImage = UndeterminedHoveredImage_DEPRECATED->Brush;
+			WidgetStyle.UndeterminedHoveredImage = UndeterminedHoveredImage_DEPRECATED->Brush;
 			UndeterminedHoveredImage_DEPRECATED = nullptr;
 		}
 
 		if ( UndeterminedPressedImage_DEPRECATED != nullptr )
 		{
-			WidgetStyle->CheckBoxStyle.UndeterminedPressedImage = UndeterminedPressedImage_DEPRECATED->Brush;
+			WidgetStyle.UndeterminedPressedImage = UndeterminedPressedImage_DEPRECATED->Brush;
 			UndeterminedPressedImage_DEPRECATED = nullptr;
 		}
 	}
diff --git a/Engine/Source/Runtime/UMG/Private/Components/ComboBoxString.cpp b/Engine/Source/Runtime/UMG/Private/Components/ComboBoxString.cpp
--- a/Engine/Source/Runtime/UMG/Private/Components/ComboBoxString.cpp
+++ b/Engine/Source/Runtime/UMG/Private/Components/ComboBoxString.cpp
@@ -10,11 +10,8 @@
 UComboBoxString::UComboBoxString(const FPostConstructInitializeProperties& PCIP)
 	: Super(PCIP)
 {
-	static const FName StyleName(TEXT("Style"));
-	WidgetStyle = PCIP.CreateDefaultSubobject<UComboBoxWidgetStyle>(this, StyleName);
-
 	SComboBox< TSharedPtr<FString> >::FArguments SlateDefaults;
-	WidgetStyle->ComboBoxStyle = *SlateDefaults._ComboBoxStyle;
+	WidgetStyle = *SlateDefaults._ComboBoxStyle;
 
 	ContentPadding = FMargin(4.0, 2.0);
 	MaxListHeight = 450.0f;
@@ -45,7 +42,7 @@ TSharedRef<SWidget> UComboBoxString::RebuildWidget()
 
 	MyComboBox =
 		SNew(SComboBox< TSharedPtr<FString> >)
-		.ComboBoxStyle(&WidgetStyle->ComboBoxStyle)
+		.ComboBoxStyle(&WidgetStyle)
 		.OptionsSource(&Options)
 		.InitiallySelectedItem(SelectedOptionPtr)
 		.ContentPadding(ContentPadding)
diff --git a/Engine/Source/Runtime/UMG/Private/Components/EditableText.cpp b/Engine/Source/Runtime/UMG/Private/Components/EditableText.cpp
--- a/Engine/Source/Runtime/UMG/Private/Components/EditableText.cpp
+++ b/Engine/Source/Runtime/UMG/Private/Components/EditableText.cpp
@@ -10,11 +10,8 @@
 UEditableText::UEditableText(const FPostConstructInitializeProperties& PCIP)
 	: Super(PCIP)
 {
-	static const FName StyleName(TEXT("Style"));
-	WidgetStyle = PCIP.CreateDefaultSubobject<UEditableTextWidgetStyle>(this, StyleName);
-
 	SEditableText::FArguments Defaults;
-	WidgetStyle->EditableTextStyle = *Defaults._Style;
+	WidgetStyle = *Defaults._Style;
 
 	ColorAndOpacity = FLinearColor::Black;
 
@@ -49,7 +46,7 @@ TSharedRef<SWidget> UEditableText::RebuildWidget()
 	}
 	
 	MyEditableText = SNew(SEditableText)
-	.Style(&WidgetStyle->EditableTextStyle)
+	.Style(&WidgetStyle)
 	.Font(FSlateFontInfo(FontPath, Font.Size))
 	.MinDesiredWidth(MinimumDesiredWidth)
 	.IsCaretMovedWhenGainFocus(IsCaretMovedWhenGainFocus)
@@ -147,7 +144,7 @@ void UEditableText::PostLoad()
 			const FEditableTextStyle* StylePtr = Style_DEPRECATED->GetStyle<FEditableTextStyle>();
 			if ( StylePtr != nullptr )
 			{
-				WidgetStyle->EditableTextStyle = *StylePtr;
+				WidgetStyle = *StylePtr;
 			}
 
 			Style_DEPRECATED = nullptr;
@@ -155,25 +152,25 @@ void UEditableText::PostLoad()
 
 		if ( BackgroundImageSelected_DEPRECATED != nullptr )
 		{
-			WidgetStyle->EditableTextStyle.BackgroundImageSelected = BackgroundImageSelected_DEPRECATED->Brush;
+			WidgetStyle.BackgroundImageSelected = BackgroundImageSelected_DEPRECATED->Brush;
 			BackgroundImageSelected_DEPRECATED = nullptr;
 		}
 
 		if ( BackgroundImageSelectionTarget_DEPRECATED != nullptr )
 		{
-			WidgetStyle->EditableTextStyle.BackgroundImageSelectionTarget = BackgroundImageSelectionTarget_DEPRECATED->Brush;
+			WidgetStyle.BackgroundImageSelectionTarget = BackgroundImageSelectionTarget_DEPRECATED->Brush;
 			BackgroundImageSelectionTarget_DEPRECATED = nullptr;
 		}
 
 		if ( BackgroundImageComposing_DEPRECATED != nullptr )
 		{
-			WidgetStyle->EditableTextStyle.BackgroundImageComposing = BackgroundImageComposing_DEPRECATED->Brush;
+			WidgetStyle.BackgroundImageComposing = BackgroundImageComposing_DEPRECATED->Brush;
 			BackgroundImageComposing_DEPRECATED = nullptr;
 		}
 
 		if ( CaretImage_DEPRECATED != nullptr )
 		{
-			WidgetStyle->EditableTextStyle.CaretImage = CaretImage_DEPRECATED->Brush;
+			WidgetStyle.CaretImage = CaretImage_DEPRECATED->Brush;
 			CaretImage_DEPRECATED = nullptr;
 		}
 	}
diff --git a/Engine/Source/Runtime/UMG/Public/Components/Button.h b/Engine/Source/Runtime/UMG/Public/Components/Button.h
--- a/Engine/Source/Runtime/UMG/Public/Components/Button.h
+++ b/Engine/Source/Runtime/UMG/Public/Components/Button.h
@@ -21,9 +21,10 @@ public:
 	UPROPERTY()
 	USlateWidgetStyleAsset* Style_DEPRECATED;
 
-	/** The button style used at runtime by the slate button */
-	UPROPERTY(VisibleAnywhere, Instanced, Category=Appearance, meta=( DisplayName="Style" ))
-	TSubobjectPtr<UButtonWidgetStyle> WidgetStyle;
+	/** The button style used at runtime */
+	UPROPERTY(EditAnywhere, BlueprintReadWrite, Category="Style", meta=( DisplayName="Style" ))
+	FButtonStyle WidgetStyle;
+	/*ShowOnlyInnerProperties,*/
 	
 	/** The color multiplier for the button content */
 	UPROPERTY(EditDefaultsOnly, Category=Appearance )
@@ -50,14 +51,6 @@ public:
 	/** Called when the button is clicked */
 	UPROPERTY(BlueprintAssignable)
 	FOnButtonClickedEvent OnClicked;
-
-	/** Sets the look and feel of a button from a new button style struct. */
-	UFUNCTION(BlueprintCallable, Category="Button|Appearance")
-	void SetButtonStyle(FButtonStyle InButtonStyle);
-
-	/** Gets a dynamic button style that can be changed. */
-	UFUNCTION(BlueprintCallable, Category="Button|Appearance")
-	FButtonStyle GetButtonStyle();
 	
 	/** Sets the color multiplier for the button content */
 	UFUNCTION(BlueprintCallable, Category="Button|Appearance")
diff --git a/Engine/Source/Runtime/UMG/Public/Components/CheckBox.h b/Engine/Source/Runtime/UMG/Public/Components/CheckBox.h
--- a/Engine/Source/Runtime/UMG/Public/Components/CheckBox.h
+++ b/Engine/Source/Runtime/UMG/Public/Components/CheckBox.h
@@ -18,8 +18,8 @@ class UMG_API UCheckBox : public UContentWidget
 
 public:
 	/** The checkbox bar style */
-	UPROPERTY(VisibleAnywhere, Instanced, Category=Appearance, meta=( DisplayName="Style" ))
-	TSubobjectPtr<UCheckBoxWidgetStyle> WidgetStyle;
+	UPROPERTY(EditAnywhere, BlueprintReadWrite, Category="Style", meta=( DisplayName="Style" ))
+	FCheckBoxStyle WidgetStyle;
 
 	/** Style of the check box */
 	UPROPERTY()
diff --git a/Engine/Source/Runtime/UMG/Public/Components/ComboBoxString.h b/Engine/Source/Runtime/UMG/Public/Components/ComboBoxString.h
--- a/Engine/Source/Runtime/UMG/Public/Components/ComboBoxString.h
+++ b/Engine/Source/Runtime/UMG/Public/Components/ComboBoxString.h
@@ -20,8 +20,8 @@ class UMG_API UComboBoxString : public UWidget
 public:
 
 	/** The style */
-	UPROPERTY(VisibleAnywhere, Instanced, Category=Appearance, meta=( DisplayName="Style" ))
-	TSubobjectPtr<UComboBoxWidgetStyle> WidgetStyle;
+	UPROPERTY(EditAnywhere, BlueprintReadWrite, Category="Style", meta=( DisplayName="Style" ))
+	FComboBoxStyle WidgetStyle;
 
 	/** The list of items to be displayed on the combobox. */
 	UPROPERTY(EditDefaultsOnly, Category=Content)
diff --git a/Engine/Source/Runtime/UMG/Public/Components/EditableText.h b/Engine/Source/Runtime/UMG/Public/Components/EditableText.h
--- a/Engine/Source/Runtime/UMG/Public/Components/EditableText.h
+++ b/Engine/Source/Runtime/UMG/Public/Components/EditableText.h
@@ -35,9 +35,9 @@ public:
 	UPROPERTY()
 	FGetText HintTextDelegate;
 
-	/** The button style used at runtime by the slate button */
-	UPROPERTY(VisibleAnywhere, Instanced, Category=Appearance, meta=(DisplayName="Style") )
-	TSubobjectPtr<UEditableTextWidgetStyle> WidgetStyle;
+	/** The style */
+	UPROPERTY(EditAnywhere, BlueprintReadWrite, Category="Style", meta=( DisplayName="Style" ))
+	FEditableTextStyle WidgetStyle;
 
 	/** Text style */
 	UPROPERTY()
