UMG - Upgrading progress bar to use the new method of styling.

diff --git a/Engine/Source/Runtime/UMG/Private/Components/ProgressBar.cpp b/Engine/Source/Runtime/UMG/Private/Components/ProgressBar.cpp
--- a/Engine/Source/Runtime/UMG/Private/Components/ProgressBar.cpp
+++ b/Engine/Source/Runtime/UMG/Private/Components/ProgressBar.cpp
@@ -10,11 +10,8 @@
 UProgressBar::UProgressBar(const FPostConstructInitializeProperties& PCIP)
 	: Super(PCIP)
 {
-	static const FName StyleName(TEXT("Style"));
-	ProgressStyle = PCIP.CreateDefaultSubobject<UProgressWidgetStyle>(this, StyleName);
-
 	SProgressBar::FArguments SlateDefaults;
-	ProgressStyle->ProgressBarStyle = *SlateDefaults._Style;
+	WidgetStyle = *SlateDefaults._Style;
 
 	BarFillType = EProgressBarFillType::LeftToRight;
 	bIsMarquee = false;
@@ -41,7 +38,7 @@ void UProgressBar::SynchronizeProperties()
 {
 	TAttribute< TOptional<float> > PercentBinding = OPTIONAL_BINDING_CONVERT(float, Percent, TOptional<float>, ConvertFloatToOptionalFloat);
 
-	MyProgressBar->SetStyle(&ProgressStyle->ProgressBarStyle);
+	MyProgressBar->SetStyle(&WidgetStyle);
 
 	MyProgressBar->SetPercent(bIsMarquee ? TOptional<float>() : PercentBinding);
 	
@@ -78,7 +75,7 @@ void UProgressBar::PostLoad()
 			const FProgressBarStyle* StylePtr = Style_DEPRECATED->GetStyle<FProgressBarStyle>();
 			if ( StylePtr != nullptr )
 			{
-				ProgressStyle->ProgressBarStyle = *StylePtr;
+				WidgetStyle = *StylePtr;
 			}
 
 			Style_DEPRECATED = nullptr;
@@ -86,19 +83,19 @@ void UProgressBar::PostLoad()
 
 		if ( BackgroundImage_DEPRECATED != nullptr )
 		{
-			ProgressStyle->ProgressBarStyle.BackgroundImage = BackgroundImage_DEPRECATED->Brush;
+			WidgetStyle.BackgroundImage = BackgroundImage_DEPRECATED->Brush;
 			BackgroundImage_DEPRECATED = nullptr;
 		}
 
 		if ( FillImage_DEPRECATED != nullptr )
 		{
-			ProgressStyle->ProgressBarStyle.FillImage = FillImage_DEPRECATED->Brush;
+			WidgetStyle.FillImage = FillImage_DEPRECATED->Brush;
 			FillImage_DEPRECATED = nullptr;
 		}
 
 		if ( MarqueeImage_DEPRECATED != nullptr )
 		{
-			ProgressStyle->ProgressBarStyle.MarqueeImage = MarqueeImage_DEPRECATED->Brush;
+			WidgetStyle.MarqueeImage = MarqueeImage_DEPRECATED->Brush;
 			MarqueeImage_DEPRECATED = nullptr;
 		}
 	}
diff --git a/Engine/Source/Runtime/UMG/Public/Components/ProgressBar.h b/Engine/Source/Runtime/UMG/Public/Components/ProgressBar.h
--- a/Engine/Source/Runtime/UMG/Public/Components/ProgressBar.h
+++ b/Engine/Source/Runtime/UMG/Public/Components/ProgressBar.h
@@ -15,8 +15,8 @@ class UMG_API UProgressBar : public UWidget
 public:
 
 	/** The progress bar style */
-	UPROPERTY(VisibleAnywhere, Instanced, Category="Style")
-	TSubobjectPtr<UProgressWidgetStyle> ProgressStyle;
+	UPROPERTY(EditAnywhere, BlueprintReadWrite, Category="Style", meta=( DisplayName="Style" ))
+	FProgressBarStyle WidgetStyle;
 
 	/** Style used for the progress bar */
 	UPROPERTY()
