UMG - Upgrading ScrollBox to use the new method of styling.

diff --git a/Engine/Source/Runtime/UMG/Private/Components/ScrollBox.cpp b/Engine/Source/Runtime/UMG/Private/Components/ScrollBox.cpp
--- a/Engine/Source/Runtime/UMG/Private/Components/ScrollBox.cpp
+++ b/Engine/Source/Runtime/UMG/Private/Components/ScrollBox.cpp
@@ -12,9 +12,13 @@ UScrollBox::UScrollBox(const FPostConstructInitializeProperties& PCIP)
 {
 	bIsVariable = false;
 
+	Orientation = Orient_Vertical;
+
 	SScrollBox::FArguments Defaults;
 	Visiblity = UWidget::ConvertRuntimeToSerializedVisiblity(Defaults._Visibility.Get());
-	Orientation = Defaults._Orientation;
+
+	WidgetStyle = *Defaults._Style;
+	WidgetBarStyle = *Defaults._ScrollBarStyle;
 }
 
 void UScrollBox::ReleaseNativeWidget()
@@ -53,23 +57,9 @@ void UScrollBox::OnSlotRemoved(UPanelSlot* Slot)
 
 TSharedRef<SWidget> UScrollBox::RebuildWidget()
 {
-	const FScrollBoxStyle* StylePtr = ( Style != nullptr ) ? Style->GetStyle<FScrollBoxStyle>() : nullptr;
-	if ( StylePtr == nullptr )
-	{
-		SScrollBox::FArguments Defaults;
-		StylePtr = Defaults._Style;
-	}
-
-	const FScrollBarStyle* BarStylePtr = ( BarStyle != nullptr ) ? BarStyle->GetStyle<FScrollBarStyle>() : nullptr;
-	if ( BarStylePtr == nullptr )
-	{
-		SScrollBox::FArguments Defaults;
-		BarStylePtr = Defaults._ScrollBarStyle;
-	}
-
 	MyScrollBox = SNew(SScrollBox)
-		.Style(StylePtr)
-		.ScrollBarStyle(BarStylePtr)
+		.Style(&WidgetStyle)
+		.ScrollBarStyle(&WidgetBarStyle)
 		.Orientation(Orientation);
 
 	for ( UPanelSlot* Slot : Slots )
@@ -100,6 +90,36 @@ void UScrollBox::SetScrollOffset(float NewScrollOffset)
 	}
 }
 
+void UScrollBox::PostLoad()
+{
+	Super::PostLoad();
+
+	if ( GetLinkerUE4Version() < VER_UE4_DEPRECATE_UMG_STYLE_ASSETS )
+	{
+		if ( Style_DEPRECATED != nullptr )
+		{
+			const FScrollBoxStyle* StylePtr = Style_DEPRECATED->GetStyle<FScrollBoxStyle>();
+			if ( StylePtr != nullptr )
+			{
+				WidgetStyle = *StylePtr;
+			}
+
+			Style_DEPRECATED = nullptr;
+		}
+
+		if ( BarStyle_DEPRECATED != nullptr )
+		{
+			const FScrollBarStyle* StylePtr = BarStyle_DEPRECATED->GetStyle<FScrollBarStyle>();
+			if ( StylePtr != nullptr )
+			{
+				WidgetBarStyle = *StylePtr;
+			}
+
+			BarStyle_DEPRECATED = nullptr;
+		}
+	}
+}
+
 #if WITH_EDITOR
 
 const FSlateBrush* UScrollBox::GetEditorIcon()
diff --git a/Engine/Source/Runtime/UMG/Public/Components/ScrollBox.h b/Engine/Source/Runtime/UMG/Public/Components/ScrollBox.h
--- a/Engine/Source/Runtime/UMG/Public/Components/ScrollBox.h
+++ b/Engine/Source/Runtime/UMG/Public/Components/ScrollBox.h
@@ -10,20 +10,24 @@ class UMG_API UScrollBox : public UPanelWidget
 {
 	GENERATED_UCLASS_BODY()
 
-	/** Style of the scrollbox */
-	UPROPERTY(EditDefaultsOnly, Category=Style, meta=( DisplayThumbnail = "true" ))
-	USlateWidgetStyleAsset* Style;
+	/** The style */
+	UPROPERTY(EditAnywhere, BlueprintReadWrite, Category="Style", meta=( DisplayName="Style" ))
+	FScrollBoxStyle WidgetStyle;
 
-	/** Style of the scrollbox's bar */
-	UPROPERTY(EditDefaultsOnly, Category=Style, meta=( DisplayThumbnail = "true" ))
-	USlateWidgetStyleAsset* BarStyle;
+	/** The bar style */
+	UPROPERTY(EditAnywhere, BlueprintReadWrite, Category="Style", meta=( DisplayName="Bar Style" ))
+	FScrollBarStyle WidgetBarStyle;
+
+	UPROPERTY()
+	USlateWidgetStyleAsset* Style_DEPRECATED;
+
+	UPROPERTY()
+	USlateWidgetStyleAsset* BarStyle_DEPRECATED;
 
 	/** The orientation of the scrolling and stacking in the box. */
 	UPROPERTY(EditDefaultsOnly, Category = "Behavior")
 	TEnumAsByte<EOrientation> Orientation;
 
-	virtual void ReleaseNativeWidget() override;
-
 	/** Updates the scroll offset of the scrollbox */
 	UFUNCTION(BlueprintCallable, Category="Widget")
 	void SetScrollOffset(float NewScrollOffset);
@@ -32,6 +36,14 @@ class UMG_API UScrollBox : public UPanelWidget
 	virtual void SynchronizeProperties() override;
 	// End of UWidget interface
 
+	// UVisual interface
+	virtual void ReleaseNativeWidget() override;
+	// End of UVisual interface
+
+	// Begin UObject interface
+	virtual void PostLoad() override;
+	// End of UObject interface
+
 #if WITH_EDITOR
 	// UWidget interface
 	virtual const FSlateBrush* GetEditorIcon() override;
