UMG - Upgrading several more widgets to use the new mechanism for styling.  Making a few things advanced/not.

diff --git a/Engine/Source/Runtime/Slate/Public/Widgets/Input/SSpinBox.h b/Engine/Source/Runtime/Slate/Public/Widgets/Input/SSpinBox.h
--- a/Engine/Source/Runtime/Slate/Public/Widgets/Input/SSpinBox.h
+++ b/Engine/Source/Runtime/Slate/Public/Widgets/Input/SSpinBox.h
@@ -87,14 +87,16 @@ public:
 	{
 		check(InArgs._Style);
 
+		Style = InArgs._Style;
+
 		ForegroundColor = InArgs._Style->ForegroundColor;
+
 		ValueAttribute = InArgs._Value;
 		OnValueChanged = InArgs._OnValueChanged;
 		OnValueCommitted = InArgs._OnValueCommitted;
 		OnBeginSliderMovement = InArgs._OnBeginSliderMovement;
 		OnEndSliderMovement = InArgs._OnEndSliderMovement;
 		MinDesiredWidth = InArgs._MinDesiredWidth;
-		ArrowImageWidth = InArgs._Style->ArrowsImage.ImageSize.X;
 	
 		MinValue = InArgs._MinValue;
 		MaxValue = InArgs._MaxValue;
@@ -126,7 +128,8 @@ public:
 		.Padding( InArgs._ContentPadding )
 		[
 			SNew(SHorizontalBox)
-			+SHorizontalBox::Slot()
+
+			+ SHorizontalBox::Slot()
 			.FillWidth(1.0f)
 			.Padding( TextMargin )
 			.HAlign(HAlign_Fill) 
@@ -137,7 +140,8 @@ public:
 				.Text( this, &SSpinBox<NumericType>::GetValueAsString )
 				.MinDesiredWidth( this, &SSpinBox<NumericType>::GetTextMinDesiredWidth )
 			]
-			+SHorizontalBox::Slot()
+
+			+ SHorizontalBox::Slot()
 			.FillWidth(1.0f)
 			.Padding( TextMargin )
 			.HAlign(HAlign_Fill) 
@@ -153,17 +157,17 @@ public:
 				.ClearKeyboardFocusOnCommit( InArgs._ClearKeyboardFocusOnCommit )
 				.SelectAllTextOnCommit( InArgs._SelectAllTextOnCommit )
 				.MinDesiredWidth( this, &SSpinBox<NumericType>::GetTextMinDesiredWidth )
-			]			
-			+SHorizontalBox::Slot()
+			]
+
+			+ SHorizontalBox::Slot()
 			.AutoWidth()
 			.HAlign(HAlign_Fill) 
 			.VAlign(VAlign_Center) 
 			[
 				SNew(SImage)
 				.Image( &InArgs._Style->ArrowsImage )
-				.ColorAndOpacity( FSlateColor::UseForeground() )
+				.ColorAndOpacity(FSlateColor::UseForeground())
 			]
-
 		];
 	}
 	
@@ -191,7 +195,7 @@ public:
 			BackgroundImage,
 			MyClippingRect,
 			DrawEffects,
-			InWidgetStyle.GetColorAndOpacityTint()
+			BackgroundImage->GetTint(InWidgetStyle) * InWidgetStyle.GetColorAndOpacityTint()
 			);
 
 		const int32 FilledLayer = BackgroundLayer + 1;
@@ -222,7 +226,7 @@ public:
 					FillImage,
 					MyClippingRect,
 					DrawEffects,
-					InWidgetStyle.GetColorAndOpacityTint()
+					FillImage->GetTint(InWidgetStyle) * InWidgetStyle.GetColorAndOpacityTint()
 					);
 			}
 		}
@@ -701,6 +705,8 @@ private:
 		bUnlimitedSpinRange = !((MinValue.Get().IsSet() && MaxValue.Get().IsSet()) || (MinSliderValue.Get().IsSet() && MaxSliderValue.Get().IsSet()));
 	}
 
+	const FSpinBoxStyle* Style;
+
 	const FSlateBrush* BackgroundHoveredBrush;
 	const FSlateBrush* BackgroundBrush;
 	const FSlateBrush* ActiveFillBrush;
@@ -718,7 +724,7 @@ private:
 	TAttribute<float> MinDesiredWidth;
 	float GetTextMinDesiredWidth() const
 	{
-		return FMath::Max(0.0f, MinDesiredWidth.Get() - ArrowImageWidth);
+		return FMath::Max(0.0f, MinDesiredWidth.Get() - Style->ArrowsImage.ImageSize.X);
 	}
 
 	/** Whether the user is dragging the slider */
@@ -738,7 +744,4 @@ private:
 	/** This is the cached value the user believes it to be (usually different due to truncation to an int). Used for identifying 
 		external forces on the spinbox and syncing the internal value to them. Synced when a value is committed to the spinbox. */
 	NumericType CachedExternalValue;
-
-	/** The width of the arrow image, used for calculating the min desired width of the internal text controls. */
-	float ArrowImageWidth;
 };
diff --git a/Engine/Source/Runtime/SlateCore/Public/Styling/SlateColor.h b/Engine/Source/Runtime/SlateCore/Public/Styling/SlateColor.h
--- a/Engine/Source/Runtime/SlateCore/Public/Styling/SlateColor.h
+++ b/Engine/Source/Runtime/SlateCore/Public/Styling/SlateColor.h
@@ -16,16 +16,16 @@ namespace ESlateColorStylingMode
 	enum Type
 	{
 		/** Color value is stored in this Slate color. */
-		UseColor_Specified,
+		UseColor_Specified UMETA(DisplayName="Specified Color"),
 
 		/** Color value is stored in the linked color. */
-		UseColor_Specified_Link,
+		UseColor_Specified_Link UMETA(Hidden),
 
 		/** Use the widget's foreground color. */
-		UseColor_Foreground,
+		UseColor_Foreground UMETA(DisplayName="Foreground Color"),
 
 		/** Use the widget's subdued color. */
-		UseColor_Foreground_Subdued
+		UseColor_Foreground_Subdued UMETA(Hidden)
 	};
 }
 
diff --git a/Engine/Source/Runtime/SlateCore/Public/Styling/SlateTypes.h b/Engine/Source/Runtime/SlateCore/Public/Styling/SlateTypes.h
--- a/Engine/Source/Runtime/SlateCore/Public/Styling/SlateTypes.h
+++ b/Engine/Source/Runtime/SlateCore/Public/Styling/SlateTypes.h
@@ -182,7 +182,7 @@ struct SLATECORE_API FTextBlockStyle : public FSlateWidgetStyle
 	FTextBlockStyle& SetFontSize(uint16 InSize) { Font.Size = InSize; return *this; }
 
 	/** The color and opacity of this text */
-	UPROPERTY(EditAnywhere, BlueprintReadWrite, Category=Appearance)
+	UPROPERTY(EditAnywhere, BlueprintReadWrite, Category=Appearance, meta=(DisplayName="Color"))
 	FSlateColor ColorAndOpacity;
 	FTextBlockStyle& SetColorAndOpacity(const FSlateColor& InColorAndOpacity) { ColorAndOpacity = InColorAndOpacity; return *this; }
 
@@ -202,12 +202,12 @@ struct SLATECORE_API FTextBlockStyle : public FSlateWidgetStyle
 	FTextBlockStyle& SetSelectedBackgroundColor(const FSlateColor& InSelectedBackgroundColor) { SelectedBackgroundColor = InSelectedBackgroundColor; return *this; }
 
 	/** The color of highlighted text */
-	UPROPERTY(EditAnywhere, BlueprintReadWrite, Category=Appearance)
+	UPROPERTY(EditAnywhere, BlueprintReadWrite, Category=Appearance, AdvancedDisplay)
 	FLinearColor HighlightColor;
 	FTextBlockStyle& SetHighlightColor(const FLinearColor& InHighlightColor) { HighlightColor = InHighlightColor; return *this; }
 
 	/** The shape of highlighted text */
-	UPROPERTY(EditAnywhere, BlueprintReadWrite, Category=Appearance)
+	UPROPERTY(EditAnywhere, BlueprintReadWrite, Category=Appearance, AdvancedDisplay)
 	FSlateBrush HighlightShape;
 	FTextBlockStyle& SetHighlightShape( const FSlateBrush& InHighlightShape ){ HighlightShape = InHighlightShape; return *this; }
 };
@@ -931,7 +931,7 @@ struct SLATECORE_API FSpinBoxStyle : public FSlateWidgetStyle
 	FSpinBoxStyle& SetArrowsImage( const FSlateBrush& InArrowsImage ){ ArrowsImage = InArrowsImage; return *this; }
 
 	/** Color used to draw the spinbox foreground elements */
-	UPROPERTY(EditAnywhere, BlueprintReadWrite, Category=Appearance)
+	UPROPERTY()
 	FSlateColor ForegroundColor;
 	FSpinBoxStyle& SetForegroundColor( const FSlateColor& InForegroundColor ){ ForegroundColor = InForegroundColor; return *this; }
 
diff --git a/Engine/Source/Runtime/UMG/Private/Components/Border.cpp b/Engine/Source/Runtime/UMG/Private/Components/Border.cpp
--- a/Engine/Source/Runtime/UMG/Private/Components/Border.cpp
+++ b/Engine/Source/Runtime/UMG/Private/Components/Border.cpp
@@ -255,7 +255,7 @@ void UBorder::PostLoad()
 {
 	Super::PostLoad();
 
-	if ( Brush_DEPRECATED != NULL )
+	if ( GetLinkerUE4Version() < VER_UE4_DEPRECATE_UMG_STYLE_ASSETS && Brush_DEPRECATED != nullptr )
 	{
 		Background = Brush_DEPRECATED->Brush;
 		Brush_DEPRECATED = nullptr;
diff --git a/Engine/Source/Runtime/UMG/Private/Components/EditableTextBox.cpp b/Engine/Source/Runtime/UMG/Private/Components/EditableTextBox.cpp
--- a/Engine/Source/Runtime/UMG/Private/Components/EditableTextBox.cpp
+++ b/Engine/Source/Runtime/UMG/Private/Components/EditableTextBox.cpp
@@ -28,6 +28,8 @@ UEditableTextBox::UEditableTextBox(const FPostConstructInitializeProperties& PCI
 	RevertTextOnEscape = Defaults._RevertTextOnEscape.Get();
 	ClearKeyboardFocusOnCommit = Defaults._ClearKeyboardFocusOnCommit.Get();
 	SelectAllTextOnCommit = Defaults._SelectAllTextOnCommit.Get();
+
+	WidgetStyle = *Defaults._Style;
 }
 
 void UEditableTextBox::ReleaseNativeWidget()
@@ -47,6 +49,7 @@ TSharedRef<SWidget> UEditableTextBox::RebuildWidget()
 	}
 
 	MyEditableTextBlock = SNew(SEditableTextBox)
+		.Style(&WidgetStyle)
 		.Font(FSlateFontInfo(FontPath, Font.Size))
 		.ForegroundColor(ForegroundColor)
 		.BackgroundColor(BackgroundColor)
@@ -69,12 +72,6 @@ void UEditableTextBox::SynchronizeProperties()
 {
 	Super::SynchronizeProperties();
 
-	const FEditableTextBoxStyle* StylePtr = ( Style != NULL ) ? Style->GetStyle<FEditableTextBoxStyle>() : NULL;
-	if ( StylePtr )
-	{
-		MyEditableTextBlock->SetStyle(StylePtr);
-	}
-
 	TAttribute<FText> TextBinding = OPTIONAL_BINDING(FText, Text);
 	TAttribute<FText> HintTextBinding = OPTIONAL_BINDING(FText, HintText);
 
@@ -132,6 +129,25 @@ void UEditableTextBox::HandleOnTextCommitted(const FText& Text, ETextCommit::Typ
 	OnTextCommitted.Broadcast(Text, CommitMethod);
 }
 
+void UEditableTextBox::PostLoad()
+{
+	Super::PostLoad();
+
+	if ( GetLinkerUE4Version() < VER_UE4_DEPRECATE_UMG_STYLE_ASSETS )
+	{
+		if ( Style_DEPRECATED != nullptr )
+		{
+			const FEditableTextBoxStyle* StylePtr = Style_DEPRECATED->GetStyle<FEditableTextBoxStyle>();
+			if ( StylePtr != nullptr )
+			{
+				WidgetStyle = *StylePtr;
+			}
+
+			Style_DEPRECATED = nullptr;
+		}
+	}
+}
+
 #if WITH_EDITOR
 
 const FSlateBrush* UEditableTextBox::GetEditorIcon()
diff --git a/Engine/Source/Runtime/UMG/Private/Components/Image.cpp b/Engine/Source/Runtime/UMG/Private/Components/Image.cpp
--- a/Engine/Source/Runtime/UMG/Private/Components/Image.cpp
+++ b/Engine/Source/Runtime/UMG/Private/Components/Image.cpp
@@ -18,7 +18,7 @@ void UImage::PostLoad()
 {
 	Super::PostLoad();
 
-	if ( Image_DEPRECATED != NULL )
+	if ( GetLinkerUE4Version() < VER_UE4_DEPRECATE_UMG_STYLE_ASSETS && Image_DEPRECATED != nullptr )
 	{
 		Brush = Image_DEPRECATED->Brush;
 		Image_DEPRECATED = nullptr;
diff --git a/Engine/Source/Runtime/UMG/Private/Components/ScrollBar.cpp b/Engine/Source/Runtime/UMG/Private/Components/ScrollBar.cpp
--- a/Engine/Source/Runtime/UMG/Private/Components/ScrollBar.cpp
+++ b/Engine/Source/Runtime/UMG/Private/Components/ScrollBar.cpp
@@ -15,19 +15,22 @@ UScrollBar::UScrollBar(const FPostConstructInitializeProperties& PCIP)
 	bAlwaysShowScrollbar = true;
 	Orientation = Orient_Vertical;
 	Thickness = FVector2D(12.0f, 12.0f);
+
+	SScrollBar::FArguments Defaults;
+	WidgetStyle = *Defaults._Style;
 }
 
-TSharedRef<SWidget> UScrollBar::RebuildWidget()
+void UScrollBar::ReleaseNativeWidget()
 {
-	const FScrollBarStyle* StylePtr = ( Style != nullptr ) ? Style->GetStyle<FScrollBarStyle>() : nullptr;
-	if ( StylePtr == nullptr )
-	{
-		SScrollBar::FArguments Defaults;
-		StylePtr = Defaults._Style;
-	}
+	Super::ReleaseNativeWidget();
 
+	MyScrollBar.Reset();
+}
+
+TSharedRef<SWidget> UScrollBar::RebuildWidget()
+{
 	MyScrollBar = SNew(SScrollBar)
-		.Style(StylePtr)
+		.Style(&WidgetStyle)
 		.AlwaysShowScrollbar(bAlwaysShowScrollbar)
 		.Orientation(Orientation)
 		.Thickness(Thickness);
@@ -50,6 +53,25 @@ void UScrollBar::SetState(float InOffsetFraction, float InThumbSizeFraction)
 	}
 }
 
+void UScrollBar::PostLoad()
+{
+	Super::PostLoad();
+
+	if ( GetLinkerUE4Version() < VER_UE4_DEPRECATE_UMG_STYLE_ASSETS )
+	{
+		if ( Style_DEPRECATED != nullptr )
+		{
+			const FScrollBarStyle* StylePtr = Style_DEPRECATED->GetStyle<FScrollBarStyle>();
+			if ( StylePtr != nullptr )
+			{
+				WidgetStyle = *StylePtr;
+			}
+
+			Style_DEPRECATED = nullptr;
+		}
+	}
+}
+
 #if WITH_EDITOR
 
 const FSlateBrush* UScrollBar::GetEditorIcon()
diff --git a/Engine/Source/Runtime/UMG/Private/Components/Spacer.cpp b/Engine/Source/Runtime/UMG/Private/Components/Spacer.cpp
--- a/Engine/Source/Runtime/UMG/Private/Components/Spacer.cpp
+++ b/Engine/Source/Runtime/UMG/Private/Components/Spacer.cpp
@@ -14,6 +14,13 @@ USpacer::USpacer(const FPostConstructInitializeProperties& PCIP)
 	bIsVariable = false;
 }
 
+void USpacer::ReleaseNativeWidget()
+{
+	Super::ReleaseNativeWidget();
+
+	MySpacer.Reset();
+}
+
 void USpacer::SetSize(FVector2D InSize)
 {
 	Size = InSize;
diff --git a/Engine/Source/Runtime/UMG/Private/Components/SpinBox.cpp b/Engine/Source/Runtime/UMG/Private/Components/SpinBox.cpp
--- a/Engine/Source/Runtime/UMG/Private/Components/SpinBox.cpp
+++ b/Engine/Source/Runtime/UMG/Private/Components/SpinBox.cpp
@@ -26,6 +26,9 @@ USpinBox::USpinBox(const FPostConstructInitializeProperties& PCIP)
 	MinDesiredWidth = Defaults._MinDesiredWidth.Get();
 	ClearKeyboardFocusOnCommit = Defaults._ClearKeyboardFocusOnCommit.Get();
 	SelectAllTextOnCommit = Defaults._SelectAllTextOnCommit.Get();
+
+	WidgetStyle = *Defaults._Style;
+	ForegroundColor = FSlateColor(FLinearColor::Black);
 }
 
 void USpinBox::ReleaseNativeWidget()
@@ -44,16 +47,8 @@ TSharedRef<SWidget> USpinBox::RebuildWidget()
 		FontPath = FPaths::EngineContentDir() / Font.FontName.ToString();
 	}
 	
-	SSpinBox<float>::FArguments Defaults;
-	
-	const FSpinBoxStyle* StylePtr = (Style != nullptr) ? Style->GetStyle<FSpinBoxStyle>() : nullptr;
-	if ( StylePtr == nullptr )
-	{
-		StylePtr = Defaults._Style;
-	}
-	
 	MySpinBox = SNew(SSpinBox<float>)
-	.Style(StylePtr)
+	.Style(&WidgetStyle)
 	.Font(FSlateFontInfo(FontPath, Font.Size))
 	.ClearKeyboardFocusOnCommit(ClearKeyboardFocusOnCommit)
 	.SelectAllTextOnCommit(SelectAllTextOnCommit)
@@ -77,6 +72,8 @@ void USpinBox::SynchronizeProperties()
 	MySpinBox->SetSliderExponent(SliderExponent);
 	MySpinBox->SetMinDesiredWidth(MinDesiredWidth);
 
+	MySpinBox->SetForegroundColor(ForegroundColor);
+
 	// Set optional values
 	bOverride_MinValue ? SetMinValue(MinValue) : ClearMinValue();
 	bOverride_MaxValue ? SetMaxValue(MaxValue) : ClearMaxValue();
@@ -246,6 +243,15 @@ void USpinBox::ClearMaxSliderValue()
 	}
 }
 
+void USpinBox::SetForegroundColor(FSlateColor InForegroundColor)
+{
+	ForegroundColor = InForegroundColor;
+	if ( MySpinBox.IsValid() )
+	{
+		MySpinBox->SetForegroundColor(ForegroundColor);
+	}
+}
+
 // Event handlers
 void USpinBox::HandleOnValueChanged(float InValue)
 {
@@ -267,6 +273,26 @@ void USpinBox::HandleOnEndSliderMovement(float InValue)
 	OnEndSliderMovement.Broadcast(InValue);
 }
 
+void USpinBox::PostLoad()
+{
+	Super::PostLoad();
+
+	if ( GetLinkerUE4Version() < VER_UE4_DEPRECATE_UMG_STYLE_ASSETS )
+	{
+		if ( Style_DEPRECATED != nullptr )
+		{
+			const FSpinBoxStyle* StylePtr = Style_DEPRECATED->GetStyle<FSpinBoxStyle>();
+			if ( StylePtr != nullptr )
+			{
+				WidgetStyle = *StylePtr;
+			}
+
+			Style_DEPRECATED = nullptr;
+		}
+	}
+}
+
+
 #if WITH_EDITOR
 
 const FSlateBrush* USpinBox::GetEditorIcon()
diff --git a/Engine/Source/Runtime/UMG/Private/Components/TextBlock.cpp b/Engine/Source/Runtime/UMG/Private/Components/TextBlock.cpp
--- a/Engine/Source/Runtime/UMG/Private/Components/TextBlock.cpp
+++ b/Engine/Source/Runtime/UMG/Private/Components/TextBlock.cpp
@@ -12,17 +12,26 @@ UTextBlock::UTextBlock(const FPostConstructInitializeProperties& PCIP)
 {
 	bIsVariable = false;
 
+	STextBlock::FArguments Defaults;
+	WidgetStyle = *Defaults._TextStyle;
+
 	Text = LOCTEXT("TextBlockDefaultValue", "Text Block");
 	ShadowOffset = FVector2D(1.0f, 1.0f);
 	ColorAndOpacity = FLinearColor::White;
 	ShadowColorAndOpacity = FLinearColor::Transparent;
-
-	Style = NULL;
+	LineHeightPercentage = 1.0f;
 
 	// @TODO UMG HACK Special font initialization hack since there are no font assets yet for slate.
 	Font = FSlateFontInfo(TEXT("Slate/Fonts/Roboto-Bold.ttf"), 24);
 }
 
+void UTextBlock::ReleaseNativeWidget()
+{
+	Super::ReleaseNativeWidget();
+
+	MyTextBlock.Reset();
+}
+
 void UTextBlock::SetColorAndOpacity(FLinearColor InColorAndOpacity)
 {
 	ColorAndOpacity = InColorAndOpacity;
@@ -52,19 +61,8 @@ void UTextBlock::SetShadowOffset(FVector2D InShadowOffset)
 
 TSharedRef<SWidget> UTextBlock::RebuildWidget()
 {
-	MyTextBlock = SNew(STextBlock);
-
-#if WITH_EDITOR
-	MyEditorTextBlock = SNew(SInlineEditableTextBlock)
-		.OnTextCommitted(BIND_UOBJECT_DELEGATE(FOnTextCommitted, HandleTextCommitted))
-		//.IsSelected(InArgs._IsSelected)
-		;
-
-	//if ( IsDesignTime() )
-	//{
-	//	return MyEditorTextBlock.ToSharedRef();
-	//}
-#endif
+	MyTextBlock = SNew(STextBlock)
+		.TextStyle(&WidgetStyle);
 
 	return MyTextBlock.ToSharedRef();
 }
@@ -80,13 +78,6 @@ void UTextBlock::SynchronizeProperties()
 		FontPath = FPaths::EngineContentDir() / Font.FontName.ToString();
 	}
 
-	const FTextBlockStyle* StylePtr = ( Style != NULL ) ? Style->GetStyle<FTextBlockStyle>() : NULL;
-	if ( StylePtr == NULL )
-	{
-		STextBlock::FArguments Defaults;
-		StylePtr = Defaults._TextStyle;
-	}
-
 	TAttribute<FText> TextBinding = OPTIONAL_BINDING(FText, Text);
 	TAttribute<FSlateColor> ColorAndOpacityBinding = OPTIONAL_BINDING(FSlateColor, ColorAndOpacity);
 	TAttribute<FLinearColor> ShadowColorAndOpacityBinding = OPTIONAL_BINDING(FLinearColor, ShadowColorAndOpacity);
@@ -94,21 +85,14 @@ void UTextBlock::SynchronizeProperties()
 	MyTextBlock->SetText(TextBinding);
 	MyTextBlock->SetFont(FSlateFontInfo(FontPath, Font.Size));
 	MyTextBlock->SetColorAndOpacity(ColorAndOpacityBinding);
-	MyTextBlock->SetTextStyle(StylePtr);
 	MyTextBlock->SetShadowOffset(ShadowOffset);
 	MyTextBlock->SetShadowColorAndOpacity(ShadowColorAndOpacityBinding);
 	MyTextBlock->SetAutoWrapText(AutoWrapText);
 	MyTextBlock->SetWrapTextAt(WrapTextAt != 0 ? WrapTextAt : TAttribute<float>());
 	MyTextBlock->SetMinDesiredWidth(MinDesiredWidth);
-
-#if WITH_EDITOR
-	MyEditorTextBlock->SetText(TextBinding);
-	//MyEditorTextBlock->SetFont(FSlateFontInfo(FontPath, Font.Size));
-	//MyEditorTextBlock->SetColorAndOpacity(ColorAndOpacityBinding);
-	//MyEditorTextBlock->SetTextStyle(StylePtr);
-	//MyEditorTextBlock->SetShadowOffset(ShadowOffset);
-	//MyEditorTextBlock->SetShadowColorAndOpacity(ShadowColorAndOpacityBinding);
-#endif
+	MyTextBlock->SetLineHeightPercentage(LineHeightPercentage);
+	MyTextBlock->SetMargin(Margin);
+	MyTextBlock->SetJustification(Justification);
 }
 
 FText UTextBlock::GetText() const
@@ -130,6 +114,25 @@ void UTextBlock::SetText(FText InText)
 	}
 }
 
+void UTextBlock::PostLoad()
+{
+	Super::PostLoad();
+
+	if ( GetLinkerUE4Version() < VER_UE4_DEPRECATE_UMG_STYLE_ASSETS )
+	{
+		if ( Style_DEPRECATED != nullptr )
+		{
+			const FTextBlockStyle* StylePtr = Style_DEPRECATED->GetStyle<FTextBlockStyle>();
+			if ( StylePtr != nullptr )
+			{
+				WidgetStyle = *StylePtr;
+			}
+
+			Style_DEPRECATED = nullptr;
+		}
+	}
+}
+
 #if WITH_EDITOR
 
 FString UTextBlock::GetLabelMetadata() const
@@ -148,11 +151,6 @@ void UTextBlock::HandleTextCommitted(const FText& InText, ETextCommit::Type Comm
 	//     Need a way to recognize one particular widget and forward things to them!
 }
 
-void UTextBlock::OnBeginEdit()
-{
-	MyEditorTextBlock->EnterEditingMode();
-}
-
 const FSlateBrush* UTextBlock::GetEditorIcon()
 {
 	return FUMGStyle::Get().GetBrush("Widget.TextBlock");
diff --git a/Engine/Source/Runtime/UMG/Public/Components/EditableTextBox.h b/Engine/Source/Runtime/UMG/Public/Components/EditableTextBox.h
--- a/Engine/Source/Runtime/UMG/Public/Components/EditableTextBox.h
+++ b/Engine/Source/Runtime/UMG/Public/Components/EditableTextBox.h
@@ -17,9 +17,13 @@ public:
 
 public:
 
+	/** The style */
+	UPROPERTY(EditAnywhere, BlueprintReadWrite, Category="Style", meta=( DisplayName="Style" ))
+	FEditableTextBoxStyle WidgetStyle;
+
 	/** Style used for the text box */
-	UPROPERTY(EditDefaultsOnly, Category="Style", meta=( DisplayThumbnail = "true" ))
-	USlateWidgetStyleAsset* Style;
+	UPROPERTY()
+	USlateWidgetStyleAsset* Style_DEPRECATED;
 
 	/** The text content for this editable text box widget */
 	UPROPERTY(EditDefaultsOnly, Category=Content)
@@ -120,7 +124,13 @@ public:
 	virtual void SynchronizeProperties() override;
 	// End of UWidget interface
 
+	// UVisual interface
 	virtual void ReleaseNativeWidget() override;
+	// End of UVisual interface
+
+	// Begin UObject interface
+	virtual void PostLoad() override;
+	// End of UObject interface
 
 #if WITH_EDITOR
 	virtual const FSlateBrush* GetEditorIcon() override;
diff --git a/Engine/Source/Runtime/UMG/Public/Components/ScrollBar.h b/Engine/Source/Runtime/UMG/Public/Components/ScrollBar.h
--- a/Engine/Source/Runtime/UMG/Public/Components/ScrollBar.h
+++ b/Engine/Source/Runtime/UMG/Public/Components/ScrollBar.h
@@ -13,8 +13,11 @@ class UMG_API UScrollBar : public UPanelWidget
 public:
 
 	/** Style of the scrollbar */
-	UPROPERTY(EditDefaultsOnly, Category="Style", meta=( DisplayThumbnail = "true" ))
-	USlateWidgetStyleAsset* Style;
+	UPROPERTY(EditAnywhere, BlueprintReadWrite, Category="Style", meta=( DisplayName="Style" ))
+	FScrollBarStyle WidgetStyle;
+
+	UPROPERTY()
+	USlateWidgetStyleAsset* Style_DEPRECATED;
 
 	/**  */
 	UPROPERTY(EditDefaultsOnly, Category="Behavior")
@@ -56,11 +59,18 @@ public:
 	///** @return True if the user is scrolling by dragging the scroll bar thumb. */
 	//bool IsScrolling() const;
 
-
 	// UWidget interface
 	virtual void SynchronizeProperties() override;
 	// End of UWidget interface
 
+	// UVisual interface
+	virtual void ReleaseNativeWidget() override;
+	// End of UVisual interface
+
+	// Begin UObject interface
+	virtual void PostLoad() override;
+	// End of UObject interface
+
 #if WITH_EDITOR
 	// UWidget interface
 	virtual const FSlateBrush* GetEditorIcon() override;
diff --git a/Engine/Source/Runtime/UMG/Public/Components/Spacer.h b/Engine/Source/Runtime/UMG/Public/Components/Spacer.h
--- a/Engine/Source/Runtime/UMG/Public/Components/Spacer.h
+++ b/Engine/Source/Runtime/UMG/Public/Components/Spacer.h
@@ -26,6 +26,10 @@ public:
 	virtual void SynchronizeProperties() override;
 	// End of UWidget interface
 
+	// UVisual interface
+	virtual void ReleaseNativeWidget() override;
+	// End of UVisual interface
+
 #if WITH_EDITOR
 	virtual const FSlateBrush* GetEditorIcon() override;
 	virtual const FText GetPaletteCategory() override;
diff --git a/Engine/Source/Runtime/UMG/Public/Components/SpinBox.h b/Engine/Source/Runtime/UMG/Public/Components/SpinBox.h
--- a/Engine/Source/Runtime/UMG/Public/Components/SpinBox.h
+++ b/Engine/Source/Runtime/UMG/Public/Components/SpinBox.h
@@ -16,9 +16,12 @@ public:
 	DECLARE_DYNAMIC_MULTICAST_DELEGATE(FOnSpinBoxBeginSliderMovement);
 
 public:
-	/** Text style */
-	UPROPERTY(EditDefaultsOnly, Category = Style, meta = (DisplayThumbnail = "true"))
-	USlateWidgetStyleAsset* Style;
+	/** The Style */
+	UPROPERTY(EditAnywhere, BlueprintReadWrite, Category="Style", meta=( DisplayName="Style" ))
+	FSpinBoxStyle WidgetStyle;
+
+	UPROPERTY()
+	USlateWidgetStyleAsset* Style_DEPRECATED;
 
 	/** Value stored in this spin box */
 	UPROPERTY(EditDefaultsOnly, Category = Content)
@@ -48,6 +51,9 @@ public:
 	UPROPERTY(EditDefaultsOnly, Category = Content, AdvancedDisplay)
 	bool SelectAllTextOnCommit;
 
+	UPROPERTY(EditDefaultsOnly, Category="Style")
+	FSlateColor ForegroundColor;
+
 public:
 	/** Called when the value is changed interactively by the user */
 	UPROPERTY(BlueprintAssignable, Category = "Widget Event")
@@ -125,12 +131,21 @@ public:
 	UFUNCTION(BlueprintCallable, Category = "Behavior")
 	void ClearMaxSliderValue();
 
+	/**  */
+	UFUNCTION(BlueprintCallable, Category = "Appearance")
+	void SetForegroundColor(FSlateColor InForegroundColor);
+
+	// UWidget interface
+	virtual void SynchronizeProperties() override;
+	// End of UWidget interface
+
 	// UVisual interface
 	virtual void ReleaseNativeWidget() override;
 	// End of UVisual interface
 
-	// UWidget interface
-	virtual void SynchronizeProperties() override;
+	// Begin UObject interface
+	virtual void PostLoad() override;
+	// End of UObject interface
 
 #if WITH_EDITOR
 	virtual const FSlateBrush* GetEditorIcon() override;
diff --git a/Engine/Source/Runtime/UMG/Public/Components/TextBlock.h b/Engine/Source/Runtime/UMG/Public/Components/TextBlock.h
--- a/Engine/Source/Runtime/UMG/Public/Components/TextBlock.h
+++ b/Engine/Source/Runtime/UMG/Public/Components/TextBlock.h
@@ -37,9 +37,12 @@ public:
 	void SetShadowOffset(FVector2D InShadowOffset);
 
 public:
-	/** The style to use to render the text */
-	UPROPERTY(EditDefaultsOnly, Category=Style, meta=( DisplayThumbnail = "true" ))
-	USlateWidgetStyleAsset* Style;
+	/** The text block style */
+	UPROPERTY(EditAnywhere, BlueprintReadWrite, Category="Style", meta=( DisplayName="Style" ))
+	FTextBlockStyle WidgetStyle;
+
+	UPROPERTY()
+	USlateWidgetStyleAsset* Style_DEPRECATED;
 
 	/** The text to display */
 	UPROPERTY(EditDefaultsOnly, Category=Content)
@@ -66,25 +69,37 @@ public:
 	FVector2D ShadowOffset;
 
 	/** The color of the shadow */
-	UPROPERTY(EditDefaultsOnly, Category=Appearance)
+	UPROPERTY(EditDefaultsOnly, Category=Appearance, meta=( DisplayName="Shadow Color" ))
 	FLinearColor ShadowColorAndOpacity;
 
 	/** A bindable delegate for the ShadowColorAndOpacity. */
 	UPROPERTY()
 	FGetLinearColor ShadowColorAndOpacityDelegate;
 
-	/** Whether text wraps onto a new line when it's length exceeds this width; if this value is zero or negative, no wrapping occurs. */
-	UPROPERTY(EditDefaultsOnly, Category=Appearance, AdvancedDisplay)
-	float WrapTextAt;
+	/** How the text should be aligned with the margin. */
+	UPROPERTY(EditDefaultsOnly, Category=Appearance)
+	TEnumAsByte<ETextJustify::Type> Justification;
 
 	/** True if we're wrapping text automatically based on the computed horizontal space for this widget */
-	UPROPERTY(EditDefaultsOnly, Category=Appearance, AdvancedDisplay)
+	UPROPERTY(EditDefaultsOnly, Category=Appearance)
 	bool AutoWrapText;
 
+	/** Whether text wraps onto a new line when it's length exceeds this width; if this value is zero or negative, no wrapping occurs. */
+	UPROPERTY(EditDefaultsOnly, Category=Appearance, AdvancedDisplay)
+	float WrapTextAt;
+
 	/** The minimum desired size for the text */
 	UPROPERTY(EditDefaultsOnly, Category=Appearance, AdvancedDisplay)
 	float MinDesiredWidth;
 
+	/** The amount of blank space left around the edges of text area. */
+	UPROPERTY(EditDefaultsOnly, Category=Appearance, AdvancedDisplay)
+	FMargin Margin;
+
+	/** The amount to scale each lines height by. */
+	UPROPERTY(EditDefaultsOnly, Category=Appearance, AdvancedDisplay)
+	float LineHeightPercentage;
+
 	///** Called when this text is double clicked */
 	//SLATE_EVENT(FOnClicked, OnDoubleClicked)
 
@@ -104,9 +119,17 @@ public:
 	void SetText(FText InText);
 
 	// UWidget interface
-	void SynchronizeProperties() override;
+	virtual void SynchronizeProperties() override;
 	// End of UWidget interface
 
+	// UVisual interface
+	virtual void ReleaseNativeWidget() override;
+	// End of UVisual interface
+
+	// Begin UObject interface
+	virtual void PostLoad() override;
+	// End of UObject interface
+
 #if WITH_EDITOR
 	// UWidget interface
 	virtual const FSlateBrush* GetEditorIcon() override;
@@ -116,8 +139,6 @@ public:
 	virtual FString GetLabelMetadata() const override;
 
 	void HandleTextCommitted(const FText& InText, ETextCommit::Type CommitteType);
-
-	virtual void OnBeginEdit() override;
 #endif
 
 protected:
@@ -128,8 +149,4 @@ protected:
 protected:
 
 	TSharedPtr<STextBlock> MyTextBlock;
-
-#if WITH_EDITOR
-	TSharedPtr<SInlineEditableTextBlock> MyEditorTextBlock;
-#endif
 };
