UMG - Upgrading Multiline textboxes to use the new style system.

diff --git a/Engine/Source/Runtime/UMG/Private/Components/MultiLineEditableText.cpp b/Engine/Source/Runtime/UMG/Private/Components/MultiLineEditableText.cpp
--- a/Engine/Source/Runtime/UMG/Private/Components/MultiLineEditableText.cpp
+++ b/Engine/Source/Runtime/UMG/Private/Components/MultiLineEditableText.cpp
@@ -10,10 +10,20 @@
 UMultiLineEditableText::UMultiLineEditableText(const FPostConstructInitializeProperties& PCIP)
 	: Super(PCIP)
 {
+	SMultiLineEditableText::FArguments Defaults;
+	WidgetStyle = *Defaults._TextStyle;
+
 	// HACK Special font initialization hack since there are no font assets yet for slate.
 	Font = FSlateFontInfo(TEXT("Slate/Fonts/Roboto-Bold.ttf"), 12);
 }
 
+void UMultiLineEditableText::ReleaseNativeWidget()
+{
+	Super::ReleaseNativeWidget();
+
+	MyMultiLineEditableText.Reset();
+}
+
 TSharedRef<SWidget> UMultiLineEditableText::RebuildWidget()
 {
 	FString FontPath = FPaths::GameContentDir() / Font.FontName.ToString();
@@ -23,16 +33,8 @@ TSharedRef<SWidget> UMultiLineEditableText::RebuildWidget()
 		FontPath = FPaths::EngineContentDir() / Font.FontName.ToString();
 	}
 	
-	SMultiLineEditableText::FArguments Defaults;
-	
-	//const FMultiLineEditableTextStyle* StylePtr = ( Style != NULL ) ? Style->GetStyle<FMultiLineEditableTextStyle>() : NULL;
-	//if ( StylePtr == NULL )
-	//{
-	//	StylePtr = Defaults._Style;
-	//}
-	
 	MyMultiLineEditableText = SNew(SMultiLineEditableText)
-	//.Style(StylePtr)
+	.TextStyle(&WidgetStyle)
 	.Font(FSlateFontInfo(FontPath, Font.Size))
 	.Justification(Justification)
 //	.MinDesiredWidth(MinimumDesiredWidth)
diff --git a/Engine/Source/Runtime/UMG/Private/Components/MultiLineEditableTextBox.cpp b/Engine/Source/Runtime/UMG/Private/Components/MultiLineEditableTextBox.cpp
--- a/Engine/Source/Runtime/UMG/Private/Components/MultiLineEditableTextBox.cpp
+++ b/Engine/Source/Runtime/UMG/Private/Components/MultiLineEditableTextBox.cpp
@@ -16,6 +16,17 @@ UMultiLineEditableTextBox::UMultiLineEditableTextBox(const FPostConstructInitial
 
 	// HACK Special font initialization hack since there are no font assets yet for slate.
 	Font = FSlateFontInfo(TEXT("Slate/Fonts/Roboto-Bold.ttf"), 12);
+
+	SMultiLineEditableTextBox::FArguments Defaults;
+	WidgetStyle = *Defaults._Style;
+	TextStyle = *Defaults._TextStyle;
+}
+
+void UMultiLineEditableTextBox::ReleaseNativeWidget()
+{
+	Super::ReleaseNativeWidget();
+
+	MyEditableTextBlock.Reset();
 }
 
 TSharedRef<SWidget> UMultiLineEditableTextBox::RebuildWidget()
@@ -26,8 +37,10 @@ TSharedRef<SWidget> UMultiLineEditableTextBox::RebuildWidget()
 	{
 		FontPath = FPaths::EngineContentDir() / Font.FontName.ToString();
 	}
-
+	
 	MyEditableTextBlock = SNew(SMultiLineEditableTextBox)
+		.Style(&WidgetStyle)
+		.TextStyle(&TextStyle)
 		.Font(FSlateFontInfo(FontPath, Font.Size))
 		.Justification(Justification)
 		.ForegroundColor(ForegroundColor)
@@ -51,12 +64,6 @@ void UMultiLineEditableTextBox::SynchronizeProperties()
 {
 	Super::SynchronizeProperties();
 
-	//const FMultiLineEditableTextBoxStyle* StylePtr = ( Style != NULL ) ? Style->GetStyle<FMultiLineEditableTextBoxStyle>() : NULL;
-	//if ( StylePtr )
-	//{
-	//	MyEditableTextBlock->SetStyle(StylePtr);
-	//}
-
 	MyEditableTextBlock->SetText(Text);
 //	MyEditableTextBlock->SetHintText(HintText);
 //	MyEditableTextBlock->SetIsReadOnly(IsReadOnly);
@@ -103,6 +110,25 @@ void UMultiLineEditableTextBox::HandleOnTextCommitted(const FText& Text, ETextCo
 	OnTextCommitted.Broadcast(Text, CommitMethod);
 }
 
+void UMultiLineEditableTextBox::PostLoad()
+{
+	Super::PostLoad();
+
+	if ( GetLinkerUE4Version() < VER_UE4_DEPRECATE_UMG_STYLE_ASSETS )
+	{
+		if ( Style_DEPRECATED != nullptr )
+		{
+			const FEditableTextBoxStyle* StylePtr = Style_DEPRECATED->GetStyle<FEditableTextBoxStyle>();
+			if ( StylePtr != nullptr )
+			{
+				WidgetStyle = *StylePtr;
+			}
+
+			Style_DEPRECATED = nullptr;
+		}
+	}
+}
+
 #if WITH_EDITOR
 
 const FSlateBrush* UMultiLineEditableTextBox::GetEditorIcon()
diff --git a/Engine/Source/Runtime/UMG/Public/Components/MultiLineEditableText.h b/Engine/Source/Runtime/UMG/Public/Components/MultiLineEditableText.h
--- a/Engine/Source/Runtime/UMG/Public/Components/MultiLineEditableText.h
+++ b/Engine/Source/Runtime/UMG/Public/Components/MultiLineEditableText.h
@@ -17,6 +17,10 @@ public:
 
 public:
 
+	/** The style */
+	UPROPERTY(EditAnywhere, BlueprintReadWrite, Category="Style", meta=( DisplayName="Style" ))
+	FTextBlockStyle WidgetStyle;
+
 	/** The text content for this editable text box widget */
 	UPROPERTY(EditDefaultsOnly, Category=Content)
 	FText Text;
@@ -49,6 +53,10 @@ public:
 	virtual void SynchronizeProperties() override;
 	// End of UWidget interface
 
+	// UVisual interface
+	virtual void ReleaseNativeWidget() override;
+	// End of UVisual interface
+
 #if WITH_EDITOR
 	virtual const FSlateBrush* GetEditorIcon() override;
 	virtual const FText GetPaletteCategory() override;
diff --git a/Engine/Source/Runtime/UMG/Public/Components/MultiLineEditableTextBox.h b/Engine/Source/Runtime/UMG/Public/Components/MultiLineEditableTextBox.h
--- a/Engine/Source/Runtime/UMG/Public/Components/MultiLineEditableTextBox.h
+++ b/Engine/Source/Runtime/UMG/Public/Components/MultiLineEditableTextBox.h
@@ -17,9 +17,16 @@ public:
 
 public:
 
-	/** Style used for the text box */
-	UPROPERTY(EditDefaultsOnly, Category="Style", meta=( DisplayThumbnail = "true" ))
-	USlateWidgetStyleAsset* Style;
+	/** The style */
+	UPROPERTY(EditAnywhere, BlueprintReadWrite, Category="Style", meta=( DisplayName="Style" ))
+	FEditableTextBoxStyle WidgetStyle;
+
+	/** The text style */
+	UPROPERTY(EditAnywhere, BlueprintReadWrite, Category="Style", meta=( DisplayName="Text Style" ))
+	FTextBlockStyle TextStyle;
+
+	UPROPERTY()
+	USlateWidgetStyleAsset* Style_DEPRECATED;
 
 	/** The text content for this editable text box widget */
 	UPROPERTY(EditDefaultsOnly, Category=Content)
@@ -71,6 +78,14 @@ public:
 	virtual void SynchronizeProperties() override;
 	// End of UWidget interface
 
+	// UVisual interface
+	virtual void ReleaseNativeWidget() override;
+	// End of UVisual interface
+
+	// Begin UObject interface
+	virtual void PostLoad() override;
+	// End of UObject interface
+
 #if WITH_EDITOR
 	virtual const FSlateBrush* GetEditorIcon() override;
 	virtual const FText GetPaletteCategory() override;
