Added the ability to get the currently selected item from the ComboBoxString widget.

diff --git a/Engine/Source/Runtime/UMG/Private/Components/ComboBoxString.cpp b/Engine/Source/Runtime/UMG/Private/Components/ComboBoxString.cpp
--- a/Engine/Source/Runtime/UMG/Private/Components/ComboBoxString.cpp
+++ b/Engine/Source/Runtime/UMG/Private/Components/ComboBoxString.cpp
@@ -33,18 +33,17 @@ TSharedRef<SWidget> UComboBoxString::RebuildWidget()
 		AddOption(DefaultOptions);
 	}
 
-	TSharedPtr<FString> SelectedOptionPtr;
 	int32 InitialIndex = FindOptionIndex(SelectedOption);
 	if ( InitialIndex != -1 )
 	{
-		SelectedOptionPtr = Options[InitialIndex];
+		CurrentOptionPtr = Options[InitialIndex];
 	}
 
 	MyComboBox =
 		SNew(SComboBox< TSharedPtr<FString> >)
 		.ComboBoxStyle(&WidgetStyle)
 		.OptionsSource(&Options)
-		.InitiallySelectedItem(SelectedOptionPtr)
+		.InitiallySelectedItem(CurrentOptionPtr)
 		.ContentPadding(ContentPadding)
 		.MaxListHeight(MaxListHeight)
 		.HasDownArrow(HasDownArrow)
@@ -59,7 +58,7 @@ TSharedRef<SWidget> UComboBoxString::RebuildWidget()
 	if ( InitialIndex != -1 )
 	{
 		// Generate the widget for the initially selected widget if needed
-		ComoboBoxContent->SetContent(HandleGenerateWidget(SelectedOptionPtr));
+		ComoboBoxContent->SetContent(HandleGenerateWidget(CurrentOptionPtr));
 	}
 
 	return MyComboBox.ToSharedRef();
@@ -73,6 +72,10 @@ void UComboBoxString::AddOption(const FString& Option)
 bool UComboBoxString::RemoveOption(const FString& Option)
 {
 	int32 OptionIndex = FindOptionIndex(Option);
+	if (Options[OptionIndex] == CurrentOptionPtr)
+	{
+		CurrentOptionPtr.Reset();
+	}
 	if ( OptionIndex != -1 )
 	{
 		Options.RemoveAt(OptionIndex);
@@ -82,11 +85,11 @@ bool UComboBoxString::RemoveOption(const FString& Option)
 	return false;
 }
 
-int32 UComboBoxString::FindOptionIndex(const FString& Option)
+int32 UComboBoxString::FindOptionIndex(const FString& Option) const
 {
 	for ( int32 OptionIndex = 0; OptionIndex < Options.Num(); OptionIndex++ )
 	{
-		TSharedPtr<FString>& OptionAtIndex = Options[OptionIndex];
+		const TSharedPtr<FString>& OptionAtIndex = Options[OptionIndex];
 
 		if ( ( *OptionAtIndex ) == Option )
 		{
@@ -97,11 +100,39 @@ int32 UComboBoxString::FindOptionIndex(const FString& Option)
 	return -1;
 }
 
+FString UComboBoxString::GetOptionAtIndex(int32 Index) const
+{
+	if (Index >= 0 && Index < Options.Num())
+	{
+		return *(Options[Index]);
+	}
+	return FString();
+}
+
 void UComboBoxString::ClearOptions()
 {
 	Options.Empty();
 }
 
+void UComboBoxString::SetSelectedOption(FString Option)
+{
+	int32 InitialIndex = FindOptionIndex(Option);
+	if (InitialIndex != -1)
+	{
+		CurrentOptionPtr = Options[InitialIndex];
+		ComoboBoxContent->SetContent(HandleGenerateWidget(CurrentOptionPtr));
+	}
+}
+
+FString UComboBoxString::GetSelectedOption() const
+{
+	if (CurrentOptionPtr.IsValid())
+	{
+		return *CurrentOptionPtr;
+	}
+	return FString();
+}
+
 TSharedRef<SWidget> UComboBoxString::HandleGenerateWidget(TSharedPtr<FString> Item) const
 {
 	// Call the user's delegate to see if they want to generate a custom widget bound to the data source.
@@ -120,6 +151,8 @@ TSharedRef<SWidget> UComboBoxString::HandleGenerateWidget(TSharedPtr<FString> It
 
 void UComboBoxString::HandleSelectionChanged(TSharedPtr<FString> Item, ESelectInfo::Type SelectionType)
 {
+	CurrentOptionPtr = Item;
+
 	if ( !IsDesignTime() )
 	{
 		OnSelectionChanged.Broadcast(*Item, SelectionType);
diff --git a/Engine/Source/Runtime/UMG/Public/Components/ComboBoxString.h b/Engine/Source/Runtime/UMG/Public/Components/ComboBoxString.h
--- a/Engine/Source/Runtime/UMG/Public/Components/ComboBoxString.h
+++ b/Engine/Source/Runtime/UMG/Public/Components/ComboBoxString.h
@@ -66,11 +66,20 @@ public:
 	bool RemoveOption(const FString& Option);
 
 	UFUNCTION(BlueprintCallable, Category="ComboBox")
-	int32 FindOptionIndex(const FString& Option);
+	int32 FindOptionIndex(const FString& Option) const;
+
+	UFUNCTION(BlueprintCallable, Category = "ComboBox")
+	FString GetOptionAtIndex(int32 Index) const;
 
 	UFUNCTION(BlueprintCallable, Category="ComboBox")
 	void ClearOptions();
 
+	UFUNCTION(BlueprintCallable, Category = "ComboBox")
+	void SetSelectedOption(FString Option);
+
+	UFUNCTION(BlueprintCallable, Category = "ComboBox")
+	FString GetSelectedOption() const;
+
 	virtual void ReleaseNativeWidget() override;
 
 #if WITH_EDITOR
@@ -101,4 +110,7 @@ protected:
 
 	/** A shared pointer to a container that holds the comobobox content that is selected */
 	TSharedPtr< SBox > ComoboBoxContent;
+
+	/** A shared pointer to the current selected string */
+	TSharedPtr<FString> CurrentOptionPtr;
 };
