UMG - Fixing a bug in the spin box that was preventing the value from being properly clamped using the users settings.

diff --git a/Engine/Source/Runtime/UMG/Private/Components/SpinBox.cpp b/Engine/Source/Runtime/UMG/Private/Components/SpinBox.cpp
--- a/Engine/Source/Runtime/UMG/Private/Components/SpinBox.cpp
+++ b/Engine/Source/Runtime/UMG/Private/Components/SpinBox.cpp
@@ -65,9 +65,6 @@ void USpinBox::SynchronizeProperties()
 {
 	Super::SynchronizeProperties();
 
-	TAttribute<float> ValueBinding = OPTIONAL_BINDING(float, Value);
-	MySpinBox->SetValue(ValueBinding);
-	
 	MySpinBox->SetDelta(Delta);
 	MySpinBox->SetSliderExponent(SliderExponent);
 	MySpinBox->SetMinDesiredWidth(MinDesiredWidth);
@@ -79,6 +76,10 @@ void USpinBox::SynchronizeProperties()
 	bOverride_MaxValue ? SetMaxValue(MaxValue) : ClearMaxValue();
 	bOverride_MinSliderValue ? SetMinSliderValue(MinSliderValue) : ClearMinSliderValue();
 	bOverride_MaxSliderValue ? SetMaxSliderValue(MaxSliderValue) : ClearMaxSliderValue();
+
+	// Always set the value last so that the max/min values are taken into account.
+	TAttribute<float> ValueBinding = OPTIONAL_BINDING(float, Value);
+	MySpinBox->SetValue(ValueBinding);
 }
 
 float USpinBox::GetValue() const
@@ -255,22 +256,34 @@ void USpinBox::SetForegroundColor(FSlateColor InForegroundColor)
 // Event handlers
 void USpinBox::HandleOnValueChanged(float InValue)
 {
-	OnValueChanged.Broadcast(InValue);
+	if ( !IsDesignTime() )
+	{
+		OnValueChanged.Broadcast(InValue);
+	}
 }
 
 void USpinBox::HandleOnValueCommitted(float InValue, ETextCommit::Type CommitMethod)
 {
-	OnValueCommitted.Broadcast(InValue, CommitMethod);
+	if ( !IsDesignTime() )
+	{
+		OnValueCommitted.Broadcast(InValue, CommitMethod);
+	}
 }
 
 void USpinBox::HandleOnBeginSliderMovement()
 {
-	OnBeginSliderMovement.Broadcast();
+	if ( !IsDesignTime() )
+	{
+		OnBeginSliderMovement.Broadcast();
+	}
 }
 
 void USpinBox::HandleOnEndSliderMovement(float InValue)
 {
-	OnEndSliderMovement.Broadcast(InValue);
+	if ( !IsDesignTime() )
+	{
+		OnEndSliderMovement.Broadcast(InValue);
+	}
 }
 
 void USpinBox::PostLoad()
