- Updated FText property editor customization to default to multiline
- Use "meta=(SingleLine=true)" in the UPROPERTY declaration to restrict the field to single line entry

diff --git a/Engine/Source/Editor/IntroTutorials/Public/EditorTutorial.h b/Engine/Source/Editor/IntroTutorials/Public/EditorTutorial.h
--- a/Engine/Source/Editor/IntroTutorials/Public/EditorTutorial.h
+++ b/Engine/Source/Editor/IntroTutorials/Public/EditorTutorial.h
@@ -56,7 +56,7 @@ struct INTROTUTORIALS_API FTutorialCategory
 	FText Title;
 
 	/** Localized text to use to describe this category */
-	UPROPERTY(EditAnywhere, Category="Content", meta=(MultiLine=true))
+	UPROPERTY(EditAnywhere, Category="Content")
 	FText Description;
 
 	/** Icon for this tutorial, used when presented to the user in the tutorial browser. Only used if there isn't a valid texture to use. */
@@ -92,7 +92,7 @@ struct INTROTUTORIALS_API FTutorialContent
 	FString ExcerptName;
 
 	/** Localized text to use with this content */
-	UPROPERTY(EditAnywhere, Category="Content", meta=(MultiLine=true))
+	UPROPERTY(EditAnywhere, Category="Content")
 	FText Text;
 };
 
diff --git a/Engine/Source/Editor/PropertyEditor/Private/UserInterface/PropertyEditor/SPropertyEditorText.cpp b/Engine/Source/Editor/PropertyEditor/Private/UserInterface/PropertyEditor/SPropertyEditorText.cpp
--- a/Engine/Source/Editor/PropertyEditor/Private/UserInterface/PropertyEditor/SPropertyEditorText.cpp
+++ b/Engine/Source/Editor/PropertyEditor/Private/UserInterface/PropertyEditor/SPropertyEditorText.cpp
@@ -13,7 +13,7 @@ void SPropertyEditorText::Construct( const FArguments& InArgs, const TSharedRef<
 	PropertyEditor = InPropertyEditor;
 
 	const UProperty* Property = InPropertyEditor->GetProperty();
-	bIsMultiLine = Property->GetBoolMetaData("MultiLine");
+	bIsMultiLine = !Property->GetBoolMetaData("SingleLine");
 	if(bIsMultiLine)
 	{
 		ChildSlot
diff --git a/Engine/Source/Runtime/Engine/Classes/Materials/MaterialExpressionCustom.h b/Engine/Source/Runtime/Engine/Classes/Materials/MaterialExpressionCustom.h
--- a/Engine/Source/Runtime/Engine/Classes/Materials/MaterialExpressionCustom.h
+++ b/Engine/Source/Runtime/Engine/Classes/Materials/MaterialExpressionCustom.h
@@ -33,7 +33,7 @@ class UMaterialExpressionCustom : public UMaterialExpression
 {
 	GENERATED_UCLASS_BODY()
 
-	UPROPERTY(EditAnywhere, Category=MaterialExpressionCustom, meta=(MultiLine=true))
+	UPROPERTY(EditAnywhere, Category=MaterialExpressionCustom)
 	FString Code;
 
 	UPROPERTY(EditAnywhere, Category=MaterialExpressionCustom)
diff --git a/Engine/Source/Runtime/UMG/Public/Components/EditableText.h b/Engine/Source/Runtime/UMG/Public/Components/EditableText.h
--- a/Engine/Source/Runtime/UMG/Public/Components/EditableText.h
+++ b/Engine/Source/Runtime/UMG/Public/Components/EditableText.h
@@ -20,7 +20,7 @@ public:
 public:
 
 	/** The text content for this editable text box widget */
-	UPROPERTY(EditDefaultsOnly, Category=Content)
+	UPROPERTY(EditDefaultsOnly, Category = Content, meta=(SingleLine=true))
 	FText Text;
 
 	/** A bindable delegate to allow logic to drive the text of the widget */
@@ -28,7 +28,7 @@ public:
 	FGetText TextDelegate;
 
 	/** Hint text that appears when there is no text in the text box */
-	UPROPERTY(EditDefaultsOnly, Category=Content)
+	UPROPERTY(EditDefaultsOnly, Category = Content, meta=(SingleLine=true))
 	FText HintText;
 
 	/** A bindable delegate to allow logic to drive the hint text of the widget */
diff --git a/Engine/Source/Runtime/UMG/Public/Components/EditableTextBox.h b/Engine/Source/Runtime/UMG/Public/Components/EditableTextBox.h
--- a/Engine/Source/Runtime/UMG/Public/Components/EditableTextBox.h
+++ b/Engine/Source/Runtime/UMG/Public/Components/EditableTextBox.h
@@ -26,7 +26,7 @@ public:
 	USlateWidgetStyleAsset* Style_DEPRECATED;
 
 	/** The text content for this editable text box widget */
-	UPROPERTY(EditDefaultsOnly, Category=Content)
+	UPROPERTY(EditDefaultsOnly, Category = Content, meta=(SingleLine=true))
 	FText Text;
 
 	/** A bindable delegate to allow logic to drive the text of the widget */
@@ -34,7 +34,7 @@ public:
 	FGetText TextDelegate;
 
 	/** Hint text that appears when there is no text in the text box */
-	UPROPERTY(EditDefaultsOnly, Category=Content)
+	UPROPERTY(EditDefaultsOnly, Category = Content, meta=(SingleLine=true))
 	FText HintText;
 
 	/** A bindable delegate to allow logic to drive the hint text of the widget */
