UMG - Scrollbars are experimental.

diff --git a/Engine/Source/Runtime/UMG/Public/Components/ScrollBar.h b/Engine/Source/Runtime/UMG/Public/Components/ScrollBar.h
--- a/Engine/Source/Runtime/UMG/Public/Components/ScrollBar.h
+++ b/Engine/Source/Runtime/UMG/Public/Components/ScrollBar.h
@@ -5,7 +5,7 @@
 #include "ScrollBar.generated.h"
 
 /** */
-UCLASS(ClassGroup=UserInterface, meta=( IsAdvanced = "True" ))
+UCLASS(Experimental, ClassGroup=UserInterface, meta=( IsAdvanced = "True" ))
 class UMG_API UScrollBar : public UPanelWidget
 {
 	GENERATED_UCLASS_BODY()
