UMG - Showing a dashed border outline around the scale box.

diff --git a/Engine/Source/Runtime/UMG/Private/Components/ScaleBox.cpp b/Engine/Source/Runtime/UMG/Private/Components/ScaleBox.cpp
--- a/Engine/Source/Runtime/UMG/Private/Components/ScaleBox.cpp
+++ b/Engine/Source/Runtime/UMG/Private/Components/ScaleBox.cpp
@@ -32,7 +32,7 @@ TSharedRef<SWidget> UScaleBox::RebuildWidget()
 		Cast<UScaleBoxSlot>(GetContentSlot())->BuildSlot(MyScaleBox.ToSharedRef());
 	}
 
-	return MyScaleBox.ToSharedRef();
+	return BuildDesignTimeWidget( MyScaleBox.ToSharedRef() );
 }
 
 void UScaleBox::SynchronizeProperties()
