#nonttp - Adding code for missing UMG icons, needs assets tho

diff --git a/Engine/Source/Runtime/UMG/Private/UMGStyle.cpp b/Engine/Source/Runtime/UMG/Private/UMGStyle.cpp
--- a/Engine/Source/Runtime/UMG/Private/UMGStyle.cpp
+++ b/Engine/Source/Runtime/UMG/Private/UMGStyle.cpp
@@ -75,6 +75,9 @@ TSharedRef< FSlateStyleSet > FUMGStyle::Create()
 	Style->Set("Widget.CircularThrobber", new IMAGE_BRUSH(TEXT("icon_umg_ThrobberA_16x"), Icon16x16));
 	Style->Set("Widget.NativeWidgetHost", new IMAGE_BRUSH(TEXT("NativeWidgetHost"), Icon16x16));
 	Style->Set("Widget.ScaleBox", new IMAGE_BRUSH(TEXT("ScaleBox"), Icon16x16));
+	Style->Set("Widget.SizeBox", new IMAGE_BRUSH(TEXT("SizeBox"), Icon16x16));
+	Style->Set("Widget.SpinBox", new IMAGE_BRUSH(TEXT("SpinBox"), Icon16x16));
+	Style->Set("Widget.Grid", new IMAGE_BRUSH(TEXT("Grid"), Icon16x16));
 
 	Style->Set("Widget.UserWidget", new IMAGE_BRUSH(TEXT("UserWidget"), Icon16x16));
 
