Changing TextBlock's Color to an FSlateColor.

diff --git a/Engine/Source/Runtime/UMG/Private/Components/TextBlock.cpp b/Engine/Source/Runtime/UMG/Private/Components/TextBlock.cpp
--- a/Engine/Source/Runtime/UMG/Private/Components/TextBlock.cpp
+++ b/Engine/Source/Runtime/UMG/Private/Components/TextBlock.cpp
@@ -32,7 +32,7 @@ void UTextBlock::ReleaseNativeWidget()
 	MyTextBlock.Reset();
 }
 
-void UTextBlock::SetColorAndOpacity(FLinearColor InColorAndOpacity)
+void UTextBlock::SetColorAndOpacity(FSlateColor InColorAndOpacity)
 {
 	ColorAndOpacity = InColorAndOpacity;
 	if( MyTextBlock.IsValid() )
diff --git a/Engine/Source/Runtime/UMG/Public/Components/TextBlock.h b/Engine/Source/Runtime/UMG/Public/Components/TextBlock.h
--- a/Engine/Source/Runtime/UMG/Public/Components/TextBlock.h
+++ b/Engine/Source/Runtime/UMG/Public/Components/TextBlock.h
@@ -17,7 +17,7 @@ public:
 	 * @param InColorAndOpacity		The new text color and opacity
 	 */
 	UFUNCTION(BlueprintCallable, Category="Appearance")
-	void SetColorAndOpacity(FLinearColor InColorAndOpacity);
+	void SetColorAndOpacity(FSlateColor InColorAndOpacity);
 
 	/**  
 	 * Sets the color and opacity of the text drop shadow
@@ -54,7 +54,7 @@ public:
 
 	/** The color of the text */
 	UPROPERTY(EditDefaultsOnly, Category=Appearance)
-	FLinearColor ColorAndOpacity;
+	FSlateColor ColorAndOpacity;
 
 	/** A bindable delegate for the ColorAndOpacity. */
 	UPROPERTY()
