UMG - Changing the displayname of AutoSize.

diff --git a/Engine/Source/Runtime/UMG/Public/Components/CanvasPanelSlot.h b/Engine/Source/Runtime/UMG/Public/Components/CanvasPanelSlot.h
--- a/Engine/Source/Runtime/UMG/Public/Components/CanvasPanelSlot.h
+++ b/Engine/Source/Runtime/UMG/Public/Components/CanvasPanelSlot.h
@@ -37,7 +37,7 @@ class UMG_API UCanvasPanelSlot : public UPanelSlot
 	FAnchorData LayoutData;
 
 	/** When AutoSize is true we use the widget's desired size */
-	UPROPERTY(EditDefaultsOnly, Category="Layout (Canvas Slot)", AdvancedDisplay)
+	UPROPERTY(EditDefaultsOnly, Category="Layout (Canvas Slot)", AdvancedDisplay, meta=(DisplayName="Size To Content"))
 	bool bAutoSize;
 
 	/** The order priority this widget is rendered in.  Higher values are rendered last (and so they will appear to be on top). */
