Added icons to the various tabs in the UMG

diff --git a/Engine/Source/Editor/UMGEditor/Private/TabFactory/AnimationTabSummoner.cpp b/Engine/Source/Editor/UMGEditor/Private/TabFactory/AnimationTabSummoner.cpp
--- a/Engine/Source/Editor/UMGEditor/Private/TabFactory/AnimationTabSummoner.cpp
+++ b/Engine/Source/Editor/UMGEditor/Private/TabFactory/AnimationTabSummoner.cpp
@@ -17,6 +17,7 @@ FAnimationTabSummoner::FAnimationTabSummoner(TSharedPtr<class FWidgetBlueprintEd
 		, BlueprintEditor(InBlueprintEditor)
 {
 	TabLabel = LOCTEXT("AnimationsTabLabel", "Animations");
+	TabIcon = FUMGStyle::Get().GetBrush("UMGEditor.Tabs.Animations");
 
 	bIsSingleton = true;
 
diff --git a/Engine/Source/Editor/UMGEditor/Private/TabFactory/DesignerTabSummoner.cpp b/Engine/Source/Editor/UMGEditor/Private/TabFactory/DesignerTabSummoner.cpp
--- a/Engine/Source/Editor/UMGEditor/Private/TabFactory/DesignerTabSummoner.cpp
+++ b/Engine/Source/Editor/UMGEditor/Private/TabFactory/DesignerTabSummoner.cpp
@@ -14,7 +14,7 @@ FDesignerTabSummoner::FDesignerTabSummoner(TSharedPtr<class FWidgetBlueprintEdit
 		, BlueprintEditor(InBlueprintEditor)
 {
 	TabLabel = LOCTEXT("DesignerTabLabel", "Designer");
-	TabIcon = FEditorStyle::GetBrush("UMGEditor.Tabs.Designer");
+	TabIcon = FUMGStyle::Get().GetBrush("UMGEditor.Tabs.Designer");
 
 	bIsSingleton = true;
 
diff --git a/Engine/Source/Editor/UMGEditor/Private/TabFactory/DetailsTabSummoner.cpp b/Engine/Source/Editor/UMGEditor/Private/TabFactory/DetailsTabSummoner.cpp
--- a/Engine/Source/Editor/UMGEditor/Private/TabFactory/DetailsTabSummoner.cpp
+++ b/Engine/Source/Editor/UMGEditor/Private/TabFactory/DetailsTabSummoner.cpp
@@ -14,7 +14,7 @@ FDetailsTabSummoner::FDetailsTabSummoner(TSharedPtr<class FWidgetBlueprintEditor
 		, BlueprintEditor(InBlueprintEditor)
 {
 	TabLabel = LOCTEXT("WidgetDetails_TabLabel", "Details");
-	//TabIcon = FEditorStyle::GetBrush("Kismet.Tabs.Components");
+	TabIcon = FEditorStyle::GetBrush("LevelEditor.Tabs.Details");
 
 	bIsSingleton = true;
 
diff --git a/Engine/Source/Editor/UMGEditor/Private/TabFactory/HierarchyTabSummoner.cpp b/Engine/Source/Editor/UMGEditor/Private/TabFactory/HierarchyTabSummoner.cpp
--- a/Engine/Source/Editor/UMGEditor/Private/TabFactory/HierarchyTabSummoner.cpp
+++ b/Engine/Source/Editor/UMGEditor/Private/TabFactory/HierarchyTabSummoner.cpp
@@ -14,7 +14,7 @@ FHierarchyTabSummoner::FHierarchyTabSummoner(TSharedPtr<class FWidgetBlueprintEd
 		, BlueprintEditor(InBlueprintEditor)
 {
 	TabLabel = LOCTEXT("SlateHierarchyTabLabel", "Hierarchy");
-	//TabIcon = FEditorStyle::GetBrush("Kismet.Tabs.Components");
+	TabIcon = FEditorStyle::Get().GetBrush("Kismet.Tabs.Palette");
 
 	bIsSingleton = true;
 
diff --git a/Engine/Source/Editor/UMGEditor/Private/TabFactory/PaletteTabSummoner.cpp b/Engine/Source/Editor/UMGEditor/Private/TabFactory/PaletteTabSummoner.cpp
--- a/Engine/Source/Editor/UMGEditor/Private/TabFactory/PaletteTabSummoner.cpp
+++ b/Engine/Source/Editor/UMGEditor/Private/TabFactory/PaletteTabSummoner.cpp
@@ -14,7 +14,7 @@ FPaletteTabSummoner::FPaletteTabSummoner(TSharedPtr<class FWidgetBlueprintEditor
 		, BlueprintEditor(InBlueprintEditor)
 {
 	TabLabel = LOCTEXT("WidgetTemplatesTabLabel", "Palette");
-	//TabIcon = FEditorStyle::GetBrush("Kismet.Tabs.Components");
+	TabIcon = FUMGStyle::Get().GetBrush("UMGEditor.Tabs.Palette");
 
 	bIsSingleton = true;
 
diff --git a/Engine/Source/Editor/UMGEditor/Private/TabFactory/SequencerTabSummoner.cpp b/Engine/Source/Editor/UMGEditor/Private/TabFactory/SequencerTabSummoner.cpp
--- a/Engine/Source/Editor/UMGEditor/Private/TabFactory/SequencerTabSummoner.cpp
+++ b/Engine/Source/Editor/UMGEditor/Private/TabFactory/SequencerTabSummoner.cpp
@@ -15,7 +15,7 @@ FSequencerTabSummoner::FSequencerTabSummoner(TSharedPtr<class FWidgetBlueprintEd
 		, BlueprintEditor(InBlueprintEditor)
 {
 	TabLabel = LOCTEXT("SequencerLabel", "Timeline");
-
+	TabIcon = FUMGStyle::Get().GetBrush("UMGEditor.Tabs.Timeline");
 
 	bIsSingleton = true;
 
diff --git a/Engine/Source/Editor/UMGEditor/Private/UMGEditorPrivatePCH.h b/Engine/Source/Editor/UMGEditor/Private/UMGEditorPrivatePCH.h
--- a/Engine/Source/Editor/UMGEditor/Private/UMGEditorPrivatePCH.h
+++ b/Engine/Source/Editor/UMGEditor/Private/UMGEditorPrivatePCH.h
@@ -18,6 +18,7 @@
 #include "ObjectEditorUtils.h"
 
 #include "UMG.h"
+#include "UMGStyle.h"
 
 #include "WidgetDesignerSettings.h"
 
diff --git a/Engine/Source/Runtime/UMG/Private/UMGStyle.cpp b/Engine/Source/Runtime/UMG/Private/UMGStyle.cpp
--- a/Engine/Source/Runtime/UMG/Private/UMGStyle.cpp
+++ b/Engine/Source/Runtime/UMG/Private/UMGStyle.cpp
@@ -37,6 +37,7 @@ FName FUMGStyle::GetStyleSetName()
 
 const FVector2D Icon16x16(16.0f, 16.0f);
 const FVector2D Icon20x20(20.0f, 20.0f);
+const FVector2D Icon40x40(40.0f, 40.0f);
 
 TSharedRef< FSlateStyleSet > FUMGStyle::Create()
 {
@@ -81,6 +82,17 @@ TSharedRef< FSlateStyleSet > FUMGStyle::Create()
 
 	Style->Set("Widget.UserWidget", new IMAGE_BRUSH(TEXT("UserWidget"), Icon16x16));
 
+	Style->Set("UMGEditor.Tabs.Animations", new IMAGE_BRUSH(TEXT("Animations_16x"), Icon16x16));
+	Style->Set("UMGEditor.Tabs.Designer", new IMAGE_BRUSH(TEXT("Designer_16x"), Icon16x16));
+	Style->Set("UMGEditor.Tabs.Palette", new IMAGE_BRUSH(TEXT("Palette_16x"), Icon16x16));
+	Style->Set("UMGEditor.Tabs.Timeline", new IMAGE_BRUSH(TEXT("Timeline_16x"), Icon16x16));
+
+	Style->Set("UMGEditor.Icons.Animations", new IMAGE_BRUSH(TEXT("Animations_40x"), Icon40x40));
+	Style->Set("UMGEditor.Icons.Designer", new IMAGE_BRUSH(TEXT("Designer_40x"), Icon40x40));
+	Style->Set("UMGEditor.Icons.Palette", new IMAGE_BRUSH(TEXT("Palette_40x"), Icon40x40));
+	Style->Set("UMGEditor.Icons.Timeline", new IMAGE_BRUSH(TEXT("Timeline_40x"), Icon40x40));
+
+
 	return Style;
 }
 
