UMG - Modal has been removed from AddToViewport now that the PlayerController controls it with SetInputMode.  Updating all widgets to better handle destruction and cleanup of slate resources.  Fixing an issue specifically with STextBlock accessing style information on Tick, which just after a Blueprint compiled can be removed.  UTextBlock's now forcefully clear the style of widgets before being destroyed just in case they're still held onto someone that's presenting them for an extra frame.

diff --git a/Engine/Source/Editor/UMGEditor/Private/Designer/SDesignerView.cpp b/Engine/Source/Editor/UMGEditor/Private/Designer/SDesignerView.cpp
--- a/Engine/Source/Editor/UMGEditor/Private/Designer/SDesignerView.cpp
+++ b/Engine/Source/Editor/UMGEditor/Private/Designer/SDesignerView.cpp
@@ -139,9 +139,6 @@ void SDesignerView::Construct(const FArguments& InArgs, TSharedPtr<FWidgetBluepr
 	Register(MakeShareable(new FUniformGridSlotExtension()));
 	Register(MakeShareable(new FGridSlotExtension()));
 
-	UWidgetBlueprint* Blueprint = GetBlueprint();
-	Blueprint->OnChanged().AddSP(this, &SDesignerView::OnBlueprintChanged);
-
 	FWidgetBlueprintCompiler::OnWidgetBlueprintCompiled.AddSP( this, &SDesignerView::OnBlueprintCompiled );
 
 	BindCommands();
@@ -656,31 +653,16 @@ void SDesignerView::Register(TSharedRef<FDesignerExtension> Extension)
 	DesignerExtensions.Add(Extension);
 }
 
-void SDesignerView::OnBlueprintChanged(UBlueprint* InBlueprint)
-{
-	if ( InBlueprint )
-	{
-		
-	}
-}
-
 void SDesignerView::OnBlueprintCompiled(UBlueprint* InBlueprint)
 {
 	UBlueprint* BP = BlueprintEditor.Pin()->GetBlueprintObj();
 	if( BP == InBlueprint )
 	{
-		UpdatePreviewWidget();
-	}
-}
+		CachedWidgetGeometry.Reset();
 
-void SDesignerView::OnObjectPropertyChanged(UObject* ObjectBeingModified, FPropertyChangedEvent& PropertyChangedEvent)
-{
-	if ( !ensure(ObjectBeingModified) )
-	{
-		return;
+		const bool bForceUpdate = true;
+		UpdatePreviewWidget(bForceUpdate);
 	}
-
-	//UpdatePreview(InBlueprint);
 }
 
 FWidgetReference SDesignerView::GetWidgetAtCursor(const FGeometry& MyGeometry, const FPointerEvent& MouseEvent, FArrangedWidget& ArrangedWidget)
@@ -1022,11 +1004,11 @@ int32 SDesignerView::OnPaint(const FPaintArgs& Args, const FGeometry& AllottedGe
 	return LayerId;
 }
 
-void SDesignerView::UpdatePreviewWidget()
+void SDesignerView::UpdatePreviewWidget(bool bForceUpdate)
 {
 	UUserWidget* LatestPreviewWidget = BlueprintEditor.Pin()->GetPreview();
 
-	if ( LatestPreviewWidget != PreviewWidget )
+	if ( LatestPreviewWidget != PreviewWidget || bForceUpdate )
 	{
 		PreviewWidget = LatestPreviewWidget;
 		if ( PreviewWidget )
@@ -1070,7 +1052,8 @@ void SDesignerView::Tick(const FGeometry& AllottedGeometry, const double InCurre
 	CachedDesignerGeometry = AllottedGeometry;
 	HoverTime += InDeltaTime;
 
-	UpdatePreviewWidget();
+	const bool bForceUpdate = false;
+	UpdatePreviewWidget(bForceUpdate);
 
 	// Update the selected widget to match the selected template.
 	if ( PreviewWidget )
diff --git a/Engine/Source/Editor/UMGEditor/Private/Designer/SDesignerView.h b/Engine/Source/Editor/UMGEditor/Private/Designer/SDesignerView.h
--- a/Engine/Source/Editor/UMGEditor/Private/Designer/SDesignerView.h
+++ b/Engine/Source/Editor/UMGEditor/Private/Designer/SDesignerView.h
@@ -83,7 +83,7 @@ private:
 	void ResolvePendingSelectedWidgets();
 
 	/** Updates the designer to display the latest preview widget */
-	void UpdatePreviewWidget();
+	void UpdatePreviewWidget(bool bForceUpdate);
 
 	void ClearExtensionWidgets();
 	void CreateExtensionWidgetsForSelection();
@@ -102,14 +102,9 @@ private:
 	/** Gets the blueprint being edited by the designer */
 	UWidgetBlueprint* GetBlueprint() const;
 
-	/** Called whenever the blueprint is structurally changed. */
-	void OnBlueprintChanged(UBlueprint* InBlueprint);
-
-        /** Called whenever the blueprint is recompiled */
+	/** Called whenever the blueprint is recompiled */
 	void OnBlueprintCompiled(UBlueprint* InBlueprint);
 
-	void OnObjectPropertyChanged(UObject* ObjectBeingModified, FPropertyChangedEvent& PropertyChangedEvent);
-
 	void PopulateWidgetGeometryCache(FArrangedWidget& Root);
 
 	void CacheSelectedWidgetGeometry();
diff --git a/Engine/Source/Editor/UMGEditor/Public/IUMGDesigner.h b/Engine/Source/Editor/UMGEditor/Public/IUMGDesigner.h
--- a/Engine/Source/Editor/UMGEditor/Public/IUMGDesigner.h
+++ b/Engine/Source/Editor/UMGEditor/Public/IUMGDesigner.h
@@ -35,9 +35,13 @@ public:
 	/** @return The Geometry representing the designer area, useful for when you need to convert mouse into designer space. */
 	virtual FGeometry GetDesignerGeometry() const = 0;
 
-	/**  */
+	/**
+	 * Gets the previous frames widget geometry.
+	 */
 	virtual bool GetWidgetGeometry(const FWidgetReference& Widget, FGeometry& Geometry) const = 0;
 
-	/**  */
+	/**
+	 * Gets the previous frames widget geometry of the parent of the provided widget.
+	 */
 	virtual bool GetWidgetParentGeometry(const FWidgetReference& Widget, FGeometry& Geometry) const = 0;
 };
diff --git a/Engine/Source/Runtime/UMG/Private/Components/Border.cpp b/Engine/Source/Runtime/UMG/Private/Components/Border.cpp
--- a/Engine/Source/Runtime/UMG/Private/Components/Border.cpp
+++ b/Engine/Source/Runtime/UMG/Private/Components/Border.cpp
@@ -25,9 +25,9 @@ UBorder::UBorder(const FPostConstructInitializeProperties& PCIP)
 	bShowEffectWhenDisabled = true;
 }
 
-void UBorder::ReleaseNativeWidget()
+void UBorder::ReleaseSlateResources(bool bReleaseChildren)
 {
-	Super::ReleaseNativeWidget();
+	Super::ReleaseSlateResources(bReleaseChildren);
 
 	MyBorder.Reset();
 }
diff --git a/Engine/Source/Runtime/UMG/Private/Components/BorderSlot.cpp b/Engine/Source/Runtime/UMG/Private/Components/BorderSlot.cpp
--- a/Engine/Source/Runtime/UMG/Private/Components/BorderSlot.cpp
+++ b/Engine/Source/Runtime/UMG/Private/Components/BorderSlot.cpp
@@ -15,9 +15,9 @@ UBorderSlot::UBorderSlot(const FPostConstructInitializeProperties& PCIP)
 	VerticalAlignment = VAlign_Fill;
 }
 
-void UBorderSlot::ReleaseNativeWidget()
+void UBorderSlot::ReleaseSlateResources(bool bReleaseChildren)
 {
-	Super::ReleaseNativeWidget();
+	Super::ReleaseSlateResources(bReleaseChildren);
 
 	Border.Reset();
 }
diff --git a/Engine/Source/Runtime/UMG/Private/Components/Button.cpp b/Engine/Source/Runtime/UMG/Private/Components/Button.cpp
--- a/Engine/Source/Runtime/UMG/Private/Components/Button.cpp
+++ b/Engine/Source/Runtime/UMG/Private/Components/Button.cpp
@@ -22,9 +22,9 @@ UButton::UButton(const FPostConstructInitializeProperties& PCIP)
 	IsFocusable = true;
 }
 
-void UButton::ReleaseNativeWidget()
+void UButton::ReleaseSlateResources(bool bReleaseChildren)
 {
-	Super::ReleaseNativeWidget();
+	Super::ReleaseSlateResources(bReleaseChildren);
 
 	MyButton.Reset();
 }
diff --git a/Engine/Source/Runtime/UMG/Private/Components/ButtonSlot.cpp b/Engine/Source/Runtime/UMG/Private/Components/ButtonSlot.cpp
--- a/Engine/Source/Runtime/UMG/Private/Components/ButtonSlot.cpp
+++ b/Engine/Source/Runtime/UMG/Private/Components/ButtonSlot.cpp
@@ -14,9 +14,9 @@ UButtonSlot::UButtonSlot(const FPostConstructInitializeProperties& PCIP)
 	VerticalAlignment = VAlign_Center;
 }
 
-void UButtonSlot::ReleaseNativeWidget()
+void UButtonSlot::ReleaseSlateResources(bool bReleaseChildren)
 {
-	Super::ReleaseNativeWidget();
+	Super::ReleaseSlateResources(bReleaseChildren);
 
 	Button.Reset();
 }
diff --git a/Engine/Source/Runtime/UMG/Private/Components/CanvasPanel.cpp b/Engine/Source/Runtime/UMG/Private/Components/CanvasPanel.cpp
--- a/Engine/Source/Runtime/UMG/Private/Components/CanvasPanel.cpp
+++ b/Engine/Source/Runtime/UMG/Private/Components/CanvasPanel.cpp
@@ -16,9 +16,9 @@ UCanvasPanel::UCanvasPanel(const FPostConstructInitializeProperties& PCIP)
 	Visiblity = UWidget::ConvertRuntimeToSerializedVisiblity(Defaults._Visibility.Get());
 }
 
-void UCanvasPanel::ReleaseNativeWidget()
+void UCanvasPanel::ReleaseSlateResources(bool bReleaseChildren)
 {
-	Super::ReleaseNativeWidget();
+	Super::ReleaseSlateResources(bReleaseChildren);
 
 	MyCanvas.Reset();
 }
diff --git a/Engine/Source/Runtime/UMG/Private/Components/CanvasPanelSlot.cpp b/Engine/Source/Runtime/UMG/Private/Components/CanvasPanelSlot.cpp
--- a/Engine/Source/Runtime/UMG/Private/Components/CanvasPanelSlot.cpp
+++ b/Engine/Source/Runtime/UMG/Private/Components/CanvasPanelSlot.cpp
@@ -16,9 +16,9 @@ UCanvasPanelSlot::UCanvasPanelSlot(const FPostConstructInitializeProperties& PCI
 	ZOrder = 0;
 }
 
-void UCanvasPanelSlot::ReleaseNativeWidget()
+void UCanvasPanelSlot::ReleaseSlateResources(bool bReleaseChildren)
 {
-	Super::ReleaseNativeWidget();
+	Super::ReleaseSlateResources(bReleaseChildren);
 
 	Slot = NULL;
 }
diff --git a/Engine/Source/Runtime/UMG/Private/Components/CheckBox.cpp b/Engine/Source/Runtime/UMG/Private/Components/CheckBox.cpp
--- a/Engine/Source/Runtime/UMG/Private/Components/CheckBox.cpp
+++ b/Engine/Source/Runtime/UMG/Private/Components/CheckBox.cpp
@@ -21,9 +21,9 @@ UCheckBox::UCheckBox(const FPostConstructInitializeProperties& PCIP)
 	BorderBackgroundColor = FLinearColor::White;
 }
 
-void UCheckBox::ReleaseNativeWidget()
+void UCheckBox::ReleaseSlateResources(bool bReleaseChildren)
 {
-	Super::ReleaseNativeWidget();
+	Super::ReleaseSlateResources(bReleaseChildren);
 
 	MyCheckbox.Reset();
 }
diff --git a/Engine/Source/Runtime/UMG/Private/Components/CircularThrobber.cpp b/Engine/Source/Runtime/UMG/Private/Components/CircularThrobber.cpp
--- a/Engine/Source/Runtime/UMG/Private/Components/CircularThrobber.cpp
+++ b/Engine/Source/Runtime/UMG/Private/Components/CircularThrobber.cpp
@@ -19,9 +19,9 @@ UCircularThrobber::UCircularThrobber(const FPostConstructInitializeProperties& P
 	Radius = DefaultArgs._Radius;
 }
 
-void UCircularThrobber::ReleaseNativeWidget()
+void UCircularThrobber::ReleaseSlateResources(bool bReleaseChildren)
 {
-	Super::ReleaseNativeWidget();
+	Super::ReleaseSlateResources(bReleaseChildren);
 
 	MyCircularThrobber.Reset();
 }
diff --git a/Engine/Source/Runtime/UMG/Private/Components/ComboBox.cpp b/Engine/Source/Runtime/UMG/Private/Components/ComboBox.cpp
--- a/Engine/Source/Runtime/UMG/Private/Components/ComboBox.cpp
+++ b/Engine/Source/Runtime/UMG/Private/Components/ComboBox.cpp
@@ -12,9 +12,9 @@ UComboBox::UComboBox(const FPostConstructInitializeProperties& PCIP)
 {
 }
 
-void UComboBox::ReleaseNativeWidget()
+void UComboBox::ReleaseSlateResources(bool bReleaseChildren)
 {
-	Super::ReleaseNativeWidget();
+	Super::ReleaseSlateResources(bReleaseChildren);
 
 	MyComboBox.Reset();
 }
diff --git a/Engine/Source/Runtime/UMG/Private/Components/ComboBoxString.cpp b/Engine/Source/Runtime/UMG/Private/Components/ComboBoxString.cpp
--- a/Engine/Source/Runtime/UMG/Private/Components/ComboBoxString.cpp
+++ b/Engine/Source/Runtime/UMG/Private/Components/ComboBoxString.cpp
@@ -18,9 +18,9 @@ UComboBoxString::UComboBoxString(const FPostConstructInitializeProperties& PCIP)
 	HasDownArrow = true;
 }
 
-void UComboBoxString::ReleaseNativeWidget()
+void UComboBoxString::ReleaseSlateResources(bool bReleaseChildren)
 {
-	Super::ReleaseNativeWidget();
+	Super::ReleaseSlateResources(bReleaseChildren);
 
 	MyComboBox.Reset();
 	ComoboBoxContent.Reset();
diff --git a/Engine/Source/Runtime/UMG/Private/Components/EditableText.cpp b/Engine/Source/Runtime/UMG/Private/Components/EditableText.cpp
--- a/Engine/Source/Runtime/UMG/Private/Components/EditableText.cpp
+++ b/Engine/Source/Runtime/UMG/Private/Components/EditableText.cpp
@@ -29,9 +29,9 @@ UEditableText::UEditableText(const FPostConstructInitializeProperties& PCIP)
 	SelectAllTextOnCommit = Defaults._SelectAllTextOnCommit.Get();
 }
 
-void UEditableText::ReleaseNativeWidget()
+void UEditableText::ReleaseSlateResources(bool bReleaseChildren)
 {
-	Super::ReleaseNativeWidget();
+	Super::ReleaseSlateResources(bReleaseChildren);
 
 	MyEditableText.Reset();
 }
diff --git a/Engine/Source/Runtime/UMG/Private/Components/EditableTextBox.cpp b/Engine/Source/Runtime/UMG/Private/Components/EditableTextBox.cpp
--- a/Engine/Source/Runtime/UMG/Private/Components/EditableTextBox.cpp
+++ b/Engine/Source/Runtime/UMG/Private/Components/EditableTextBox.cpp
@@ -32,9 +32,9 @@ UEditableTextBox::UEditableTextBox(const FPostConstructInitializeProperties& PCI
 	WidgetStyle = *Defaults._Style;
 }
 
-void UEditableTextBox::ReleaseNativeWidget()
+void UEditableTextBox::ReleaseSlateResources(bool bReleaseChildren)
 {
-	Super::ReleaseNativeWidget();
+	Super::ReleaseSlateResources(bReleaseChildren);
 
 	MyEditableTextBlock.Reset();
 }
diff --git a/Engine/Source/Runtime/UMG/Private/Components/GridPanel.cpp b/Engine/Source/Runtime/UMG/Private/Components/GridPanel.cpp
--- a/Engine/Source/Runtime/UMG/Private/Components/GridPanel.cpp
+++ b/Engine/Source/Runtime/UMG/Private/Components/GridPanel.cpp
@@ -16,9 +16,9 @@ UGridPanel::UGridPanel(const FPostConstructInitializeProperties& PCIP)
 	Visiblity = UWidget::ConvertRuntimeToSerializedVisiblity(Defaults._Visibility.Get());
 }
 
-void UGridPanel::ReleaseNativeWidget()
+void UGridPanel::ReleaseSlateResources(bool bReleaseChildren)
 {
-	Super::ReleaseNativeWidget();
+	Super::ReleaseSlateResources(bReleaseChildren);
 
 	MyGridPanel.Reset();
 }
diff --git a/Engine/Source/Runtime/UMG/Private/Components/GridSlot.cpp b/Engine/Source/Runtime/UMG/Private/Components/GridSlot.cpp
--- a/Engine/Source/Runtime/UMG/Private/Components/GridSlot.cpp
+++ b/Engine/Source/Runtime/UMG/Private/Components/GridSlot.cpp
@@ -16,9 +16,9 @@ UGridSlot::UGridSlot(const FPostConstructInitializeProperties& PCIP)
 	Nudge = FVector2D(0, 0);
 }
 
-void UGridSlot::ReleaseNativeWidget()
+void UGridSlot::ReleaseSlateResources(bool bReleaseChildren)
 {
-	Super::ReleaseNativeWidget();
+	Super::ReleaseSlateResources(bReleaseChildren);
 
 	Slot = NULL;
 }
diff --git a/Engine/Source/Runtime/UMG/Private/Components/HorizontalBox.cpp b/Engine/Source/Runtime/UMG/Private/Components/HorizontalBox.cpp
--- a/Engine/Source/Runtime/UMG/Private/Components/HorizontalBox.cpp
+++ b/Engine/Source/Runtime/UMG/Private/Components/HorizontalBox.cpp
@@ -16,9 +16,9 @@ UHorizontalBox::UHorizontalBox(const FPostConstructInitializeProperties& PCIP)
 	Visiblity = UWidget::ConvertRuntimeToSerializedVisiblity(Defaults._Visibility.Get());
 }
 
-void UHorizontalBox::ReleaseNativeWidget()
+void UHorizontalBox::ReleaseSlateResources(bool bReleaseChildren)
 {
-	Super::ReleaseNativeWidget();
+	Super::ReleaseSlateResources(bReleaseChildren);
 
 	MyHorizontalBox.Reset();
 }
diff --git a/Engine/Source/Runtime/UMG/Private/Components/HorizontalBoxSlot.cpp b/Engine/Source/Runtime/UMG/Private/Components/HorizontalBoxSlot.cpp
--- a/Engine/Source/Runtime/UMG/Private/Components/HorizontalBoxSlot.cpp
+++ b/Engine/Source/Runtime/UMG/Private/Components/HorizontalBoxSlot.cpp
@@ -13,9 +13,9 @@ UHorizontalBoxSlot::UHorizontalBoxSlot(const FPostConstructInitializeProperties&
 	Size = FSlateChildSize(ESlateSizeRule::Automatic);
 }
 
-void UHorizontalBoxSlot::ReleaseNativeWidget()
+void UHorizontalBoxSlot::ReleaseSlateResources(bool bReleaseChildren)
 {
-	Super::ReleaseNativeWidget();
+	Super::ReleaseSlateResources(bReleaseChildren);
 
 	Slot = NULL;
 }
diff --git a/Engine/Source/Runtime/UMG/Private/Components/Image.cpp b/Engine/Source/Runtime/UMG/Private/Components/Image.cpp
--- a/Engine/Source/Runtime/UMG/Private/Components/Image.cpp
+++ b/Engine/Source/Runtime/UMG/Private/Components/Image.cpp
@@ -25,9 +25,9 @@ void UImage::PostLoad()
 	}
 }
 
-void UImage::ReleaseNativeWidget()
+void UImage::ReleaseSlateResources(bool bReleaseChildren)
 {
-	Super::ReleaseNativeWidget();
+	Super::ReleaseSlateResources(bReleaseChildren);
 
 	MyImage.Reset();
 }
diff --git a/Engine/Source/Runtime/UMG/Private/Components/MenuAnchor.cpp b/Engine/Source/Runtime/UMG/Private/Components/MenuAnchor.cpp
--- a/Engine/Source/Runtime/UMG/Private/Components/MenuAnchor.cpp
+++ b/Engine/Source/Runtime/UMG/Private/Components/MenuAnchor.cpp
@@ -13,9 +13,9 @@ UMenuAnchor::UMenuAnchor(const FPostConstructInitializeProperties& PCIP)
 	Placement = MenuPlacement_ComboBox;
 }
 
-void UMenuAnchor::ReleaseNativeWidget()
+void UMenuAnchor::ReleaseSlateResources(bool bReleaseChildren)
 {
-	Super::ReleaseNativeWidget();
+	Super::ReleaseSlateResources(bReleaseChildren);
 
 	MyMenuAnchor.Reset();
 }
diff --git a/Engine/Source/Runtime/UMG/Private/Components/MultiLineEditableText.cpp b/Engine/Source/Runtime/UMG/Private/Components/MultiLineEditableText.cpp
--- a/Engine/Source/Runtime/UMG/Private/Components/MultiLineEditableText.cpp
+++ b/Engine/Source/Runtime/UMG/Private/Components/MultiLineEditableText.cpp
@@ -17,9 +17,9 @@ UMultiLineEditableText::UMultiLineEditableText(const FPostConstructInitializePro
 	Font = FSlateFontInfo(TEXT("Slate/Fonts/Roboto-Bold.ttf"), 12);
 }
 
-void UMultiLineEditableText::ReleaseNativeWidget()
+void UMultiLineEditableText::ReleaseSlateResources(bool bReleaseChildren)
 {
-	Super::ReleaseNativeWidget();
+	Super::ReleaseSlateResources(bReleaseChildren);
 
 	MyMultiLineEditableText.Reset();
 }
diff --git a/Engine/Source/Runtime/UMG/Private/Components/MultiLineEditableTextBox.cpp b/Engine/Source/Runtime/UMG/Private/Components/MultiLineEditableTextBox.cpp
--- a/Engine/Source/Runtime/UMG/Private/Components/MultiLineEditableTextBox.cpp
+++ b/Engine/Source/Runtime/UMG/Private/Components/MultiLineEditableTextBox.cpp
@@ -22,9 +22,9 @@ UMultiLineEditableTextBox::UMultiLineEditableTextBox(const FPostConstructInitial
 	TextStyle = *Defaults._TextStyle;
 }
 
-void UMultiLineEditableTextBox::ReleaseNativeWidget()
+void UMultiLineEditableTextBox::ReleaseSlateResources(bool bReleaseChildren)
 {
-	Super::ReleaseNativeWidget();
+	Super::ReleaseSlateResources(bReleaseChildren);
 
 	MyEditableTextBlock.Reset();
 }
diff --git a/Engine/Source/Runtime/UMG/Private/Components/NativeWidgetHost.cpp b/Engine/Source/Runtime/UMG/Private/Components/NativeWidgetHost.cpp
--- a/Engine/Source/Runtime/UMG/Private/Components/NativeWidgetHost.cpp
+++ b/Engine/Source/Runtime/UMG/Private/Components/NativeWidgetHost.cpp
@@ -18,9 +18,9 @@ void UNativeWidgetHost::SetContent(TSharedRef<SWidget> InContent)
 	MyWidget = InContent;
 }
 
-void UNativeWidgetHost::ReleaseNativeWidget()
+void UNativeWidgetHost::ReleaseSlateResources(bool bReleaseChildren)
 {
-	Super::ReleaseNativeWidget();
+	Super::ReleaseSlateResources(bReleaseChildren);
 
 	MyWidget.Reset();
 }
diff --git a/Engine/Source/Runtime/UMG/Private/Components/Overlay.cpp b/Engine/Source/Runtime/UMG/Private/Components/Overlay.cpp
--- a/Engine/Source/Runtime/UMG/Private/Components/Overlay.cpp
+++ b/Engine/Source/Runtime/UMG/Private/Components/Overlay.cpp
@@ -16,9 +16,9 @@ UOverlay::UOverlay(const FPostConstructInitializeProperties& PCIP)
 	Visiblity = UWidget::ConvertRuntimeToSerializedVisiblity(Defaults._Visibility.Get());
 }
 
-void UOverlay::ReleaseNativeWidget()
+void UOverlay::ReleaseSlateResources(bool bReleaseChildren)
 {
-	Super::ReleaseNativeWidget();
+	Super::ReleaseSlateResources(bReleaseChildren);
 
 	MyOverlay.Reset();
 }
diff --git a/Engine/Source/Runtime/UMG/Private/Components/OverlaySlot.cpp b/Engine/Source/Runtime/UMG/Private/Components/OverlaySlot.cpp
--- a/Engine/Source/Runtime/UMG/Private/Components/OverlaySlot.cpp
+++ b/Engine/Source/Runtime/UMG/Private/Components/OverlaySlot.cpp
@@ -12,9 +12,9 @@ UOverlaySlot::UOverlaySlot(const FPostConstructInitializeProperties& PCIP)
 	VerticalAlignment = VAlign_Top;
 }
 
-void UOverlaySlot::ReleaseNativeWidget()
+void UOverlaySlot::ReleaseSlateResources(bool bReleaseChildren)
 {
-	Super::ReleaseNativeWidget();
+	Super::ReleaseSlateResources(bReleaseChildren);
 
 	Slot = NULL;
 }
diff --git a/Engine/Source/Runtime/UMG/Private/Components/PanelSlot.cpp b/Engine/Source/Runtime/UMG/Private/Components/PanelSlot.cpp
--- a/Engine/Source/Runtime/UMG/Private/Components/PanelSlot.cpp
+++ b/Engine/Source/Runtime/UMG/Private/Components/PanelSlot.cpp
@@ -45,12 +45,12 @@ bool UPanelSlot::IsDesignTime() const
 	return Parent->IsDesignTime();
 }
 
-void UPanelSlot::ReleaseNativeWidget()
+void UPanelSlot::ReleaseSlateResources(bool bReleaseChildren)
 {
-	Super::ReleaseNativeWidget();
+	Super::ReleaseSlateResources(bReleaseChildren);
 
 	if ( Content )
 	{
-		Content->ReleaseNativeWidget();
+		Content->ReleaseSlateResources(bReleaseChildren);
 	}
 }
\ No newline at end of file
diff --git a/Engine/Source/Runtime/UMG/Private/Components/PanelWidget.cpp b/Engine/Source/Runtime/UMG/Private/Components/PanelWidget.cpp
--- a/Engine/Source/Runtime/UMG/Private/Components/PanelWidget.cpp
+++ b/Engine/Source/Runtime/UMG/Private/Components/PanelWidget.cpp
@@ -18,15 +18,18 @@ UPanelWidget::UPanelWidget(const FPostConstructInitializeProperties& PCIP)
 {
 }
 
-void UPanelWidget::ReleaseNativeWidget()
+void UPanelWidget::ReleaseSlateResources(bool bReleaseChildren)
 {
-	Super::ReleaseNativeWidget();
+	Super::ReleaseSlateResources(bReleaseChildren);
 
-	for ( int32 SlotIndex = 0; SlotIndex < Slots.Num(); SlotIndex++ )
+	if ( bReleaseChildren )
 	{
-		if ( Slots[SlotIndex]->Content != NULL )
+		for ( int32 SlotIndex = 0; SlotIndex < Slots.Num(); SlotIndex++ )
 		{
-			Slots[SlotIndex]->ReleaseNativeWidget();
+			if ( Slots[SlotIndex]->Content != NULL )
+			{
+				Slots[SlotIndex]->ReleaseSlateResources(bReleaseChildren);
+			}
 		}
 	}
 }
diff --git a/Engine/Source/Runtime/UMG/Private/Components/ProgressBar.cpp b/Engine/Source/Runtime/UMG/Private/Components/ProgressBar.cpp
--- a/Engine/Source/Runtime/UMG/Private/Components/ProgressBar.cpp
+++ b/Engine/Source/Runtime/UMG/Private/Components/ProgressBar.cpp
@@ -19,9 +19,9 @@ UProgressBar::UProgressBar(const FPostConstructInitializeProperties& PCIP)
 	FillColorAndOpacity = FLinearColor::White;
 }
 
-void UProgressBar::ReleaseNativeWidget()
+void UProgressBar::ReleaseSlateResources(bool bReleaseChildren)
 {
-	Super::ReleaseNativeWidget();
+	Super::ReleaseSlateResources(bReleaseChildren);
 
 	MyProgressBar.Reset();
 }
diff --git a/Engine/Source/Runtime/UMG/Private/Components/ScaleBox.cpp b/Engine/Source/Runtime/UMG/Private/Components/ScaleBox.cpp
--- a/Engine/Source/Runtime/UMG/Private/Components/ScaleBox.cpp
+++ b/Engine/Source/Runtime/UMG/Private/Components/ScaleBox.cpp
@@ -16,9 +16,9 @@ UScaleBox::UScaleBox(const FPostConstructInitializeProperties& PCIP)
 	Stretch = EStretch::ScaleToFit;
 }
 
-void UScaleBox::ReleaseNativeWidget()
+void UScaleBox::ReleaseSlateResources(bool bReleaseChildren)
 {
-	Super::ReleaseNativeWidget();
+	Super::ReleaseSlateResources(bReleaseChildren);
 
 	MyScaleBox.Reset();
 }
diff --git a/Engine/Source/Runtime/UMG/Private/Components/ScaleBoxSlot.cpp b/Engine/Source/Runtime/UMG/Private/Components/ScaleBoxSlot.cpp
--- a/Engine/Source/Runtime/UMG/Private/Components/ScaleBoxSlot.cpp
+++ b/Engine/Source/Runtime/UMG/Private/Components/ScaleBoxSlot.cpp
@@ -14,9 +14,9 @@ UScaleBoxSlot::UScaleBoxSlot(const FPostConstructInitializeProperties& PCIP)
 	VerticalAlignment = VAlign_Center;
 }
 
-void UScaleBoxSlot::ReleaseNativeWidget()
+void UScaleBoxSlot::ReleaseSlateResources(bool bReleaseChildren)
 {
-	Super::ReleaseNativeWidget();
+	Super::ReleaseSlateResources(bReleaseChildren);
 
 	ScaleBox.Reset();
 }
diff --git a/Engine/Source/Runtime/UMG/Private/Components/ScrollBar.cpp b/Engine/Source/Runtime/UMG/Private/Components/ScrollBar.cpp
--- a/Engine/Source/Runtime/UMG/Private/Components/ScrollBar.cpp
+++ b/Engine/Source/Runtime/UMG/Private/Components/ScrollBar.cpp
@@ -20,9 +20,9 @@ UScrollBar::UScrollBar(const FPostConstructInitializeProperties& PCIP)
 	WidgetStyle = *Defaults._Style;
 }
 
-void UScrollBar::ReleaseNativeWidget()
+void UScrollBar::ReleaseSlateResources(bool bReleaseChildren)
 {
-	Super::ReleaseNativeWidget();
+	Super::ReleaseSlateResources(bReleaseChildren);
 
 	MyScrollBar.Reset();
 }
diff --git a/Engine/Source/Runtime/UMG/Private/Components/ScrollBox.cpp b/Engine/Source/Runtime/UMG/Private/Components/ScrollBox.cpp
--- a/Engine/Source/Runtime/UMG/Private/Components/ScrollBox.cpp
+++ b/Engine/Source/Runtime/UMG/Private/Components/ScrollBox.cpp
@@ -21,9 +21,9 @@ UScrollBox::UScrollBox(const FPostConstructInitializeProperties& PCIP)
 	WidgetBarStyle = *Defaults._ScrollBarStyle;
 }
 
-void UScrollBox::ReleaseNativeWidget()
+void UScrollBox::ReleaseSlateResources(bool bReleaseChildren)
 {
-	Super::ReleaseNativeWidget();
+	Super::ReleaseSlateResources(bReleaseChildren);
 
 	MyScrollBox.Reset();
 }
diff --git a/Engine/Source/Runtime/UMG/Private/Components/ScrollBoxSlot.cpp b/Engine/Source/Runtime/UMG/Private/Components/ScrollBoxSlot.cpp
--- a/Engine/Source/Runtime/UMG/Private/Components/ScrollBoxSlot.cpp
+++ b/Engine/Source/Runtime/UMG/Private/Components/ScrollBoxSlot.cpp
@@ -46,8 +46,8 @@ void UScrollBoxSlot::SynchronizeProperties()
 	SetHorizontalAlignment(HorizontalAlignment);
 }
 
-void UScrollBoxSlot::ReleaseNativeWidget()
+void UScrollBoxSlot::ReleaseSlateResources(bool bReleaseChildren)
 {
-	Super::ReleaseNativeWidget();
+	Super::ReleaseSlateResources(bReleaseChildren);
 	Slot = nullptr;
 }
diff --git a/Engine/Source/Runtime/UMG/Private/Components/SizeBox.cpp b/Engine/Source/Runtime/UMG/Private/Components/SizeBox.cpp
--- a/Engine/Source/Runtime/UMG/Private/Components/SizeBox.cpp
+++ b/Engine/Source/Runtime/UMG/Private/Components/SizeBox.cpp
@@ -14,9 +14,9 @@ USizeBox::USizeBox(const FPostConstructInitializeProperties& PCIP)
 	bIsVariable = false;
 }
 
-void USizeBox::ReleaseNativeWidget()
+void USizeBox::ReleaseSlateResources(bool bReleaseChildren)
 {
-	Super::ReleaseNativeWidget();
+	Super::ReleaseSlateResources(bReleaseChildren);
 
 	MySizeBox.Reset();
 }
diff --git a/Engine/Source/Runtime/UMG/Private/Components/SizeBoxSlot.cpp b/Engine/Source/Runtime/UMG/Private/Components/SizeBoxSlot.cpp
--- a/Engine/Source/Runtime/UMG/Private/Components/SizeBoxSlot.cpp
+++ b/Engine/Source/Runtime/UMG/Private/Components/SizeBoxSlot.cpp
@@ -16,9 +16,9 @@ USizeBoxSlot::USizeBoxSlot(const FPostConstructInitializeProperties& PCIP)
 	SBox::FArguments SizeBoxDefaults;
 }
 
-void USizeBoxSlot::ReleaseNativeWidget()
+void USizeBoxSlot::ReleaseSlateResources(bool bReleaseChildren)
 {
-	Super::ReleaseNativeWidget();
+	Super::ReleaseSlateResources(bReleaseChildren);
 
 	SizeBox.Reset();
 }
diff --git a/Engine/Source/Runtime/UMG/Private/Components/Spacer.cpp b/Engine/Source/Runtime/UMG/Private/Components/Spacer.cpp
--- a/Engine/Source/Runtime/UMG/Private/Components/Spacer.cpp
+++ b/Engine/Source/Runtime/UMG/Private/Components/Spacer.cpp
@@ -14,9 +14,9 @@ USpacer::USpacer(const FPostConstructInitializeProperties& PCIP)
 	bIsVariable = false;
 }
 
-void USpacer::ReleaseNativeWidget()
+void USpacer::ReleaseSlateResources(bool bReleaseChildren)
 {
-	Super::ReleaseNativeWidget();
+	Super::ReleaseSlateResources(bReleaseChildren);
 
 	MySpacer.Reset();
 }
diff --git a/Engine/Source/Runtime/UMG/Private/Components/SpinBox.cpp b/Engine/Source/Runtime/UMG/Private/Components/SpinBox.cpp
--- a/Engine/Source/Runtime/UMG/Private/Components/SpinBox.cpp
+++ b/Engine/Source/Runtime/UMG/Private/Components/SpinBox.cpp
@@ -31,9 +31,9 @@ USpinBox::USpinBox(const FPostConstructInitializeProperties& PCIP)
 	ForegroundColor = FSlateColor(FLinearColor::Black);
 }
 
-void USpinBox::ReleaseNativeWidget()
+void USpinBox::ReleaseSlateResources(bool bReleaseChildren)
 {
-	Super::ReleaseNativeWidget();
+	Super::ReleaseSlateResources(bReleaseChildren);
 
 	MySpinBox.Reset();
 }
diff --git a/Engine/Source/Runtime/UMG/Private/Components/TextBlock.cpp b/Engine/Source/Runtime/UMG/Private/Components/TextBlock.cpp
--- a/Engine/Source/Runtime/UMG/Private/Components/TextBlock.cpp
+++ b/Engine/Source/Runtime/UMG/Private/Components/TextBlock.cpp
@@ -12,8 +12,7 @@ UTextBlock::UTextBlock(const FPostConstructInitializeProperties& PCIP)
 {
 	bIsVariable = false;
 
-	STextBlock::FArguments Defaults;
-	WidgetStyle = *Defaults._TextStyle;
+	WidgetStyle = *UTextBlock::GetDefaultStyle();
 
 	Text = LOCTEXT("TextBlockDefaultValue", "Text Block");
 	ShadowOffset = FVector2D(1.0f, 1.0f);
@@ -25,11 +24,16 @@ UTextBlock::UTextBlock(const FPostConstructInitializeProperties& PCIP)
 	Font = FSlateFontInfo(TEXT("Slate/Fonts/Roboto-Bold.ttf"), 24);
 }
 
-void UTextBlock::ReleaseNativeWidget()
+void UTextBlock::ReleaseSlateResources(bool bReleaseChildren)
 {
-	Super::ReleaseNativeWidget();
+	Super::ReleaseSlateResources(bReleaseChildren);
 
-	MyTextBlock.Reset();
+	if ( MyTextBlock.IsValid() )
+	{
+		MyTextBlock->SetTextStyle(UTextBlock::GetDefaultStyle());
+
+		MyTextBlock.Reset();
+	}
 }
 
 void UTextBlock::SetColorAndOpacity(FSlateColor InColorAndOpacity)
@@ -44,7 +48,7 @@ void UTextBlock::SetColorAndOpacity(FSlateColor InColorAndOpacity)
 void UTextBlock::SetShadowColorAndOpacity(FLinearColor InShadowColorAndOpacity)
 {
 	ShadowColorAndOpacity = InShadowColorAndOpacity;
-	if(MyTextBlock.IsValid())
+	if( MyTextBlock.IsValid() )
 	{
 		MyTextBlock->SetShadowColorAndOpacity(InShadowColorAndOpacity);
 	}
@@ -53,7 +57,7 @@ void UTextBlock::SetShadowColorAndOpacity(FLinearColor InShadowColorAndOpacity)
 void UTextBlock::SetShadowOffset(FVector2D InShadowOffset)
 {
 	ShadowOffset = InShadowOffset;
-	if(MyTextBlock.IsValid())
+	if( MyTextBlock.IsValid() )
 	{
 		MyTextBlock->SetShadowOffset(ShadowOffset);
 	}
@@ -73,7 +77,7 @@ void UTextBlock::SynchronizeProperties()
 
 	FString FontPath = FPaths::GameContentDir() / Font.FontName.ToString();
 
-	if ( !FPaths::FileExists(FontPath) )
+	if( !FPaths::FileExists(FontPath) )
 	{
 		FontPath = FPaths::EngineContentDir() / Font.FontName.ToString();
 	}
@@ -133,6 +137,12 @@ void UTextBlock::PostLoad()
 	}
 }
 
+const FTextBlockStyle* UTextBlock::GetDefaultStyle()
+{
+	static const STextBlock::FArguments Defaults;
+	return Defaults._TextStyle;
+}
+
 #if WITH_EDITOR
 
 FString UTextBlock::GetLabelMetadata() const
diff --git a/Engine/Source/Runtime/UMG/Private/Components/Throbber.cpp b/Engine/Source/Runtime/UMG/Private/Components/Throbber.cpp
--- a/Engine/Source/Runtime/UMG/Private/Components/Throbber.cpp
+++ b/Engine/Source/Runtime/UMG/Private/Components/Throbber.cpp
@@ -21,9 +21,9 @@ UThrobber::UThrobber(const FPostConstructInitializeProperties& PCIP)
 	bAnimateOpacity = (DefaultArgs._Animate & SThrobber::Opacity) != 0;
 }
 
-void UThrobber::ReleaseNativeWidget()
+void UThrobber::ReleaseSlateResources(bool bReleaseChildren)
 {
-	Super::ReleaseNativeWidget();
+	Super::ReleaseSlateResources(bReleaseChildren);
 
 	MyThrobber.Reset();
 }
diff --git a/Engine/Source/Runtime/UMG/Private/Components/UniformGridPanel.cpp b/Engine/Source/Runtime/UMG/Private/Components/UniformGridPanel.cpp
--- a/Engine/Source/Runtime/UMG/Private/Components/UniformGridPanel.cpp
+++ b/Engine/Source/Runtime/UMG/Private/Components/UniformGridPanel.cpp
@@ -16,9 +16,9 @@ UUniformGridPanel::UUniformGridPanel(const FPostConstructInitializeProperties& P
 	Visiblity = UWidget::ConvertRuntimeToSerializedVisiblity(Defaults._Visibility.Get());
 }
 
-void UUniformGridPanel::ReleaseNativeWidget()
+void UUniformGridPanel::ReleaseSlateResources(bool bReleaseChildren)
 {
-	Super::ReleaseNativeWidget();
+	Super::ReleaseSlateResources(bReleaseChildren);
 
 	MyUniformGridPanel.Reset();
 }
diff --git a/Engine/Source/Runtime/UMG/Private/Components/UniformGridSlot.cpp b/Engine/Source/Runtime/UMG/Private/Components/UniformGridSlot.cpp
--- a/Engine/Source/Runtime/UMG/Private/Components/UniformGridSlot.cpp
+++ b/Engine/Source/Runtime/UMG/Private/Components/UniformGridSlot.cpp
@@ -13,9 +13,9 @@ UUniformGridSlot::UUniformGridSlot(const FPostConstructInitializeProperties& PCI
 	VerticalAlignment = VAlign_Top;
 }
 
-void UUniformGridSlot::ReleaseNativeWidget()
+void UUniformGridSlot::ReleaseSlateResources(bool bReleaseChildren)
 {
-	Super::ReleaseNativeWidget();
+	Super::ReleaseSlateResources(bReleaseChildren);
 
 	Slot = NULL;
 }
diff --git a/Engine/Source/Runtime/UMG/Private/Components/VerticalBox.cpp b/Engine/Source/Runtime/UMG/Private/Components/VerticalBox.cpp
--- a/Engine/Source/Runtime/UMG/Private/Components/VerticalBox.cpp
+++ b/Engine/Source/Runtime/UMG/Private/Components/VerticalBox.cpp
@@ -16,9 +16,9 @@ UVerticalBox::UVerticalBox(const FPostConstructInitializeProperties& PCIP)
 	Visiblity = UWidget::ConvertRuntimeToSerializedVisiblity( Defaults._Visibility.Get() );
 }
 
-void UVerticalBox::ReleaseNativeWidget()
+void UVerticalBox::ReleaseSlateResources(bool bReleaseChildren)
 {
-	Super::ReleaseNativeWidget();
+	Super::ReleaseSlateResources(bReleaseChildren);
 
 	MyVerticalBox.Reset();
 }
diff --git a/Engine/Source/Runtime/UMG/Private/Components/VerticalBoxSlot.cpp b/Engine/Source/Runtime/UMG/Private/Components/VerticalBoxSlot.cpp
--- a/Engine/Source/Runtime/UMG/Private/Components/VerticalBoxSlot.cpp
+++ b/Engine/Source/Runtime/UMG/Private/Components/VerticalBoxSlot.cpp
@@ -14,9 +14,9 @@ UVerticalBoxSlot::UVerticalBoxSlot(const FPostConstructInitializeProperties& PCI
 	Size = FSlateChildSize(ESlateSizeRule::Automatic);
 }
 
-void UVerticalBoxSlot::ReleaseNativeWidget()
+void UVerticalBoxSlot::ReleaseSlateResources(bool bReleaseChildren)
 {
-	Super::ReleaseNativeWidget();
+	Super::ReleaseSlateResources(bReleaseChildren);
 
 	Slot = NULL;
 }
diff --git a/Engine/Source/Runtime/UMG/Private/Components/Viewport.cpp b/Engine/Source/Runtime/UMG/Private/Components/Viewport.cpp
--- a/Engine/Source/Runtime/UMG/Private/Components/Viewport.cpp
+++ b/Engine/Source/Runtime/UMG/Private/Components/Viewport.cpp
@@ -499,9 +499,9 @@ UViewport::UViewport(const FPostConstructInitializeProperties& PCIP)
 	////ParentArgs.RenderDirectlyToWindow(true);
 }
 
-void UViewport::ReleaseNativeWidget()
+void UViewport::ReleaseSlateResources(bool bReleaseChildren)
 {
-	Super::ReleaseNativeWidget();
+	Super::ReleaseSlateResources(bReleaseChildren);
 
 	ViewportWidget.Reset();
 }
diff --git a/Engine/Source/Runtime/UMG/Private/Components/Visual.cpp b/Engine/Source/Runtime/UMG/Private/Components/Visual.cpp
--- a/Engine/Source/Runtime/UMG/Private/Components/Visual.cpp
+++ b/Engine/Source/Runtime/UMG/Private/Components/Visual.cpp
@@ -10,6 +10,14 @@ UVisual::UVisual(const FPostConstructInitializeProperties& PCIP)
 {
 }
 
-void UVisual::ReleaseNativeWidget()
+void UVisual::ReleaseSlateResources(bool bReleaseChildren)
 {
 }
+
+void UVisual::BeginDestroy()
+{
+	Super::BeginDestroy();
+
+	const bool bReleaseChildren = false;
+	ReleaseSlateResources(bReleaseChildren);
+}
\ No newline at end of file
diff --git a/Engine/Source/Runtime/UMG/Private/Components/WidgetSwitcher.cpp b/Engine/Source/Runtime/UMG/Private/Components/WidgetSwitcher.cpp
--- a/Engine/Source/Runtime/UMG/Private/Components/WidgetSwitcher.cpp
+++ b/Engine/Source/Runtime/UMG/Private/Components/WidgetSwitcher.cpp
@@ -16,9 +16,9 @@ UWidgetSwitcher::UWidgetSwitcher(const FPostConstructInitializeProperties& PCIP)
 	Visiblity = UWidget::ConvertRuntimeToSerializedVisiblity(Defaults._Visibility.Get());
 }
 
-void UWidgetSwitcher::ReleaseNativeWidget()
+void UWidgetSwitcher::ReleaseSlateResources(bool bReleaseChildren)
 {
-	Super::ReleaseNativeWidget();
+	Super::ReleaseSlateResources(bReleaseChildren);
 
 	MyWidgetSwitcher.Reset();
 }
diff --git a/Engine/Source/Runtime/UMG/Private/Components/WidgetSwitcherSlot.cpp b/Engine/Source/Runtime/UMG/Private/Components/WidgetSwitcherSlot.cpp
--- a/Engine/Source/Runtime/UMG/Private/Components/WidgetSwitcherSlot.cpp
+++ b/Engine/Source/Runtime/UMG/Private/Components/WidgetSwitcherSlot.cpp
@@ -13,9 +13,9 @@ UWidgetSwitcherSlot::UWidgetSwitcherSlot(const FPostConstructInitializePropertie
 	VerticalAlignment = VAlign_Fill;
 }
 
-void UWidgetSwitcherSlot::ReleaseNativeWidget()
+void UWidgetSwitcherSlot::ReleaseSlateResources(bool bReleaseChildren)
 {
-	Super::ReleaseNativeWidget();
+	Super::ReleaseSlateResources(bReleaseChildren);
 
 	Slot = NULL;
 }
diff --git a/Engine/Source/Runtime/UMG/Private/Slate/SObjectWidget.cpp b/Engine/Source/Runtime/UMG/Private/Slate/SObjectWidget.cpp
--- a/Engine/Source/Runtime/UMG/Private/Slate/SObjectWidget.cpp
+++ b/Engine/Source/Runtime/UMG/Private/Slate/SObjectWidget.cpp
@@ -18,7 +18,12 @@ SObjectWidget::~SObjectWidget(void)
 {
 	if ( UObjectInitialized() && WidgetObject )
 	{
-		WidgetObject->ReleaseNativeWidget();
+		// NOTE: When the SObjectWidget gets released we know that the User Widget has
+		// been removed from the slate widget hierarchy.  When this occurs, we need to 
+		// immediately release all slate widget widgets to deletion from taking n-frames
+		// due to widget nesting.
+		const bool bReleaseChildren = true;
+		WidgetObject->ReleaseSlateResources(bReleaseChildren);
 	}
 }
 
@@ -35,7 +40,6 @@ void SObjectWidget::Tick(const FGeometry& AllottedGeometry, const double InCurre
 	}
 }
 
-
 int32 SObjectWidget::OnPaint(const FPaintArgs& Args, const FGeometry& AllottedGeometry, const FSlateRect& MyClippingRect, FSlateWindowElementList& OutDrawElements, int32 LayerId, const FWidgetStyle& InWidgetStyle, bool bParentEnabled) const
 {
 	int32 MaxLayer = SCompoundWidget::OnPaint(Args, AllottedGeometry, MyClippingRect, OutDrawElements, LayerId, InWidgetStyle, bParentEnabled);
diff --git a/Engine/Source/Runtime/UMG/Private/UserWidget.cpp b/Engine/Source/Runtime/UMG/Private/UserWidget.cpp
--- a/Engine/Source/Runtime/UMG/Private/UserWidget.cpp
+++ b/Engine/Source/Runtime/UMG/Private/UserWidget.cpp
@@ -11,113 +11,6 @@
 
 #define LOCTEXT_NAMESPACE "UMG"
 
-/**
- * This class holds onto the widget when it's placed into the viewport.
- */
-class SViewportWidgetHost : public SCompoundWidget
-{
-	SLATE_BEGIN_ARGS(SViewportWidgetHost)
-	{
-		_Visibility = EVisibility::SelfHitTestInvisible;
-	}
-
-		SLATE_DEFAULT_SLOT(FArguments, Content)
-
-	SLATE_END_ARGS()
-
-	void Construct(const FArguments& InArgs, TSharedPtr<SWidget> InUserWidget, bool bInModal)
-	{
-		bModal = bInModal;
-		UserWidget = InUserWidget;
-
-		SetVisibility(bModal ? EVisibility::Visible : EVisibility::SelfHitTestInvisible);
-
-		ChildSlot
-		[
-			InArgs._Content.Widget
-		];
-	}
-
-	virtual bool SupportsKeyboardFocus() const override
-	{
-		return bModal ? true : false;
-	}
-
-	FReply OnKeyboardFocusReceived(const FGeometry& MyGeometry, const FKeyboardFocusEvent& InKeyboardFocusEvent) override
-	{
-		// if we support focus, release the focus captures and lock the focus to this widget 
-		if ( bModal )
-		{
-			return FReply::Handled();// .ReleaseMouseCapture().ReleaseJoystickCapture();// .LockMouseToWidget(SharedThis(this));
-		}
-		else
-		{
-			return FReply::Unhandled();
-		}
-	}
-
-	virtual FCursorReply OnCursorQuery(const FGeometry& MyGeometry, const FPointerEvent& CursorEvent) const override
-	{
-		// If we support focus, show the default mouse cursor 
-		if ( SupportsKeyboardFocus() )
-		{
-			return FCursorReply::Cursor(EMouseCursor::Default);
-		}
-		else
-		{
-			return SCompoundWidget::OnCursorQuery(MyGeometry, CursorEvent);
-		}
-	}
-
-	virtual FReply OnKeyChar(const FGeometry& MyGeometry, const FCharacterEvent& InCharacterEvent) override
-	{
-		return bModal ? FReply::Handled() : FReply::Unhandled();
-	}
-
-	virtual FReply OnKeyDown(const FGeometry& MyGeometry, const FKeyboardEvent& InKeyboardEvent) override
-	{
-		return bModal ? FReply::Handled() : FReply::Unhandled();
-	}
-	
-	virtual FReply OnKeyUp(const FGeometry& MyGeometry, const FKeyboardEvent& InKeyboardEvent) override
-	{
-		return bModal ? FReply::Handled() : FReply::Unhandled();
-	}
-
-	virtual FReply OnMouseButtonDown(const FGeometry& MyGeometry, const FPointerEvent& MouseEvent) override
-	{
-		return bModal ? FReply::Handled().CaptureMouse(AsShared()).CaptureJoystick(AsShared()) : FReply::Unhandled();
-	}
-
-	virtual FReply OnMouseButtonUp(const FGeometry& MyGeometry, const FPointerEvent& MouseEvent) override
-	{
-		return bModal ? FReply::Handled().ReleaseMouseCapture().ReleaseJoystickCapture() : FReply::Unhandled();
-	}
-
-	virtual FReply OnMouseButtonDoubleClick(const FGeometry& MyGeometry, const FPointerEvent& MouseEvent) override
-	{
-		return bModal ? FReply::Handled() : FReply::Unhandled();
-	}
-
-	virtual FReply OnMouseMove(const FGeometry& MyGeometry, const FPointerEvent& MouseEvent) override
-	{
-		return bModal ? FReply::Handled() : FReply::Unhandled();
-	}
-
-	virtual FReply OnMouseWheel(const FGeometry& MyGeometry, const FPointerEvent& MouseEvent) override
-	{
-		return bModal ? FReply::Handled() : FReply::Unhandled();
-	}
-
-	bool IsModal() const { return bModal; }
-
-protected:
-	bool bModal;
-
-	TSharedPtr<SWidget> UserWidget;
-};
-
-
 /////////////////////////////////////////////////////
 // UUserWidget
 
@@ -159,14 +52,14 @@ void UUserWidget::PostDuplicate(bool bDuplicateForPIE)
 	Initialize();
 }
 
-void UUserWidget::ReleaseNativeWidget()
+void UUserWidget::ReleaseSlateResources(bool bReleaseChildren)
 {
-	Super::ReleaseNativeWidget();
+	Super::ReleaseSlateResources(bReleaseChildren);
 
 	UWidget* RootWidget = GetRootWidgetComponent();
 	if ( RootWidget )
 	{
-		RootWidget->ReleaseNativeWidget();
+		RootWidget->ReleaseSlateResources(bReleaseChildren);
 	}
 }
 
@@ -535,19 +428,19 @@ UWidget* UUserWidget::GetRootWidgetComponent()
 	return NULL;
 }
 
-void UUserWidget::AddToViewport(bool bModal)
+void UUserWidget::AddToViewport()
 {
 	if ( !FullScreenWidget.IsValid() )
 	{
 		TSharedPtr<SWidget> OutUserSlateWidget;
 		TSharedRef<SWidget> RootWidget = MakeViewportWidget(OutUserSlateWidget);
 
-		TSharedRef<SViewportWidgetHost> WidgetHost = SNew(SViewportWidgetHost, OutUserSlateWidget, bModal)
-			[
-				RootWidget
-			];
+		//TSharedRef<SViewportWidgetHost> WidgetHost = SNew(SViewportWidgetHost, OutUserSlateWidget)
+		//	[
+		//		RootWidget
+		//	];
 
-		FullScreenWidget = WidgetHost;
+		FullScreenWidget = OutUserSlateWidget;
 
 		// If this is a game world add the widget to the current worlds viewport.
 		UWorld* World = GetWorld();
@@ -555,25 +448,25 @@ void UUserWidget::AddToViewport(bool bModal)
 		{
 			if ( UGameViewportClient* ViewportClient = World->GetGameViewport() )
 			{
-				ViewportClient->AddViewportWidgetContent(WidgetHost);
-
-				//TODO UMG this isn't what should manage focus, a higher level window controller, probably the viewport needs to understand
-				// the Widget stack, and the dialog stack.
-				if ( bModal )
-				{
-					if ( FSceneViewport* SceneViewport = ViewportClient->GetGameViewport() )
-					{
-						TWeakPtr<SViewport> GameViewportWidget = SceneViewport->GetViewportWidget();
-
-						if ( GameViewportWidget.IsValid() )
-						{
-							GameViewportWidget.Pin()->SetWidgetToFocusOnActivate(OutUserSlateWidget);
-
-							FSlateApplication::Get().SetFocusToGameViewport();
-							FSlateApplication::Get().SetJoystickCaptorToGameViewport();
-						}
-					}
-				}
+				ViewportClient->AddViewportWidgetContent(OutUserSlateWidget.ToSharedRef());
+
+				////TODO UMG this isn't what should manage focus, a higher level window controller, probably the viewport needs to understand
+				//// the Widget stack, and the dialog stack.
+				//if ( bModal )
+				//{
+				//	if ( FSceneViewport* SceneViewport = ViewportClient->GetGameViewport() )
+				//	{
+				//		TWeakPtr<SViewport> GameViewportWidget = SceneViewport->GetViewportWidget();
+
+				//		if ( GameViewportWidget.IsValid() )
+				//		{
+				//			GameViewportWidget.Pin()->SetWidgetToFocusOnActivate(OutUserSlateWidget);
+
+				//			FSlateApplication::Get().SetFocusToGameViewport();
+				//			FSlateApplication::Get().SetJoystickCaptorToGameViewport();
+				//		}
+				//	}
+				//}
 			}
 		}
 	}
@@ -583,7 +476,7 @@ void UUserWidget::RemoveFromViewport()
 {
 	if ( FullScreenWidget.IsValid() )
 	{
-		TSharedPtr<SViewportWidgetHost> WidgetHost = FullScreenWidget.Pin();
+		TSharedPtr<SWidget> WidgetHost = FullScreenWidget.Pin();
 
 		// If this is a game world add the widget to the current worlds viewport.
 		UWorld* World = GetWorld();
@@ -593,24 +486,24 @@ void UUserWidget::RemoveFromViewport()
 			{
 				ViewportClient->RemoveViewportWidgetContent(WidgetHost.ToSharedRef());
 
-				if ( WidgetHost->IsModal() )
-				{
-					if ( FSceneViewport* SceneViewport = ViewportClient->GetGameViewport() )
-					{
-						TWeakPtr<SViewport> GameViewportWidget = SceneViewport->GetViewportWidget();
-
-						if ( GameViewportWidget.IsValid() )
-						{
-							//TODO UMG this isn't what should manage focus, a higher level window controller, probably the viewport needs to understand
-							// the Widget stack, and the dialog stack.
-							GameViewportWidget.Pin()->ClearWidgetToFocusOnActivate();
-							FSlateApplication::Get().ClearKeyboardFocus(EKeyboardFocusCause::SetDirectly);
-
-							FSlateApplication::Get().SetFocusToGameViewport();
-							FSlateApplication::Get().SetJoystickCaptorToGameViewport();
-						}
-					}
-				}
+				//if ( WidgetHost->IsModal() )
+				//{
+				//	if ( FSceneViewport* SceneViewport = ViewportClient->GetGameViewport() )
+				//	{
+				//		TWeakPtr<SViewport> GameViewportWidget = SceneViewport->GetViewportWidget();
+
+				//		if ( GameViewportWidget.IsValid() )
+				//		{
+				//			//TODO UMG this isn't what should manage focus, a higher level window controller, probably the viewport needs to understand
+				//			// the Widget stack, and the dialog stack.
+				//			GameViewportWidget.Pin()->ClearWidgetToFocusOnActivate();
+				//			FSlateApplication::Get().ClearKeyboardFocus(EKeyboardFocusCause::SetDirectly);
+
+				//			FSlateApplication::Get().SetFocusToGameViewport();
+				//			FSlateApplication::Get().SetJoystickCaptorToGameViewport();
+				//		}
+				//	}
+				//}
 			}
 		}
 	}
diff --git a/Engine/Source/Runtime/UMG/Public/Blueprint/UserWidget.h b/Engine/Source/Runtime/UMG/Public/Blueprint/UserWidget.h
--- a/Engine/Source/Runtime/UMG/Public/Blueprint/UserWidget.h
+++ b/Engine/Source/Runtime/UMG/Public/Blueprint/UserWidget.h
@@ -94,7 +94,7 @@ public:
 	void Initialize();
 
 	//UVisual interface
-	virtual void ReleaseNativeWidget() override;
+	virtual void ReleaseSlateResources(bool bReleaseChildren) override;
 	// End of UVisual interface
 
 	/** Sets that this widget is being designed sets it on all children as well. */
@@ -105,7 +105,7 @@ public:
 	 * @param bModal If this dialog should steal keyboard/mouse focus and consume all input. Great for a fullscreen menu. Terrible for HUDs.
 	 */
 	UFUNCTION(BlueprintCallable, Category="User Interface|Viewport")
-	void AddToViewport(bool bModal = false);
+	void AddToViewport();
 
 	/**
 	 * Removes the widget from the viewport.
@@ -323,7 +323,7 @@ private:
 	FVector2D ViewportAlignment;
 	int32 ViewportZOrder;
 
-	TWeakPtr<class SViewportWidgetHost> FullScreenWidget;
+	TWeakPtr<SWidget> FullScreenWidget;
 
 	FLocalPlayerContext PlayerContext;
 
diff --git a/Engine/Source/Runtime/UMG/Public/Components/Border.h b/Engine/Source/Runtime/UMG/Public/Components/Border.h
--- a/Engine/Source/Runtime/UMG/Public/Components/Border.h
+++ b/Engine/Source/Runtime/UMG/Public/Components/Border.h
@@ -111,7 +111,7 @@ public:
 	// End of UWidget interface
 
 	// UVisual interface
-	virtual void ReleaseNativeWidget() override;
+	virtual void ReleaseSlateResources(bool bReleaseChildren) override;
 	// End of UVisual interface
 
 	// Begin UObject interface
diff --git a/Engine/Source/Runtime/UMG/Public/Components/BorderSlot.h b/Engine/Source/Runtime/UMG/Public/Components/BorderSlot.h
--- a/Engine/Source/Runtime/UMG/Public/Components/BorderSlot.h
+++ b/Engine/Source/Runtime/UMG/Public/Components/BorderSlot.h
@@ -47,7 +47,7 @@ public:
 	/** Builds the underlying slot for the slate border. */
 	void BuildSlot(TSharedRef<SBorder> InBorder);
 
-	virtual void ReleaseNativeWidget() override;
+	virtual void ReleaseSlateResources(bool bReleaseChildren) override;
 
 public:
 
diff --git a/Engine/Source/Runtime/UMG/Public/Components/Button.h b/Engine/Source/Runtime/UMG/Public/Components/Button.h
--- a/Engine/Source/Runtime/UMG/Public/Components/Button.h
+++ b/Engine/Source/Runtime/UMG/Public/Components/Button.h
@@ -73,7 +73,7 @@ public:
 	// End of UWidget interface
 
 	// UVisual interface
-	virtual void ReleaseNativeWidget() override;
+	virtual void ReleaseSlateResources(bool bReleaseChildren) override;
 	// End of UVisual interface
 
 	// Begin UObject interface
diff --git a/Engine/Source/Runtime/UMG/Public/Components/ButtonSlot.h b/Engine/Source/Runtime/UMG/Public/Components/ButtonSlot.h
--- a/Engine/Source/Runtime/UMG/Public/Components/ButtonSlot.h
+++ b/Engine/Source/Runtime/UMG/Public/Components/ButtonSlot.h
@@ -42,7 +42,7 @@ public:
 	/** Builds the underlying slot for the slate button. */
 	void BuildSlot(TSharedRef<SButton> InButton);
 
-	virtual void ReleaseNativeWidget() override;
+	virtual void ReleaseSlateResources(bool bReleaseChildren) override;
 
 private:
 
diff --git a/Engine/Source/Runtime/UMG/Public/Components/CanvasPanel.h b/Engine/Source/Runtime/UMG/Public/Components/CanvasPanel.h
--- a/Engine/Source/Runtime/UMG/Public/Components/CanvasPanel.h
+++ b/Engine/Source/Runtime/UMG/Public/Components/CanvasPanel.h
@@ -20,7 +20,7 @@ public:
 	/** Computes the geometry for a particular slot based on the current geometry of the canvas. */
 	bool GetGeometryForSlot(UCanvasPanelSlot* Slot, FGeometry& ArrangedGeometry) const;
 
-	void ReleaseNativeWidget() override;
+	void ReleaseSlateResources(bool bReleaseChildren) override;
 
 #if WITH_EDITOR
 	// UWidget interface
diff --git a/Engine/Source/Runtime/UMG/Public/Components/CanvasPanelSlot.h b/Engine/Source/Runtime/UMG/Public/Components/CanvasPanelSlot.h
--- a/Engine/Source/Runtime/UMG/Public/Components/CanvasPanelSlot.h
+++ b/Engine/Source/Runtime/UMG/Public/Components/CanvasPanelSlot.h
@@ -86,7 +86,7 @@ class UMG_API UCanvasPanelSlot : public UPanelSlot
 	virtual void SynchronizeProperties() override;
 	// End of UPanelSlot interface
 
-	virtual void ReleaseNativeWidget() override;
+	virtual void ReleaseSlateResources(bool bReleaseChildren) override;
 
 #if WITH_EDITOR
 	// UObject interface
diff --git a/Engine/Source/Runtime/UMG/Public/Components/CheckBox.h b/Engine/Source/Runtime/UMG/Public/Components/CheckBox.h
--- a/Engine/Source/Runtime/UMG/Public/Components/CheckBox.h
+++ b/Engine/Source/Runtime/UMG/Public/Components/CheckBox.h
@@ -112,7 +112,7 @@ public:
 	// End of UWidget interface
 
 	// UVisual interface
-	virtual void ReleaseNativeWidget() override;
+	virtual void ReleaseSlateResources(bool bReleaseChildren) override;
 	// End of UVisual interface
 
 	// Begin UObject interface
diff --git a/Engine/Source/Runtime/UMG/Public/Components/CircularThrobber.h b/Engine/Source/Runtime/UMG/Public/Components/CircularThrobber.h
--- a/Engine/Source/Runtime/UMG/Public/Components/CircularThrobber.h
+++ b/Engine/Source/Runtime/UMG/Public/Components/CircularThrobber.h
@@ -31,7 +31,7 @@ public:
 	// End of UWidget interface
 
 	// UVisual interface
-	virtual void ReleaseNativeWidget() override;
+	virtual void ReleaseSlateResources(bool bReleaseChildren) override;
 	// End of UVisual interface
 
 	// Begin UObject interface
diff --git a/Engine/Source/Runtime/UMG/Public/Components/ComboBox.h b/Engine/Source/Runtime/UMG/Public/Components/ComboBox.h
--- a/Engine/Source/Runtime/UMG/Public/Components/ComboBox.h
+++ b/Engine/Source/Runtime/UMG/Public/Components/ComboBox.h
@@ -22,7 +22,7 @@ public:
 	UPROPERTY(EditDefaultsOnly, Category=Events)
 	FGenerateWidgetForObject OnGenerateWidgetEvent;
 
-	virtual void ReleaseNativeWidget() override;
+	virtual void ReleaseSlateResources(bool bReleaseChildren) override;
 
 #if WITH_EDITOR
 	virtual const FSlateBrush* GetEditorIcon() override;
diff --git a/Engine/Source/Runtime/UMG/Public/Components/ComboBoxString.h b/Engine/Source/Runtime/UMG/Public/Components/ComboBoxString.h
--- a/Engine/Source/Runtime/UMG/Public/Components/ComboBoxString.h
+++ b/Engine/Source/Runtime/UMG/Public/Components/ComboBoxString.h
@@ -80,7 +80,7 @@ public:
 	UFUNCTION(BlueprintCallable, Category = "ComboBox")
 	FString GetSelectedOption() const;
 
-	virtual void ReleaseNativeWidget() override;
+	virtual void ReleaseSlateResources(bool bReleaseChildren) override;
 
 #if WITH_EDITOR
 	virtual const FSlateBrush* GetEditorIcon() override;
diff --git a/Engine/Source/Runtime/UMG/Public/Components/EditableText.h b/Engine/Source/Runtime/UMG/Public/Components/EditableText.h
--- a/Engine/Source/Runtime/UMG/Public/Components/EditableText.h
+++ b/Engine/Source/Runtime/UMG/Public/Components/EditableText.h
@@ -140,7 +140,7 @@ public:
 	// End of UWidget interface
 
 	// UVisual interface
-	virtual void ReleaseNativeWidget() override;
+	virtual void ReleaseSlateResources(bool bReleaseChildren) override;
 	// End of UVisual interface
 
 	// Begin UObject interface
diff --git a/Engine/Source/Runtime/UMG/Public/Components/EditableTextBox.h b/Engine/Source/Runtime/UMG/Public/Components/EditableTextBox.h
--- a/Engine/Source/Runtime/UMG/Public/Components/EditableTextBox.h
+++ b/Engine/Source/Runtime/UMG/Public/Components/EditableTextBox.h
@@ -125,7 +125,7 @@ public:
 	// End of UWidget interface
 
 	// UVisual interface
-	virtual void ReleaseNativeWidget() override;
+	virtual void ReleaseSlateResources(bool bReleaseChildren) override;
 	// End of UVisual interface
 
 	// Begin UObject interface
diff --git a/Engine/Source/Runtime/UMG/Public/Components/GridPanel.h b/Engine/Source/Runtime/UMG/Public/Components/GridPanel.h
--- a/Engine/Source/Runtime/UMG/Public/Components/GridPanel.h
+++ b/Engine/Source/Runtime/UMG/Public/Components/GridPanel.h
@@ -32,7 +32,7 @@ public:
 	virtual void SynchronizeProperties() override;
 	// End of UWidget interface
 
-	virtual void ReleaseNativeWidget() override;
+	virtual void ReleaseSlateResources(bool bReleaseChildren) override;
 
 #if WITH_EDITOR
 	virtual const FSlateBrush* GetEditorIcon() override;
diff --git a/Engine/Source/Runtime/UMG/Public/Components/GridSlot.h b/Engine/Source/Runtime/UMG/Public/Components/GridSlot.h
--- a/Engine/Source/Runtime/UMG/Public/Components/GridSlot.h
+++ b/Engine/Source/Runtime/UMG/Public/Components/GridSlot.h
@@ -80,7 +80,7 @@ public:
 	/** Builds the underlying FSlot for the Slate layout panel. */
 	void BuildSlot(TSharedRef<SGridPanel> GridPanel);
 
-	virtual void ReleaseNativeWidget() override;
+	virtual void ReleaseSlateResources(bool bReleaseChildren) override;
 
 private:
 	//TODO UMG Slots should hold weak or shared refs to slots.
diff --git a/Engine/Source/Runtime/UMG/Public/Components/HorizontalBox.h b/Engine/Source/Runtime/UMG/Public/Components/HorizontalBox.h
--- a/Engine/Source/Runtime/UMG/Public/Components/HorizontalBox.h
+++ b/Engine/Source/Runtime/UMG/Public/Components/HorizontalBox.h
@@ -20,7 +20,7 @@ class UMG_API UHorizontalBox : public UPanelWidget
 	// End UWidget interface
 #endif
 
-	virtual void ReleaseNativeWidget() override;
+	virtual void ReleaseSlateResources(bool bReleaseChildren) override;
 
 protected:
 
diff --git a/Engine/Source/Runtime/UMG/Public/Components/HorizontalBoxSlot.h b/Engine/Source/Runtime/UMG/Public/Components/HorizontalBoxSlot.h
--- a/Engine/Source/Runtime/UMG/Public/Components/HorizontalBoxSlot.h
+++ b/Engine/Source/Runtime/UMG/Public/Components/HorizontalBoxSlot.h
@@ -45,7 +45,7 @@ public:
 	virtual void SynchronizeProperties() override;
 	// End of UPanelSlot interface
 
-	virtual void ReleaseNativeWidget() override;
+	virtual void ReleaseSlateResources(bool bReleaseChildren) override;
 
 private:
 	SHorizontalBox::FSlot* Slot;
diff --git a/Engine/Source/Runtime/UMG/Public/Components/Image.h b/Engine/Source/Runtime/UMG/Public/Components/Image.h
--- a/Engine/Source/Runtime/UMG/Public/Components/Image.h
+++ b/Engine/Source/Runtime/UMG/Public/Components/Image.h
@@ -70,7 +70,7 @@ public:
 	// End of UWidget interface
 
 	// UVisual interface
-	virtual void ReleaseNativeWidget() override;
+	virtual void ReleaseSlateResources(bool bReleaseChildren) override;
 	// End of UVisual interface
 
 	// Begin UObject interface
diff --git a/Engine/Source/Runtime/UMG/Public/Components/MenuAnchor.h b/Engine/Source/Runtime/UMG/Public/Components/MenuAnchor.h
--- a/Engine/Source/Runtime/UMG/Public/Components/MenuAnchor.h
+++ b/Engine/Source/Runtime/UMG/Public/Components/MenuAnchor.h
@@ -52,7 +52,7 @@ public:
 	///** @return Whether this menu has open submenus */
 	//bool HasOpenSubMenus() const;
 
-	virtual void ReleaseNativeWidget() override;
+	virtual void ReleaseSlateResources(bool bReleaseChildren) override;
 
 #if WITH_EDITOR
 	virtual const FSlateBrush* GetEditorIcon() override;
diff --git a/Engine/Source/Runtime/UMG/Public/Components/MultiLineEditableText.h b/Engine/Source/Runtime/UMG/Public/Components/MultiLineEditableText.h
--- a/Engine/Source/Runtime/UMG/Public/Components/MultiLineEditableText.h
+++ b/Engine/Source/Runtime/UMG/Public/Components/MultiLineEditableText.h
@@ -54,7 +54,7 @@ public:
 	// End of UWidget interface
 
 	// UVisual interface
-	virtual void ReleaseNativeWidget() override;
+	virtual void ReleaseSlateResources(bool bReleaseChildren) override;
 	// End of UVisual interface
 
 #if WITH_EDITOR
diff --git a/Engine/Source/Runtime/UMG/Public/Components/MultiLineEditableTextBox.h b/Engine/Source/Runtime/UMG/Public/Components/MultiLineEditableTextBox.h
--- a/Engine/Source/Runtime/UMG/Public/Components/MultiLineEditableTextBox.h
+++ b/Engine/Source/Runtime/UMG/Public/Components/MultiLineEditableTextBox.h
@@ -79,7 +79,7 @@ public:
 	// End of UWidget interface
 
 	// UVisual interface
-	virtual void ReleaseNativeWidget() override;
+	virtual void ReleaseSlateResources(bool bReleaseChildren) override;
 	// End of UVisual interface
 
 	// Begin UObject interface
diff --git a/Engine/Source/Runtime/UMG/Public/Components/NativeWidgetHost.h b/Engine/Source/Runtime/UMG/Public/Components/NativeWidgetHost.h
--- a/Engine/Source/Runtime/UMG/Public/Components/NativeWidgetHost.h
+++ b/Engine/Source/Runtime/UMG/Public/Components/NativeWidgetHost.h
@@ -15,7 +15,7 @@ class UMG_API UNativeWidgetHost : public UWidget
 
 	void SetContent(TSharedRef<SWidget> InContent);
 
-	virtual void ReleaseNativeWidget() override;
+	virtual void ReleaseSlateResources(bool bReleaseChildren) override;
 
 #if WITH_EDITOR
 	virtual const FSlateBrush* GetEditorIcon() override;
diff --git a/Engine/Source/Runtime/UMG/Public/Components/Overlay.h b/Engine/Source/Runtime/UMG/Public/Components/Overlay.h
--- a/Engine/Source/Runtime/UMG/Public/Components/Overlay.h
+++ b/Engine/Source/Runtime/UMG/Public/Components/Overlay.h
@@ -9,7 +9,7 @@ class UMG_API UOverlay : public UPanelWidget
 {
 	GENERATED_UCLASS_BODY()
 
-	virtual void ReleaseNativeWidget() override;
+	virtual void ReleaseSlateResources(bool bReleaseChildren) override;
 
 	/**  */
 	UFUNCTION(BlueprintCallable, Category="Widget")
diff --git a/Engine/Source/Runtime/UMG/Public/Components/OverlaySlot.h b/Engine/Source/Runtime/UMG/Public/Components/OverlaySlot.h
--- a/Engine/Source/Runtime/UMG/Public/Components/OverlaySlot.h
+++ b/Engine/Source/Runtime/UMG/Public/Components/OverlaySlot.h
@@ -48,7 +48,7 @@ public:
 	/** Builds the underlying FSlot for the Slate layout panel. */
 	void BuildSlot(TSharedRef<SOverlay> InOverlay);
 
-	virtual void ReleaseNativeWidget() override;
+	virtual void ReleaseSlateResources(bool bReleaseChildren) override;
 
 private:
 	//TODO UMG Slots should hold weak or shared refs to slots.
diff --git a/Engine/Source/Runtime/UMG/Public/Components/PanelSlot.h b/Engine/Source/Runtime/UMG/Public/Components/PanelSlot.h
--- a/Engine/Source/Runtime/UMG/Public/Components/PanelSlot.h
+++ b/Engine/Source/Runtime/UMG/Public/Components/PanelSlot.h
@@ -32,7 +32,7 @@ public:
 
 	virtual bool CanMove() const;
 
-	virtual void ReleaseNativeWidget() override;
+	virtual void ReleaseSlateResources(bool bReleaseChildren) override;
 
 	/** Applies all properties to the live slot if possible. */
 	virtual void SynchronizeProperties()
diff --git a/Engine/Source/Runtime/UMG/Public/Components/PanelWidget.h b/Engine/Source/Runtime/UMG/Public/Components/PanelWidget.h
--- a/Engine/Source/Runtime/UMG/Public/Components/PanelWidget.h
+++ b/Engine/Source/Runtime/UMG/Public/Components/PanelWidget.h
@@ -70,7 +70,7 @@ public:
 	/** Sets that this widget is being designed sets it on all children as well. */
 	virtual void SetIsDesignTime(bool bInDesignTime) override;
 
-	virtual void ReleaseNativeWidget() override;
+	virtual void ReleaseSlateResources(bool bReleaseChildren) override;
 
 #if WITH_EDITOR
 	virtual bool LockToPanelOnDrag() const
diff --git a/Engine/Source/Runtime/UMG/Public/Components/ProgressBar.h b/Engine/Source/Runtime/UMG/Public/Components/ProgressBar.h
--- a/Engine/Source/Runtime/UMG/Public/Components/ProgressBar.h
+++ b/Engine/Source/Runtime/UMG/Public/Components/ProgressBar.h
@@ -68,7 +68,7 @@ public:
 	// End of UWidget interface
 
 	// UVisual interface
-	virtual void ReleaseNativeWidget() override;
+	virtual void ReleaseSlateResources(bool bReleaseChildren) override;
 	// End of UVisual interface
 
 	// Begin UObject interface
diff --git a/Engine/Source/Runtime/UMG/Public/Components/ScaleBox.h b/Engine/Source/Runtime/UMG/Public/Components/ScaleBox.h
--- a/Engine/Source/Runtime/UMG/Public/Components/ScaleBox.h
+++ b/Engine/Source/Runtime/UMG/Public/Components/ScaleBox.h
@@ -30,7 +30,7 @@ public:
 	virtual void SynchronizeProperties() override;
 	// End of UWidget interface
 
-	virtual void ReleaseNativeWidget() override;
+	virtual void ReleaseSlateResources(bool bReleaseChildren) override;
 
 #if WITH_EDITOR
 	virtual const FSlateBrush* GetEditorIcon() override;
diff --git a/Engine/Source/Runtime/UMG/Public/Components/ScaleBoxSlot.h b/Engine/Source/Runtime/UMG/Public/Components/ScaleBoxSlot.h
--- a/Engine/Source/Runtime/UMG/Public/Components/ScaleBoxSlot.h
+++ b/Engine/Source/Runtime/UMG/Public/Components/ScaleBoxSlot.h
@@ -44,7 +44,7 @@ public:
 	/** Builds the underlying slot for the slate ScaleBox. */
 	void BuildSlot(TSharedRef<SScaleBox> InScaleBox);
 
-	virtual void ReleaseNativeWidget() override;
+	virtual void ReleaseSlateResources(bool bReleaseChildren) override;
 
 private:
 
diff --git a/Engine/Source/Runtime/UMG/Public/Components/ScrollBar.h b/Engine/Source/Runtime/UMG/Public/Components/ScrollBar.h
--- a/Engine/Source/Runtime/UMG/Public/Components/ScrollBar.h
+++ b/Engine/Source/Runtime/UMG/Public/Components/ScrollBar.h
@@ -64,7 +64,7 @@ public:
 	// End of UWidget interface
 
 	// UVisual interface
-	virtual void ReleaseNativeWidget() override;
+	virtual void ReleaseSlateResources(bool bReleaseChildren) override;
 	// End of UVisual interface
 
 	// Begin UObject interface
diff --git a/Engine/Source/Runtime/UMG/Public/Components/ScrollBox.h b/Engine/Source/Runtime/UMG/Public/Components/ScrollBox.h
--- a/Engine/Source/Runtime/UMG/Public/Components/ScrollBox.h
+++ b/Engine/Source/Runtime/UMG/Public/Components/ScrollBox.h
@@ -37,7 +37,7 @@ class UMG_API UScrollBox : public UPanelWidget
 	// End of UWidget interface
 
 	// UVisual interface
-	virtual void ReleaseNativeWidget() override;
+	virtual void ReleaseSlateResources(bool bReleaseChildren) override;
 	// End of UVisual interface
 
 	// Begin UObject interface
diff --git a/Engine/Source/Runtime/UMG/Public/Components/ScrollBoxSlot.h b/Engine/Source/Runtime/UMG/Public/Components/ScrollBoxSlot.h
--- a/Engine/Source/Runtime/UMG/Public/Components/ScrollBoxSlot.h
+++ b/Engine/Source/Runtime/UMG/Public/Components/ScrollBoxSlot.h
@@ -39,7 +39,7 @@ public:
 	/** Builds the underlying FSlot for the Slate layout panel. */
 	void BuildSlot(TSharedRef<SScrollBox> ScrollBox);
 
-	virtual void ReleaseNativeWidget() override;
+	virtual void ReleaseSlateResources(bool bReleaseChildren) override;
 
 private:
 	//TODO UMG Slots should hold weak or shared refs to slots.
diff --git a/Engine/Source/Runtime/UMG/Public/Components/SizeBox.h b/Engine/Source/Runtime/UMG/Public/Components/SizeBox.h
--- a/Engine/Source/Runtime/UMG/Public/Components/SizeBox.h
+++ b/Engine/Source/Runtime/UMG/Public/Components/SizeBox.h
@@ -66,7 +66,7 @@ public:
 	// End of UWidget interface
 
 	// UVisual interface
-	virtual void ReleaseNativeWidget() override;
+	virtual void ReleaseSlateResources(bool bReleaseChildren) override;
 	// End of UVisual interface
 
 #if WITH_EDITOR
diff --git a/Engine/Source/Runtime/UMG/Public/Components/SizeBoxSlot.h b/Engine/Source/Runtime/UMG/Public/Components/SizeBoxSlot.h
--- a/Engine/Source/Runtime/UMG/Public/Components/SizeBoxSlot.h
+++ b/Engine/Source/Runtime/UMG/Public/Components/SizeBoxSlot.h
@@ -32,7 +32,7 @@ public:
 	/** Builds the underlying slot for the slate SizeBox. */
 	void BuildSlot(TSharedRef<SBox> InSizeBox);
 
-	virtual void ReleaseNativeWidget() override;
+	virtual void ReleaseSlateResources(bool bReleaseChildren) override;
 
 protected:
 	/** The padding area between the slot and the content it contains. */
diff --git a/Engine/Source/Runtime/UMG/Public/Components/Spacer.h b/Engine/Source/Runtime/UMG/Public/Components/Spacer.h
--- a/Engine/Source/Runtime/UMG/Public/Components/Spacer.h
+++ b/Engine/Source/Runtime/UMG/Public/Components/Spacer.h
@@ -27,7 +27,7 @@ public:
 	// End of UWidget interface
 
 	// UVisual interface
-	virtual void ReleaseNativeWidget() override;
+	virtual void ReleaseSlateResources(bool bReleaseChildren) override;
 	// End of UVisual interface
 
 #if WITH_EDITOR
diff --git a/Engine/Source/Runtime/UMG/Public/Components/SpinBox.h b/Engine/Source/Runtime/UMG/Public/Components/SpinBox.h
--- a/Engine/Source/Runtime/UMG/Public/Components/SpinBox.h
+++ b/Engine/Source/Runtime/UMG/Public/Components/SpinBox.h
@@ -140,7 +140,7 @@ public:
 	// End of UWidget interface
 
 	// UVisual interface
-	virtual void ReleaseNativeWidget() override;
+	virtual void ReleaseSlateResources(bool bReleaseChildren) override;
 	// End of UVisual interface
 
 	// Begin UObject interface
diff --git a/Engine/Source/Runtime/UMG/Public/Components/TextBlock.h b/Engine/Source/Runtime/UMG/Public/Components/TextBlock.h
--- a/Engine/Source/Runtime/UMG/Public/Components/TextBlock.h
+++ b/Engine/Source/Runtime/UMG/Public/Components/TextBlock.h
@@ -123,13 +123,15 @@ public:
 	// End of UWidget interface
 
 	// UVisual interface
-	virtual void ReleaseNativeWidget() override;
+	virtual void ReleaseSlateResources(bool bReleaseChildren) override;
 	// End of UVisual interface
 
 	// Begin UObject interface
 	virtual void PostLoad() override;
 	// End of UObject interface
 
+	static const FTextBlockStyle* GetDefaultStyle();
+
 #if WITH_EDITOR
 	// UWidget interface
 	virtual const FSlateBrush* GetEditorIcon() override;
diff --git a/Engine/Source/Runtime/UMG/Public/Components/Throbber.h b/Engine/Source/Runtime/UMG/Public/Components/Throbber.h
--- a/Engine/Source/Runtime/UMG/Public/Components/Throbber.h
+++ b/Engine/Source/Runtime/UMG/Public/Components/Throbber.h
@@ -35,7 +35,7 @@ public:
 	// End of UWidget interface
 
 	// UVisual interface
-	virtual void ReleaseNativeWidget() override;
+	virtual void ReleaseSlateResources(bool bReleaseChildren) override;
 	// End of UVisual interface
 
 	// Begin UObject interface
diff --git a/Engine/Source/Runtime/UMG/Public/Components/UniformGridPanel.h b/Engine/Source/Runtime/UMG/Public/Components/UniformGridPanel.h
--- a/Engine/Source/Runtime/UMG/Public/Components/UniformGridPanel.h
+++ b/Engine/Source/Runtime/UMG/Public/Components/UniformGridPanel.h
@@ -36,7 +36,7 @@ public:
 	virtual void SynchronizeProperties() override;
 	// End of UWidget interface
 
-	virtual void ReleaseNativeWidget() override;
+	virtual void ReleaseSlateResources(bool bReleaseChildren) override;
 
 #if WITH_EDITOR
 	virtual const FSlateBrush* GetEditorIcon() override;
diff --git a/Engine/Source/Runtime/UMG/Public/Components/UniformGridSlot.h b/Engine/Source/Runtime/UMG/Public/Components/UniformGridSlot.h
--- a/Engine/Source/Runtime/UMG/Public/Components/UniformGridSlot.h
+++ b/Engine/Source/Runtime/UMG/Public/Components/UniformGridSlot.h
@@ -60,7 +60,7 @@ public:
 	/** Builds the underlying FSlot for the Slate layout panel. */
 	void BuildSlot(TSharedRef<SUniformGridPanel> GridPanel);
 
-	virtual void ReleaseNativeWidget() override;
+	virtual void ReleaseSlateResources(bool bReleaseChildren) override;
 
 private:
 	//TODO UMG Slots should hold weak or shared refs to slots.
diff --git a/Engine/Source/Runtime/UMG/Public/Components/VerticalBox.h b/Engine/Source/Runtime/UMG/Public/Components/VerticalBox.h
--- a/Engine/Source/Runtime/UMG/Public/Components/VerticalBox.h
+++ b/Engine/Source/Runtime/UMG/Public/Components/VerticalBox.h
@@ -34,7 +34,7 @@ protected:
 	virtual void OnSlotRemoved(UPanelSlot* Slot) override;
 	// End UPanelWidget
 
-	virtual void ReleaseNativeWidget() override;
+	virtual void ReleaseSlateResources(bool bReleaseChildren) override;
 
 protected:
 
diff --git a/Engine/Source/Runtime/UMG/Public/Components/VerticalBoxSlot.h b/Engine/Source/Runtime/UMG/Public/Components/VerticalBoxSlot.h
--- a/Engine/Source/Runtime/UMG/Public/Components/VerticalBoxSlot.h
+++ b/Engine/Source/Runtime/UMG/Public/Components/VerticalBoxSlot.h
@@ -50,7 +50,7 @@ public:
 	virtual void SynchronizeProperties() override;
 	// End of UPanelSlot interface
 
-	virtual void ReleaseNativeWidget() override;
+	virtual void ReleaseSlateResources(bool bReleaseChildren) override;
 
 	/** Builds the underlying FSlot for the Slate layout panel. */
 	void BuildSlot(TSharedRef<SVerticalBox> InVerticalBox);
diff --git a/Engine/Source/Runtime/UMG/Public/Components/Viewport.h b/Engine/Source/Runtime/UMG/Public/Components/Viewport.h
--- a/Engine/Source/Runtime/UMG/Public/Components/Viewport.h
+++ b/Engine/Source/Runtime/UMG/Public/Components/Viewport.h
@@ -268,7 +268,7 @@ class UMG_API UViewport : public UContentWidget
 	virtual void SynchronizeProperties() override;
 	// End of UWidget interface
 
-	virtual void ReleaseNativeWidget() override;
+	virtual void ReleaseSlateResources(bool bReleaseChildren) override;
 
 #if WITH_EDITOR
 	virtual const FSlateBrush* GetEditorIcon() override;
diff --git a/Engine/Source/Runtime/UMG/Public/Components/Visual.h b/Engine/Source/Runtime/UMG/Public/Components/Visual.h
--- a/Engine/Source/Runtime/UMG/Public/Components/Visual.h
+++ b/Engine/Source/Runtime/UMG/Public/Components/Visual.h
@@ -11,5 +11,9 @@ class UMG_API UVisual : public UObject
 	GENERATED_UCLASS_BODY()
 
 public:
-	virtual void ReleaseNativeWidget();
+	virtual void ReleaseSlateResources(bool bReleaseChildren);
+
+	// Begin UObject interface
+	virtual void BeginDestroy() override;
+	// End UObject interface
 };
diff --git a/Engine/Source/Runtime/UMG/Public/Components/WidgetSwitcher.h b/Engine/Source/Runtime/UMG/Public/Components/WidgetSwitcher.h
--- a/Engine/Source/Runtime/UMG/Public/Components/WidgetSwitcher.h
+++ b/Engine/Source/Runtime/UMG/Public/Components/WidgetSwitcher.h
@@ -38,7 +38,7 @@ public:
 	
 	virtual void SynchronizeProperties() override;
 
-	virtual void ReleaseNativeWidget() override;
+	virtual void ReleaseSlateResources(bool bReleaseChildren) override;
 
 #if WITH_EDITOR
 	virtual const FSlateBrush* GetEditorIcon() override;
diff --git a/Engine/Source/Runtime/UMG/Public/Components/WidgetSwitcherSlot.h b/Engine/Source/Runtime/UMG/Public/Components/WidgetSwitcherSlot.h
--- a/Engine/Source/Runtime/UMG/Public/Components/WidgetSwitcherSlot.h
+++ b/Engine/Source/Runtime/UMG/Public/Components/WidgetSwitcherSlot.h
@@ -43,7 +43,7 @@ public:
 	virtual void SynchronizeProperties() override;
 	// End of UPanelSlot interface
 
-	virtual void ReleaseNativeWidget() override;
+	virtual void ReleaseSlateResources(bool bReleaseChildren) override;
 
 	/** Builds the underlying FSlot for the Slate layout panel. */
 	void BuildSlot(TSharedRef<SWidgetSwitcher> InWidgetSwitcher);
