UMG - Fixing the canvas slot so that the anchor point can be animated.

diff --git a/Engine/Source/Runtime/UMG/Private/Components/CanvasPanelSlot.cpp b/Engine/Source/Runtime/UMG/Private/Components/CanvasPanelSlot.cpp
--- a/Engine/Source/Runtime/UMG/Private/Components/CanvasPanelSlot.cpp
+++ b/Engine/Source/Runtime/UMG/Private/Components/CanvasPanelSlot.cpp
@@ -191,6 +191,24 @@ void UCanvasPanelSlot::SetZOrder(int32 InZOrder)
 	}
 }
 
+void UCanvasPanelSlot::SetMinimum(FVector2D InMinimumAnchors)
+{
+	LayoutData.Anchors.Minimum = InMinimumAnchors;
+	if ( Slot )
+	{
+		Slot->Anchors(LayoutData.Anchors);
+	}
+}
+
+void UCanvasPanelSlot::SetMaximum(FVector2D InMaximumAnchors)
+{
+	LayoutData.Anchors.Maximum = InMaximumAnchors;
+	if ( Slot )
+	{
+		Slot->Anchors(LayoutData.Anchors);
+	}
+}
+
 void UCanvasPanelSlot::SynchronizeProperties()
 {
 	SetOffsets(LayoutData.Offsets);
diff --git a/Engine/Source/Runtime/UMG/Public/Components/CanvasPanelSlot.h b/Engine/Source/Runtime/UMG/Public/Components/CanvasPanelSlot.h
--- a/Engine/Source/Runtime/UMG/Public/Components/CanvasPanelSlot.h
+++ b/Engine/Source/Runtime/UMG/Public/Components/CanvasPanelSlot.h
@@ -82,6 +82,14 @@ class UMG_API UCanvasPanelSlot : public UPanelSlot
 	UFUNCTION(BlueprintCallable, Category="Layout (Canvas Slot)")
 	void SetZOrder(int32 InZOrder);
 
+	/** Sets the anchors on the slot */
+	UFUNCTION()
+	void SetMinimum(FVector2D InMinimumAnchors);
+
+	/** Sets the anchors on the slot */
+	UFUNCTION()
+	void SetMaximum(FVector2D InMaximumAnchors);
+
 	// UPanelSlot interface
 	virtual void SynchronizeProperties() override;
 	// End of UPanelSlot interface
