UMG - Fixing compile time crash issues by clearing the slate preview widgets immediately to prevent slate objects referencing disposed UWidgets, cleaning up some old code in textblock and button.

diff --git a/Engine/Source/Editor/UMGEditor/Private/Designer/SDesignerView.cpp b/Engine/Source/Editor/UMGEditor/Private/Designer/SDesignerView.cpp
--- a/Engine/Source/Editor/UMGEditor/Private/Designer/SDesignerView.cpp
+++ b/Engine/Source/Editor/UMGEditor/Private/Designer/SDesignerView.cpp
@@ -659,8 +659,12 @@ void SDesignerView::OnBlueprintCompiled(UBlueprint* InBlueprint)
 	{
 		CachedWidgetGeometry.Reset();
 
-		const bool bForceUpdate = true;
-		UpdatePreviewWidget(bForceUpdate);
+		PreviewWidget = NULL;
+		PreviewSurface->SetContent(SNullWidget::NullWidget);
+		//PreviewWidget.ReleaseNativeWidget()
+
+		//const bool bForceUpdate = true;
+		//UpdatePreviewWidget(bForceUpdate);
 	}
 }
 
diff --git a/Engine/Source/Runtime/UMG/Private/Components/TextBlock.cpp b/Engine/Source/Runtime/UMG/Private/Components/TextBlock.cpp
--- a/Engine/Source/Runtime/UMG/Private/Components/TextBlock.cpp
+++ b/Engine/Source/Runtime/UMG/Private/Components/TextBlock.cpp
@@ -28,12 +28,7 @@ void UTextBlock::ReleaseSlateResources(bool bReleaseChildren)
 {
 	Super::ReleaseSlateResources(bReleaseChildren);
 
-	if ( MyTextBlock.IsValid() )
-	{
-		MyTextBlock->SetTextStyle(UTextBlock::GetDefaultStyle());
-
-		MyTextBlock.Reset();
-	}
+	MyTextBlock.Reset();
 }
 
 void UTextBlock::SetColorAndOpacity(FSlateColor InColorAndOpacity)
diff --git a/Engine/Source/Runtime/UMG/Public/Components/Button.h b/Engine/Source/Runtime/UMG/Public/Components/Button.h
--- a/Engine/Source/Runtime/UMG/Public/Components/Button.h
+++ b/Engine/Source/Runtime/UMG/Public/Components/Button.h
@@ -24,7 +24,6 @@ public:
 	/** The button style used at runtime */
 	UPROPERTY(EditAnywhere, BlueprintReadWrite, Category="Style", meta=( DisplayName="Style" ))
 	FButtonStyle WidgetStyle;
-	/*ShowOnlyInnerProperties,*/
 	
 	/** The color multiplier for the button content */
 	UPROPERTY(EditDefaultsOnly, Category=Appearance )
@@ -105,7 +104,4 @@ protected:
 protected:
 	/** Cached pointer to the underlying slate button owned by this UWidget */
 	TSharedPtr<SButton> MyButton;
-
-	/** Stores the style set by users dynamically. */
-	TOptional<FButtonStyle> MyStyle;
 };
