Added IsIdentity() to FWidgetTransform.

diff --git a/Engine/Source/Runtime/UMG/Public/Slate/WidgetTransform.h b/Engine/Source/Runtime/UMG/Public/Slate/WidgetTransform.h
--- a/Engine/Source/Runtime/UMG/Public/Slate/WidgetTransform.h
+++ b/Engine/Source/Runtime/UMG/Public/Slate/WidgetTransform.h
@@ -45,6 +45,15 @@ struct FWidgetTransform
 	{
 	}
 
+	bool IsIdentity() const
+	{
+		return
+			Scale == FVector2D::ZeroVector
+			&& Shear == FVector2D::ZeroVector
+			&& Angle == 0.0f
+			&& Translation == FVector2D::ZeroVector;
+	}
+
 	bool operator==( const FWidgetTransform &Other ) const
 	{
 		return Scale == Other.Scale && Shear == Other.Shear && Angle == Other.Angle && Translation == Other.Translation;
