Undoing changes to FText customization - as before, defaults to single line and specify multiline with "meta=(MultiLine="true")"

diff --git a/Engine/Source/Editor/IntroTutorials/Public/EditorTutorial.h b/Engine/Source/Editor/IntroTutorials/Public/EditorTutorial.h
--- a/Engine/Source/Editor/IntroTutorials/Public/EditorTutorial.h
+++ b/Engine/Source/Editor/IntroTutorials/Public/EditorTutorial.h
@@ -56,7 +56,7 @@ struct INTROTUTORIALS_API FTutorialCategory
 	FText Title;
 
 	/** Localized text to use to describe this category */
-	UPROPERTY(EditAnywhere, Category="Content")
+	UPROPERTY(EditAnywhere, Category="Content", meta=(MultiLine=true))
 	FText Description;
 
 	/** Icon for this tutorial, used when presented to the user in the tutorial browser. Only used if there isn't a valid texture to use. */
@@ -92,7 +92,7 @@ struct INTROTUTORIALS_API FTutorialContent
 	FString ExcerptName;
 
 	/** Localized text to use with this content */
-	UPROPERTY(EditAnywhere, Category="Content")
+	UPROPERTY(EditAnywhere, Category="Content", meta=(MultiLine=true))
 	FText Text;
 };
 
diff --git a/Engine/Source/Editor/PropertyEditor/Private/UserInterface/PropertyEditor/SPropertyEditorText.cpp b/Engine/Source/Editor/PropertyEditor/Private/UserInterface/PropertyEditor/SPropertyEditorText.cpp
--- a/Engine/Source/Editor/PropertyEditor/Private/UserInterface/PropertyEditor/SPropertyEditorText.cpp
+++ b/Engine/Source/Editor/PropertyEditor/Private/UserInterface/PropertyEditor/SPropertyEditorText.cpp
@@ -13,7 +13,7 @@ void SPropertyEditorText::Construct( const FArguments& InArgs, const TSharedRef<
 	PropertyEditor = InPropertyEditor;
 
 	const UProperty* Property = InPropertyEditor->GetProperty();
-	bIsMultiLine = !Property->GetBoolMetaData("SingleLine");
+	bIsMultiLine = Property->GetBoolMetaData("MultiLine");
 	if(bIsMultiLine)
 	{
 		ChildSlot
diff --git a/Engine/Source/Runtime/Engine/Classes/Materials/MaterialExpressionCustom.h b/Engine/Source/Runtime/Engine/Classes/Materials/MaterialExpressionCustom.h
--- a/Engine/Source/Runtime/Engine/Classes/Materials/MaterialExpressionCustom.h
+++ b/Engine/Source/Runtime/Engine/Classes/Materials/MaterialExpressionCustom.h
@@ -33,7 +33,7 @@ class UMaterialExpressionCustom : public UMaterialExpression
 {
 	GENERATED_UCLASS_BODY()
 
-	UPROPERTY(EditAnywhere, Category=MaterialExpressionCustom)
+	UPROPERTY(EditAnywhere, Category=MaterialExpressionCustom, meta=(MultiLine=true))
 	FString Code;
 
 	UPROPERTY(EditAnywhere, Category=MaterialExpressionCustom)
diff --git a/Engine/Source/Runtime/UMG/Public/Components/EditableText.h b/Engine/Source/Runtime/UMG/Public/Components/EditableText.h
--- a/Engine/Source/Runtime/UMG/Public/Components/EditableText.h
+++ b/Engine/Source/Runtime/UMG/Public/Components/EditableText.h
@@ -20,7 +20,7 @@ public:
 public:
 
 	/** The text content for this editable text box widget */
-	UPROPERTY(EditDefaultsOnly, Category = Content, meta=(SingleLine=true))
+	UPROPERTY(EditDefaultsOnly, Category=Content)
 	FText Text;
 
 	/** A bindable delegate to allow logic to drive the text of the widget */
@@ -28,7 +28,7 @@ public:
 	FGetText TextDelegate;
 
 	/** Hint text that appears when there is no text in the text box */
-	UPROPERTY(EditDefaultsOnly, Category = Content, meta=(SingleLine=true))
+	UPROPERTY(EditDefaultsOnly, Category=Content)
 	FText HintText;
 
 	/** A bindable delegate to allow logic to drive the hint text of the widget */
diff --git a/Engine/Source/Runtime/UMG/Public/Components/EditableTextBox.h b/Engine/Source/Runtime/UMG/Public/Components/EditableTextBox.h
--- a/Engine/Source/Runtime/UMG/Public/Components/EditableTextBox.h
+++ b/Engine/Source/Runtime/UMG/Public/Components/EditableTextBox.h
@@ -26,7 +26,7 @@ public:
 	USlateWidgetStyleAsset* Style_DEPRECATED;
 
 	/** The text content for this editable text box widget */
-	UPROPERTY(EditDefaultsOnly, Category = Content, meta=(SingleLine=true))
+	UPROPERTY(EditDefaultsOnly, Category=Content)
 	FText Text;
 
 	/** A bindable delegate to allow logic to drive the text of the widget */
@@ -34,7 +34,7 @@ public:
 	FGetText TextDelegate;
 
 	/** Hint text that appears when there is no text in the text box */
-	UPROPERTY(EditDefaultsOnly, Category = Content, meta=(SingleLine=true))
+	UPROPERTY(EditDefaultsOnly, Category=Content)
 	FText HintText;
 
 	/** A bindable delegate to allow logic to drive the hint text of the widget */
diff --git a/Engine/Source/Runtime/UMG/Public/Components/MultiLineEditableText.h b/Engine/Source/Runtime/UMG/Public/Components/MultiLineEditableText.h
--- a/Engine/Source/Runtime/UMG/Public/Components/MultiLineEditableText.h
+++ b/Engine/Source/Runtime/UMG/Public/Components/MultiLineEditableText.h
@@ -22,7 +22,7 @@ public:
 	FTextBlockStyle WidgetStyle;
 
 	/** The text content for this editable text box widget */
-	UPROPERTY(EditDefaultsOnly, Category=Content)
+	UPROPERTY(EditDefaultsOnly, Category=Content, meta=(MultiLine="true"))
 	FText Text;
 
 	/** The justification of the text in the multilinebox */
diff --git a/Engine/Source/Runtime/UMG/Public/Components/MultiLineEditableTextBox.h b/Engine/Source/Runtime/UMG/Public/Components/MultiLineEditableTextBox.h
--- a/Engine/Source/Runtime/UMG/Public/Components/MultiLineEditableTextBox.h
+++ b/Engine/Source/Runtime/UMG/Public/Components/MultiLineEditableTextBox.h
@@ -29,7 +29,7 @@ public:
 	USlateWidgetStyleAsset* Style_DEPRECATED;
 
 	/** The text content for this editable text box widget */
-	UPROPERTY(EditDefaultsOnly, Category=Content)
+	UPROPERTY(EditDefaultsOnly, Category=Content, meta=(MultiLine="true"))
 	FText Text;
 
 	/** The justification of the text in the multilinebox */
diff --git a/Engine/Source/Runtime/UMG/Public/Components/TextBlock.h b/Engine/Source/Runtime/UMG/Public/Components/TextBlock.h
--- a/Engine/Source/Runtime/UMG/Public/Components/TextBlock.h
+++ b/Engine/Source/Runtime/UMG/Public/Components/TextBlock.h
@@ -45,7 +45,7 @@ public:
 	USlateWidgetStyleAsset* Style_DEPRECATED;
 
 	/** The text to display */
-	UPROPERTY(EditDefaultsOnly, Category=Content)
+	UPROPERTY(EditDefaultsOnly, Category=Content, meta=(MultiLine="true"))
 	FText Text;
 	
 	/** A bindable delegate to allow logic to drive the text of the widget */
