UMG - Fixing a crash bug with blueprint live compiler replacements.  We now double check that the widget is initialized an re-initialize it if need be.  The Designer flag is now a UProperty so that it's carried over to the new instance.

diff --git a/Engine/Source/Runtime/UMG/Private/UserWidget.cpp b/Engine/Source/Runtime/UMG/Private/UserWidget.cpp
--- a/Engine/Source/Runtime/UMG/Private/UserWidget.cpp
+++ b/Engine/Source/Runtime/UMG/Private/UserWidget.cpp
@@ -310,9 +310,15 @@ UWidget* UUserWidget::GetWidgetHandle(TSharedRef<SWidget> InWidget)
 
 TSharedRef<SWidget> UUserWidget::RebuildWidget()
 {
-	TSharedPtr<SWidget> UserRootWidget;
+	// In the event this widget is replaced in memory by the blueprint compiler update
+	// the widget won't be properly initialized, so we ensure it's initialized and initialize
+	// it if it hasn't been.
+	if ( !bInitialized )
+	{
+		Initialize();
+	}
 
-	check(bInitialized);
+	TSharedPtr<SWidget> UserRootWidget;
 
 	//setup the player context on sub user widgets, if we have a valid context
 	if (PlayerContext.IsValid())
@@ -359,8 +365,6 @@ TSharedPtr<SWidget> UUserWidget::GetWidgetFromName(const FString& Name) const
 
 UWidget* UUserWidget::GetHandleFromName(const FString& Name) const
 {
-	//return WidgetTree->FindWidget(Name);
-
 	for ( UWidget* Widget : Components )
 	{
 		if ( Widget->GetName().Equals(Name, ESearchCase::IgnoreCase) )
@@ -435,11 +439,6 @@ void UUserWidget::AddToViewport()
 		TSharedPtr<SWidget> OutUserSlateWidget;
 		TSharedRef<SWidget> RootWidget = MakeViewportWidget(OutUserSlateWidget);
 
-		//TSharedRef<SViewportWidgetHost> WidgetHost = SNew(SViewportWidgetHost, OutUserSlateWidget)
-		//	[
-		//		RootWidget
-		//	];
-
 		FullScreenWidget = OutUserSlateWidget;
 
 		// If this is a game world add the widget to the current worlds viewport.
@@ -449,24 +448,6 @@ void UUserWidget::AddToViewport()
 			if ( UGameViewportClient* ViewportClient = World->GetGameViewport() )
 			{
 				ViewportClient->AddViewportWidgetContent(OutUserSlateWidget.ToSharedRef());
-
-				////TODO UMG this isn't what should manage focus, a higher level window controller, probably the viewport needs to understand
-				//// the Widget stack, and the dialog stack.
-				//if ( bModal )
-				//{
-				//	if ( FSceneViewport* SceneViewport = ViewportClient->GetGameViewport() )
-				//	{
-				//		TWeakPtr<SViewport> GameViewportWidget = SceneViewport->GetViewportWidget();
-
-				//		if ( GameViewportWidget.IsValid() )
-				//		{
-				//			GameViewportWidget.Pin()->SetWidgetToFocusOnActivate(OutUserSlateWidget);
-
-				//			FSlateApplication::Get().SetFocusToGameViewport();
-				//			FSlateApplication::Get().SetJoystickCaptorToGameViewport();
-				//		}
-				//	}
-				//}
 			}
 		}
 	}
@@ -485,25 +466,6 @@ void UUserWidget::RemoveFromViewport()
 			if ( UGameViewportClient* ViewportClient = World->GetGameViewport() )
 			{
 				ViewportClient->RemoveViewportWidgetContent(WidgetHost.ToSharedRef());
-
-				//if ( WidgetHost->IsModal() )
-				//{
-				//	if ( FSceneViewport* SceneViewport = ViewportClient->GetGameViewport() )
-				//	{
-				//		TWeakPtr<SViewport> GameViewportWidget = SceneViewport->GetViewportWidget();
-
-				//		if ( GameViewportWidget.IsValid() )
-				//		{
-				//			//TODO UMG this isn't what should manage focus, a higher level window controller, probably the viewport needs to understand
-				//			// the Widget stack, and the dialog stack.
-				//			GameViewportWidget.Pin()->ClearWidgetToFocusOnActivate();
-				//			FSlateApplication::Get().ClearKeyboardFocus(EKeyboardFocusCause::SetDirectly);
-
-				//			FSlateApplication::Get().SetFocusToGameViewport();
-				//			FSlateApplication::Get().SetJoystickCaptorToGameViewport();
-				//		}
-				//	}
-				//}
 			}
 		}
 	}
diff --git a/Engine/Source/Runtime/UMG/Private/WidgetBlueprintGeneratedClass.cpp b/Engine/Source/Runtime/UMG/Private/WidgetBlueprintGeneratedClass.cpp
--- a/Engine/Source/Runtime/UMG/Private/WidgetBlueprintGeneratedClass.cpp
+++ b/Engine/Source/Runtime/UMG/Private/WidgetBlueprintGeneratedClass.cpp
@@ -79,6 +79,8 @@ void UWidgetBlueprintGeneratedClass::InitializeWidget(UUserWidget* UserWidget) c
 
 		}
 
+		UserWidget->Components.Reset();
+
 		for ( UWidget* Widget : ClonedWidgets )
 		{
 			// Not fatal if NULL, but shouldn't happen
diff --git a/Engine/Source/Runtime/UMG/Public/Components/Widget.h b/Engine/Source/Runtime/UMG/Public/Components/Widget.h
--- a/Engine/Source/Runtime/UMG/Public/Components/Widget.h
+++ b/Engine/Source/Runtime/UMG/Public/Components/Widget.h
@@ -326,5 +326,6 @@ protected:
 	TWeakPtr<SWidget> MyGCWidget;
 	
 	/** Is this widget being displayed on a designer surface */
+	UPROPERTY(Transient)
 	bool bDesignTime;
 };
