CIS Fixes (non-editor builds, deprecated functions)

diff --git a/Engine/Source/Runtime/UMG/Private/UserWidget.cpp b/Engine/Source/Runtime/UMG/Private/UserWidget.cpp
--- a/Engine/Source/Runtime/UMG/Private/UserWidget.cpp
+++ b/Engine/Source/Runtime/UMG/Private/UserWidget.cpp
@@ -25,8 +25,10 @@ UUserWidget::UUserWidget(const FPostConstructInitializeProperties& PCIP)
 
 	bSupportsKeyboardFocus = true;
 
+#if WITH_EDITORONLY_DATA
 	bUseDesignTimeSize = false;
 	DesignTimeSize = FVector2D(100, 100);
+#endif
 }
 
 void UUserWidget::Initialize()
diff --git a/Engine/Source/Runtime/UMG/Private/WidgetTree.cpp b/Engine/Source/Runtime/UMG/Private/WidgetTree.cpp
--- a/Engine/Source/Runtime/UMG/Private/WidgetTree.cpp
+++ b/Engine/Source/Runtime/UMG/Private/WidgetTree.cpp
@@ -94,7 +94,7 @@ void UWidgetTree::GetAllWidgets(TArray<UWidget*>& Widgets) const
 
 void UWidgetTree::GetChildWidgets(UWidget* Parent, TArray<UWidget*>& Widgets) const
 {
-	if ( INamedSlotInterface* NamedSlotHost = InterfaceCast<INamedSlotInterface>(Parent) )
+	if ( INamedSlotInterface* NamedSlotHost = Cast<INamedSlotInterface>(Parent) )
 	{
 		TArray<FName> SlotNames;
 		NamedSlotHost->GetSlotNames(SlotNames);
