TTP 348172 Fixing a crash in shooter game on stop when using widgets added to the viewport.  The shooter game instance was using the FTickerObjectBase, which ticked all the time, not just during the lifetime of the game instance (inbetween init and shutdown) so it was accessing external world context pointer data that had been cleaned up.  The lifetime of the game instance is extended slightly due to the way UMG references it, keeps it around a few extra frames.  Had to introduce the shutdown from ML into game instance in 4.5, did not expose it to blueprints.  With the shutdown concept, the shooter game instance now manages its tick delegate hooking based on when it gets init and shutdown. Additionally umg was adding the wrong root widget to the viewport.

diff --git a/Engine/Source/Runtime/Engine/Private/GameInstance.cpp b/Engine/Source/Runtime/Engine/Private/GameInstance.cpp
--- a/Engine/Source/Runtime/Engine/Private/GameInstance.cpp
+++ b/Engine/Source/Runtime/Engine/Private/GameInstance.cpp
@@ -37,6 +37,9 @@ void UGameInstance::Init()
 void UGameInstance::Shutdown()
 {
 	ReceiveShutdown();
+
+	// Clear the world context pointer to prevent further access.
+	WorldContext = nullptr;
 }
 
 void UGameInstance::InitializeStandalone()
diff --git a/Engine/Source/Runtime/UMG/Private/UserWidget.cpp b/Engine/Source/Runtime/UMG/Private/UserWidget.cpp
--- a/Engine/Source/Runtime/UMG/Private/UserWidget.cpp
+++ b/Engine/Source/Runtime/UMG/Private/UserWidget.cpp
@@ -526,7 +526,7 @@ void UUserWidget::AddToViewport()
 		TSharedPtr<SWidget> OutUserSlateWidget;
 		TSharedRef<SWidget> RootWidget = MakeViewportWidget(OutUserSlateWidget);
 
-		FullScreenWidget = OutUserSlateWidget;
+		FullScreenWidget = RootWidget;
 
 		// If this is a game world add the widget to the current worlds viewport.
 		UWorld* World = GetWorld();
@@ -534,7 +534,7 @@ void UUserWidget::AddToViewport()
 		{
 			if ( UGameViewportClient* ViewportClient = World->GetGameViewport() )
 			{
-				ViewportClient->AddViewportWidgetContent(OutUserSlateWidget.ToSharedRef());
+				ViewportClient->AddViewportWidgetContent(RootWidget);
 			}
 		}
 	}
