Temporarily removing parameter, Support to come in 4.6.

diff --git a/Engine/Source/Runtime/UMG/Private/WidgetBlueprintLibrary.cpp b/Engine/Source/Runtime/UMG/Private/WidgetBlueprintLibrary.cpp
--- a/Engine/Source/Runtime/UMG/Private/WidgetBlueprintLibrary.cpp
+++ b/Engine/Source/Runtime/UMG/Private/WidgetBlueprintLibrary.cpp
@@ -54,13 +54,13 @@ void UWidgetBlueprintLibrary::SetInputMode_UIOnly(APlayerController* Target, UWi
 	}
 }
 
-void UWidgetBlueprintLibrary::SetInputMode_GameAndUI(APlayerController* Target, UWidget* InWidgetToFocus, bool bLockMouseToViewport, bool bHideCursorDuringCapture)
+void UWidgetBlueprintLibrary::SetInputMode_GameAndUI(APlayerController* Target, UWidget* InWidgetToFocus, bool bLockMouseToViewport)//, bool bHideCursorDuringCapture)
 {
 	if (Target != nullptr)
 	{
 		FInputModeGameAndUI InputMode;
 		InputMode.SetLockMouseToViewport(bLockMouseToViewport);
-		InputMode.SetHideCursorDuringCapture(bHideCursorDuringCapture);
+		//InputMode.SetHideCursorDuringCapture(bHideCursorDuringCapture);
 
 		if (InWidgetToFocus != nullptr)
 		{
diff --git a/Engine/Source/Runtime/UMG/Public/Blueprint/WidgetBlueprintLibrary.h b/Engine/Source/Runtime/UMG/Public/Blueprint/WidgetBlueprintLibrary.h
--- a/Engine/Source/Runtime/UMG/Public/Blueprint/WidgetBlueprintLibrary.h
+++ b/Engine/Source/Runtime/UMG/Public/Blueprint/WidgetBlueprintLibrary.h
@@ -23,7 +23,7 @@ class UWidgetBlueprintLibrary : public UBlueprintFunctionLibrary
 
 	/** Setup an input mode that allows only the UI to respond to user input, and if the UI doesn't handle it player input / player controller gets a chance. */
 	UFUNCTION(BlueprintCallable, Category = "Input")
-	static void SetInputMode_GameAndUI(APlayerController* Target, UWidget* InWidgetToFocus = nullptr, bool bLockMouseToViewport = false, bool bHideCursorDuringCapture = true);
+	static void SetInputMode_GameAndUI(APlayerController* Target, UWidget* InWidgetToFocus = nullptr, bool bLockMouseToViewport = false);//, bool bHideCursorDuringCapture = true);
 
 	/** Setup an input mode that allows only player input / player controller to respond to user input. */
 	UFUNCTION(BlueprintCallable, Category = "Input")
