UMG - Adding better doc to the panel widget.

diff --git a/Engine/Source/Runtime/UMG/Public/Components/PanelWidget.h b/Engine/Source/Runtime/UMG/Public/Components/PanelWidget.h
--- a/Engine/Source/Runtime/UMG/Public/Components/PanelWidget.h
+++ b/Engine/Source/Runtime/UMG/Public/Components/PanelWidget.h
@@ -18,32 +18,37 @@ public:
 
 public:
 
-	/**  */
+	/** Gets number of child widgets in the container. */
 	UFUNCTION(BlueprintCallable, Category="Widget|Panel")
 	int32 GetChildrenCount() const;
 
-	/**  */
+	/**
+	 * Gets the widget at an index.
+	 * @param Index The index of the widget.
+	 * @return The widget at the given index, or nothing if there is no widget there.
+	 */
 	UFUNCTION(BlueprintCallable, Category="Widget|Panel")
 	UWidget* GetChildAt(int32 Index) const;
 
-	/**  */
+	/** Gets the index of a specific child widget */
 	UFUNCTION(BlueprintCallable, Category="Widget|Panel")
 	int32 GetChildIndex(UWidget* Content) const;
 
-	/**  */
+	/** Removes a child by it's index. */
 	UFUNCTION(BlueprintCallable, Category="Widget|Panel")
 	bool RemoveChildAt(int32 Index);
 
-	/**  */
+	/**
+	 * Adds a new child widget to the container.  Returns the base slot type, 
+	 * requires casting to turn it into the type specific to the container.
+	 */
 	UFUNCTION(BlueprintCallable, Category="Widget|Panel")
 	UPanelSlot* AddChild(UWidget* Content);
 
 	/** Swaps the widget out of the slot at the given index, replacing it with a different widget. */
-	UFUNCTION(BlueprintCallable, Category="Widget|Panel")
 	void ReplaceChildAt(int32 Index, UWidget* Content);
 
 	/** Inserts a widget at a specific index */
-	UFUNCTION(BlueprintCallable, Category="Widget|Panel")
 	void InsertChildAt(int32 Index, UWidget* Content);
 
 	/**
