UMG - Reintroducing the menu anchor widget, you can now make custom popupable widgets with it.

diff --git a/Engine/Source/Runtime/UMG/Private/Components/MenuAnchor.cpp b/Engine/Source/Runtime/UMG/Private/Components/MenuAnchor.cpp
--- a/Engine/Source/Runtime/UMG/Private/Components/MenuAnchor.cpp
+++ b/Engine/Source/Runtime/UMG/Private/Components/MenuAnchor.cpp
@@ -78,11 +78,27 @@ TSharedRef<SWidget> UMenuAnchor::HandleGetMenuContent()
 	return SlateMenuWidget.IsValid() ? SlateMenuWidget.ToSharedRef() : SNullWidget::NullWidget;
 }
 
-void UMenuAnchor::SetIsOpen(bool InIsOpen, bool bFocusMenu)
+void UMenuAnchor::ToggleOpen(bool bFocusOnOpen)
 {
 	if ( MyMenuAnchor.IsValid() )
 	{
-		MyMenuAnchor->SetIsOpen(InIsOpen, bFocusMenu);
+		MyMenuAnchor->SetIsOpen(!MyMenuAnchor->IsOpen(), bFocusOnOpen);
+	}
+}
+
+void UMenuAnchor::Open(bool bFocusMenu)
+{
+	if ( MyMenuAnchor.IsValid() && !MyMenuAnchor->IsOpen() )
+	{
+		MyMenuAnchor->SetIsOpen(true, bFocusMenu);
+	}
+}
+
+void UMenuAnchor::Close()
+{
+	if ( MyMenuAnchor.IsValid() )
+	{
+		return MyMenuAnchor->SetIsOpen(false, false);
 	}
 }
 
@@ -105,7 +121,7 @@ const FSlateBrush* UMenuAnchor::GetEditorIcon()
 
 const FText UMenuAnchor::GetPaletteCategory()
 {
-	return LOCTEXT("Advanced", "Advanced");
+	return LOCTEXT("Primitive", "Primitive");
 }
 
 #endif
diff --git a/Engine/Source/Runtime/UMG/Private/WidgetBlueprintLibrary.cpp b/Engine/Source/Runtime/UMG/Private/WidgetBlueprintLibrary.cpp
--- a/Engine/Source/Runtime/UMG/Private/WidgetBlueprintLibrary.cpp
+++ b/Engine/Source/Runtime/UMG/Private/WidgetBlueprintLibrary.cpp
@@ -272,4 +272,9 @@ FSlateBrush UWidgetBlueprintLibrary::NoResourceBrush()
 	return FSlateNoResource();
 }
 
+void UWidgetBlueprintLibrary::DismissAllMenus()
+{
+	FSlateApplication::Get().DismissAllMenus();
+}
+
 #undef LOCTEXT_NAMESPACE
\ No newline at end of file
diff --git a/Engine/Source/Runtime/UMG/Public/Blueprint/WidgetBlueprintLibrary.h b/Engine/Source/Runtime/UMG/Public/Blueprint/WidgetBlueprintLibrary.h
--- a/Engine/Source/Runtime/UMG/Public/Blueprint/WidgetBlueprintLibrary.h
+++ b/Engine/Source/Runtime/UMG/Public/Blueprint/WidgetBlueprintLibrary.h
@@ -132,4 +132,8 @@ class UWidgetBlueprintLibrary : public UBlueprintFunctionLibrary
 	 */
 	UFUNCTION(BlueprintPure, Category="Widget|Brush")
 	static FSlateBrush NoResourceBrush();
+
+	/** Closes any popup menu */
+	UFUNCTION(BlueprintCallable, Category="Widget|Brush")
+	static void DismissAllMenus();
 };
diff --git a/Engine/Source/Runtime/UMG/Public/Components/MenuAnchor.h b/Engine/Source/Runtime/UMG/Public/Components/MenuAnchor.h
--- a/Engine/Source/Runtime/UMG/Public/Components/MenuAnchor.h
+++ b/Engine/Source/Runtime/UMG/Public/Components/MenuAnchor.h
@@ -5,41 +5,52 @@
 #include "MenuAnchor.generated.h"
 
 /**
- * The Menu Anchor allows you to specify an location that a popup menu should be anchored to, and should be
- * summoned from.
+ * The Menu Anchor allows you to specify an location that a popup menu should be anchored to, 
+ * and should be summoned from.
  */
-UCLASS(Experimental)
+UCLASS()
 class UMG_API UMenuAnchor : public UContentWidget
 {
 	GENERATED_UCLASS_BODY()
 
 public:
 
-	/** The widget class to spawn when the menu is required.  Creates the object freshly each time. */
-	UPROPERTY(EditDefaultsOnly, Category=Menu)
+	/**
+	 * The widget class to spawn when the menu is required.  Creates the widget freshly each time.  
+	 * If you want to customize the creation of the popup, you should bind a function to OnGetMenuContentEvent 
+	 * instead.
+	 */
+	UPROPERTY(EditDefaultsOnly, Category="Menu")
 	TSubclassOf<class UUserWidget> MenuClass;
 
-	/** Called when the menu content is requested to allow dynamic decision over what to display */
-	UPROPERTY(EditDefaultsOnly, Category=Events)
+	/** Called when the menu content is requested to allow a more customized handling over what to display */
+	UPROPERTY(EditDefaultsOnly, Category="Events")
 	FGetContent OnGetMenuContentEvent;
 	
 	/** The placement location of the summoned widget. */
-	UPROPERTY(EditDefaultsOnly, Category=Menu)
+	UPROPERTY(EditDefaultsOnly, Category="Menu")
 	TEnumAsByte<EMenuPlacement> Placement;
 
 public:
 
 	/**
-	 * Open or close the popup
+	 * Toggles the menus open state.
 	 *
-	 * @param InIsOpen    If true, open the popup. Otherwise close it.
-	 * @param bFocusMenu  Should we focus the popup as soon as it opens?
+	 * @param bFocusOnOpen  Should we focus the popup as soon as it opens?
 	 */
-	UFUNCTION(BlueprintCallable, Category="Widget")
-	void SetIsOpen(bool InIsOpen, bool bFocusMenu = true);
+	UFUNCTION(BlueprintCallable, Category="Menu Anchor")
+	void ToggleOpen(bool bFocusOnOpen);
+
+	/** Opens the menu if it is not already open */
+	UFUNCTION(BlueprintCallable, Category="Menu Anchor")
+	void Open(bool bFocusMenu);
+
+	/** Closes the menu if it is currently open. */
+	UFUNCTION(BlueprintCallable, Category="Menu Anchor")
+	void Close();
 
 	/** @return true if the popup is open; false otherwise. */
-	UFUNCTION(BlueprintCallable, Category="Widget")
+	UFUNCTION(BlueprintCallable, Category="Menu Anchor")
 	bool IsOpen() const;
 
 	///** @return true if we should open the menu due to a click. Sometimes we should not, if
