UMG - Menu anchor no longer attempts to create the widget with the class as well if the delegate is bound.

diff --git a/Engine/Source/Runtime/UMG/Private/Components/MenuAnchor.cpp b/Engine/Source/Runtime/UMG/Private/Components/MenuAnchor.cpp
--- a/Engine/Source/Runtime/UMG/Private/Components/MenuAnchor.cpp
+++ b/Engine/Source/Runtime/UMG/Private/Components/MenuAnchor.cpp
@@ -65,13 +65,15 @@ TSharedRef<SWidget> UMenuAnchor::HandleGetMenuContent()
 			SlateMenuWidget = MenuWidget->TakeWidget();
 		}
 	}
-
-	if ( MenuClass != NULL && !MenuClass->HasAnyClassFlags(CLASS_Abstract) )
+	else
 	{
-		UWidget* MenuWidget = (UWidget*)ConstructObject<UWidget>(MenuClass, GetOuter());
-		if ( MenuWidget )
+		if ( MenuClass != nullptr && !MenuClass->HasAnyClassFlags(CLASS_Abstract) )
 		{
-			SlateMenuWidget = MenuWidget->TakeWidget();
+			UUserWidget* MenuWidget = CreateWidget<UUserWidget>(GetWorld(), MenuClass);
+			if ( MenuWidget )
+			{
+				SlateMenuWidget = MenuWidget->TakeWidget();
+			}
 		}
 	}
 
