UMG - Adding Begin Play as keywords for the construct function for UMG widgets.

diff --git a/Engine/Source/Runtime/UMG/Public/Blueprint/UserWidget.h b/Engine/Source/Runtime/UMG/Public/Blueprint/UserWidget.h
--- a/Engine/Source/Runtime/UMG/Public/Blueprint/UserWidget.h
+++ b/Engine/Source/Runtime/UMG/Public/Blueprint/UserWidget.h
@@ -196,7 +196,7 @@ public:
 	class APawn* GetOwningPlayerPawn() const;
 
 	/** Called when the widget is constructed */
-	UFUNCTION(BlueprintNativeEvent, Category="User Interface")
+	UFUNCTION(BlueprintNativeEvent, Category="User Interface", meta=(Keywords="Begin Play"))
 	void Construct();
 
 	UFUNCTION(BlueprintNativeEvent, Category="User Interface")
