UMG - Render Transform is now readable in blueprints.

diff --git a/Engine/Source/Runtime/UMG/Public/Components/Widget.h b/Engine/Source/Runtime/UMG/Public/Components/Widget.h
--- a/Engine/Source/Runtime/UMG/Public/Components/Widget.h
+++ b/Engine/Source/Runtime/UMG/Public/Components/Widget.h
@@ -67,7 +67,7 @@ public:
 	/**
 	 * The parent slot of the UWidget.  Allows us to easily inline edit the layout controlling this widget.
 	 */
-	UPROPERTY(Instanced, EditAnywhere, BlueprintReadOnly, Category = Layout, meta = (ShowOnlyInnerProperties))
+	UPROPERTY(Instanced, EditAnywhere, BlueprintReadOnly, Category=Layout, meta=(ShowOnlyInnerProperties))
 	UPanelSlot* Slot;
 
 	/** Sets whether this widget can be modified interactively by the user */
@@ -107,11 +107,11 @@ public:
 	FGetMouseCursor CursorDelegate;
 
 	/** The render transform of the widget allows for arbitrary 2D transforms to be applied to the widget. */
-	UPROPERTY(EditDefaultsOnly, Category=RenderTransform, meta=(DisplayName="Transform"))
+	UPROPERTY(EditDefaultsOnly, BlueprintReadOnly, Category="Render Transform", meta=( DisplayName="Transform" ))
 	FWidgetTransform RenderTransform;
 
 	/** The render transform pivot controls the location about which transforms are applied.  This value is a normalized coordinate about which things like rotations will occur. */
-	UPROPERTY(EditDefaultsOnly, Category=RenderTransform, meta=(DisplayName="Pivot"))
+	UPROPERTY(EditDefaultsOnly, BlueprintReadOnly, Category="Render Transform", meta=( DisplayName="Pivot" ))
 	FVector2D RenderTransformPivot;
 
 #if WITH_EDITORONLY_DATA
