#UE4

- Making all the getters/setters on UWidget use "GetCachedWidget" instead of messing with "MyWidget" directly.

diff --git a/Engine/Source/Runtime/UMG/Private/Components/Widget.cpp b/Engine/Source/Runtime/UMG/Private/Components/Widget.cpp
--- a/Engine/Source/Runtime/UMG/Private/Components/Widget.cpp
+++ b/Engine/Source/Runtime/UMG/Private/Components/Widget.cpp
@@ -47,16 +47,17 @@ void UWidget::SetRenderTranslation(FVector2D Translation)
 
 void UWidget::UpdateRenderTransform()
 {
-	if ( MyWidget.IsValid() )
+	TSharedPtr<SWidget> SafeWidget = GetCachedWidget();
+	if (SafeWidget.IsValid())
 	{
 		if (!RenderTransform.IsIdentity())
 		{
 			FSlateRenderTransform Transform2D = ::Concatenate(FScale2D(RenderTransform.Scale), FShear2D::FromShearAngles(RenderTransform.Shear), FQuat2D(FMath::DegreesToRadians(RenderTransform.Angle)), FVector2D(RenderTransform.Translation));
-			MyWidget.Pin()->SetRenderTransform(Transform2D);
+			SafeWidget->SetRenderTransform(Transform2D);
 		}
 		else
 		{
-			MyWidget.Pin()->SetRenderTransform(TOptional<FSlateRenderTransform>());
+			SafeWidget->SetRenderTransform(TOptional<FSlateRenderTransform>());
 		}
 	}
 }
@@ -64,31 +65,37 @@ void UWidget::UpdateRenderTransform()
 void UWidget::SetRenderTransformPivot(FVector2D Pivot)
 {
 	RenderTransformPivot = Pivot;
-	if ( MyWidget.IsValid() )
+
+	TSharedPtr<SWidget> SafeWidget = GetCachedWidget();
+	if (SafeWidget.IsValid())
 	{
-		MyWidget.Pin()->SetRenderTransformPivot(Pivot);
+		SafeWidget->SetRenderTransformPivot(Pivot);
 	}
 }
 
 bool UWidget::GetIsEnabled() const
 {
-	return MyWidget.IsValid() ? MyWidget.Pin()->IsEnabled() : bIsEnabled;
+	TSharedPtr<SWidget> SafeWidget = GetCachedWidget();
+	return SafeWidget.IsValid() ? SafeWidget->IsEnabled() : bIsEnabled;
 }
 
 void UWidget::SetIsEnabled(bool bInIsEnabled)
 {
 	bIsEnabled = bInIsEnabled;
-	if ( MyWidget.IsValid() )
+
+	TSharedPtr<SWidget> SafeWidget = GetCachedWidget();
+	if (SafeWidget.IsValid())
 	{
-		MyWidget.Pin()->SetEnabled(bInIsEnabled);
+		SafeWidget->SetEnabled(bInIsEnabled);
 	}
 }
 
 TEnumAsByte<ESlateVisibility::Type> UWidget::GetVisibility()
 {
-	if ( MyWidget.IsValid() )
+	TSharedPtr<SWidget> SafeWidget = GetCachedWidget();
+	if (SafeWidget.IsValid())
 	{
-		return UWidget::ConvertRuntimeToSerializedVisiblity(MyWidget.Pin()->GetVisibility());
+		return UWidget::ConvertRuntimeToSerializedVisiblity(SafeWidget->GetVisibility());
 	}
 
 	return Visiblity;
@@ -98,9 +105,10 @@ void UWidget::SetVisibility(TEnumAsByte<ESlateVisibility::Type> InVisibility)
 {
 	Visiblity = InVisibility;
 
-	if ( MyWidget.IsValid() )
+	TSharedPtr<SWidget> SafeWidget = GetCachedWidget();
+	if (SafeWidget.IsValid())
 	{
-		return MyWidget.Pin()->SetVisibility(UWidget::ConvertSerializedVisibilityToRuntime(InVisibility));
+		return SafeWidget->SetVisibility(UWidget::ConvertSerializedVisibilityToRuntime(InVisibility));
 	}
 }
 
@@ -108,17 +116,19 @@ void UWidget::SetToolTipText(const FText& InToolTipText)
 {
 	ToolTipText = InToolTipText;
 
-	if ( MyWidget.IsValid() )
+	TSharedPtr<SWidget> SafeWidget = GetCachedWidget();
+	if (SafeWidget.IsValid())
 	{
-		return MyWidget.Pin()->SetToolTipText(InToolTipText);
+		return SafeWidget->SetToolTipText(InToolTipText);
 	}
 }
 
 bool UWidget::IsHovered() const
 {
-	if ( MyWidget.IsValid() )
+	TSharedPtr<SWidget> SafeWidget = GetCachedWidget();
+	if (SafeWidget.IsValid())
 	{
-		return MyWidget.Pin()->IsHovered();
+		return SafeWidget->IsHovered();
 	}
 
 	return false;
@@ -126,9 +136,10 @@ bool UWidget::IsHovered() const
 
 bool UWidget::HasKeyboardFocus() const
 {
-	if ( MyWidget.IsValid() )
+	TSharedPtr<SWidget> SafeWidget = GetCachedWidget();
+	if (SafeWidget.IsValid())
 	{
-		return MyWidget.Pin()->HasKeyboardFocus();
+		return SafeWidget->HasKeyboardFocus();
 	}
 
 	return false;
@@ -136,9 +147,10 @@ bool UWidget::HasKeyboardFocus() const
 
 bool UWidget::HasFocusedDescendants() const
 {
-	if ( MyWidget.IsValid() )
+	TSharedPtr<SWidget> SafeWidget = GetCachedWidget();
+	if (SafeWidget.IsValid())
 	{
-		return MyWidget.Pin()->HasFocusedDescendants();
+		return SafeWidget->HasFocusedDescendants();
 	}
 
 	return false;
@@ -146,9 +158,10 @@ bool UWidget::HasFocusedDescendants() const
 
 bool UWidget::HasMouseCapture() const
 {
-	if ( MyWidget.IsValid() )
+	TSharedPtr<SWidget> SafeWidget = GetCachedWidget();
+	if (SafeWidget.IsValid())
 	{
-		return MyWidget.Pin()->HasMouseCapture();
+		return SafeWidget->HasMouseCapture();
 	}
 
 	return false;
@@ -156,25 +169,28 @@ bool UWidget::HasMouseCapture() const
 
 void UWidget::SetKeyboardFocus() const
 {
-	if ( MyWidget.IsValid() )
+	TSharedPtr<SWidget> SafeWidget = GetCachedWidget();
+	if (SafeWidget.IsValid())
 	{
-		FSlateApplication::Get().SetKeyboardFocus(MyWidget.Pin());
+		FSlateApplication::Get().SetKeyboardFocus(SafeWidget);
 	}
 }
 
 void UWidget::ForceLayoutPrepass()
 {
-	if ( MyWidget.IsValid() )
+	TSharedPtr<SWidget> SafeWidget = GetCachedWidget();
+	if (SafeWidget.IsValid())
 	{
-		MyWidget.Pin()->SlatePrepass();
+		SafeWidget->SlatePrepass();
 	}
 }
 
 FVector2D UWidget::GetDesiredSize() const
 {
-	if ( MyWidget.IsValid() )
+	TSharedPtr<SWidget> SafeWidget = GetCachedWidget();
+	if (SafeWidget.IsValid())
 	{
-		return MyWidget.Pin()->GetDesiredSize();
+		return SafeWidget->GetDesiredSize();
 	}
 
 	return FVector2D(0, 0);
@@ -341,7 +357,8 @@ void UWidget::PostEditChangeProperty(FPropertyChangedEvent& PropertyChangedEvent
 {
 	Super::PostEditChangeProperty(PropertyChangedEvent);
 
-	if ( MyWidget.IsValid() )
+	TSharedPtr<SWidget> SafeWidget = GetCachedWidget();
+	if (SafeWidget.IsValid())
 	{
 		SynchronizeProperties();
 	}
