UMG - You can now dynamically set the fill color on the progress bar, as well as bind a function to it.

diff --git a/Engine/Source/Runtime/UMG/Private/Components/ProgressBar.cpp b/Engine/Source/Runtime/UMG/Private/Components/ProgressBar.cpp
--- a/Engine/Source/Runtime/UMG/Private/Components/ProgressBar.cpp
+++ b/Engine/Source/Runtime/UMG/Private/Components/ProgressBar.cpp
@@ -37,14 +37,16 @@ TSharedRef<SWidget> UProgressBar::RebuildWidget()
 
 void UProgressBar::SynchronizeProperties()
 {
+	Super::SynchronizeProperties();
+
 	TAttribute< TOptional<float> > PercentBinding = OPTIONAL_BINDING_CONVERT(float, Percent, TOptional<float>, ConvertFloatToOptionalFloat);
+	TAttribute<FSlateColor> FillColorAndOpacityBinding = OPTIONAL_BINDING(FSlateColor, FillColorAndOpacity);
 
 	MyProgressBar->SetStyle(&WidgetStyle);
 
-	MyProgressBar->SetPercent(bIsMarquee ? TOptional<float>() : PercentBinding);
-	
 	MyProgressBar->SetBarFillType(BarFillType);
-	MyProgressBar->SetFillColorAndOpacity(FillColorAndOpacity);
+	MyProgressBar->SetPercent(bIsMarquee ? TOptional<float>() : PercentBinding);
+	MyProgressBar->SetFillColorAndOpacity(FillColorAndOpacityBinding);
 }
 
 void UProgressBar::SetIsMarquee(bool InbIsMarquee)
@@ -56,6 +58,15 @@ void UProgressBar::SetIsMarquee(bool InbIsMarquee)
 	}
 }
 
+void UProgressBar::SetFillColorAndOpacity(FLinearColor Color)
+{
+	FillColorAndOpacity = Color;
+	if (MyProgressBar.IsValid())
+	{
+		MyProgressBar->SetFillColorAndOpacity(FillColorAndOpacity);
+	}
+}
+
 void UProgressBar::SetPercent(float InPercent)
 {
 	Percent = InPercent;
diff --git a/Engine/Source/Runtime/UMG/Public/Components/ProgressBar.h b/Engine/Source/Runtime/UMG/Public/Components/ProgressBar.h
--- a/Engine/Source/Runtime/UMG/Public/Components/ProgressBar.h
+++ b/Engine/Source/Runtime/UMG/Public/Components/ProgressBar.h
@@ -55,12 +55,20 @@ public:
 	UPROPERTY(EditDefaultsOnly, Category=Appearance)
 	FLinearColor FillColorAndOpacity;
 
+	/** */
+	UPROPERTY()
+	FGetLinearColor FillColorAndOpacityDelegate;
+
 public:
 	
 	/** Sets the current value of the ProgressBar. */
-	UFUNCTION(BlueprintCallable, Category="Behavior")
+	UFUNCTION(BlueprintCallable, Category="Progress Bar")
 	void SetPercent(float InPercent);
 
+	/** Sets the fill color of the progress bar. */
+	UFUNCTION(BlueprintCallable, Category="Progress Bar")
+	void SetFillColorAndOpacity(FLinearColor InColor);
+
 	/** Sets the progress bar to show as a marquee. */
 	UFUNCTION(BlueprintCallable, Category="Behavior")
 	void SetIsMarquee(bool InbIsMarquee);
