UMG - Removing dead code.

diff --git a/Engine/Source/Runtime/UMG/Private/Slate/SObjectWidget.cpp b/Engine/Source/Runtime/UMG/Private/Slate/SObjectWidget.cpp
--- a/Engine/Source/Runtime/UMG/Private/Slate/SObjectWidget.cpp
+++ b/Engine/Source/Runtime/UMG/Private/Slate/SObjectWidget.cpp
@@ -46,7 +46,6 @@ int32 SObjectWidget::OnPaint(const FPaintArgs& Args, const FGeometry& AllottedGe
 
 	if ( WidgetObject && !WidgetObject->IsDesignTime() )
 	{
-		//const FGeometry& AllottedGeometry, const FSlateRect& MyClippingRect, FSlateWindowElementList& OutDrawElements, int32 LayerId, const FWidgetStyle& InWidgetStyle, bool bParentEnabled
 		FPaintContext Context(AllottedGeometry, MyClippingRect, OutDrawElements, LayerId, InWidgetStyle, bParentEnabled);
 		WidgetObject->OnPaint(Context);
 		
