TTP 349084 - Grid slots are now initialized correctly with the correct span/column values.

diff --git a/Engine/Source/Runtime/UMG/Private/Components/GridSlot.cpp b/Engine/Source/Runtime/UMG/Private/Components/GridSlot.cpp
--- a/Engine/Source/Runtime/UMG/Private/Components/GridSlot.cpp
+++ b/Engine/Source/Runtime/UMG/Private/Components/GridSlot.cpp
@@ -7,7 +7,7 @@
 
 UGridSlot::UGridSlot(const FPostConstructInitializeProperties& PCIP)
 	: Super(PCIP)
-	, Slot(NULL)
+	, Slot(nullptr)
 {
 	HorizontalAlignment = HAlign_Fill;
 	VerticalAlignment = VAlign_Fill;
@@ -20,7 +20,7 @@ void UGridSlot::ReleaseSlateResources(bool bReleaseChildren)
 {
 	Super::ReleaseSlateResources(bReleaseChildren);
 
-	Slot = NULL;
+	Slot = nullptr;
 }
 
 void UGridSlot::BuildSlot(TSharedRef<SGridPanel> GridPanel)
@@ -28,13 +28,11 @@ void UGridSlot::BuildSlot(TSharedRef<SGridPanel> GridPanel)
 	Slot = &GridPanel->AddSlot(Column, Row)
 		.HAlign(HorizontalAlignment)
 		.VAlign(VerticalAlignment)
-		.Row(Row)
 		.RowSpan(RowSpan)
-		.Column(Row)
-		.ColumnSpan(RowSpan)
+		.ColumnSpan(ColumnSpan)
 		.Nudge(Nudge)
 		[
-			Content == NULL ? SNullWidget::NullWidget : Content->TakeWidget()
+			Content == nullptr ? SNullWidget::NullWidget : Content->TakeWidget()
 		];
 }
 
