UMG - You can now call GetDynamicMaterial on a brush to automatically have it transform the material on it into a MID and return it.

diff --git a/Engine/Source/Runtime/UMG/Public/Blueprint/WidgetBlueprintLibrary.h b/Engine/Source/Runtime/UMG/Public/Blueprint/WidgetBlueprintLibrary.h
--- a/Engine/Source/Runtime/UMG/Public/Blueprint/WidgetBlueprintLibrary.h
+++ b/Engine/Source/Runtime/UMG/Public/Blueprint/WidgetBlueprintLibrary.h
@@ -13,10 +13,13 @@ class UWidgetBlueprintLibrary : public UBlueprintFunctionLibrary
 	UFUNCTION(BlueprintCallable, meta=( WorldContext="WorldContextObject", FriendlyName = "Create Widget", BlueprintInternalUseOnly = "true" ), Category="User Interface|Widget")
 	static class UUserWidget* Create(UObject* WorldContextObject, TSubclassOf<class UUserWidget> WidgetType, APlayerController* OwningPlayer);
 
+	/**
+	 * Creates a new drag and drop operation that can be returned from a drag begin to inform the UI what i
+	 * being dragged and dropped and what it looks like.
+	 */
 	UFUNCTION(BlueprintCallable, Category="User Interface|Drag and Drop", meta=( BlueprintInternalUseOnly = "true" ))
 	static UDragDropOperation* CreateDragDropOperation(TSubclassOf<UDragDropOperation> OperationClass);
 	
-
 	/** Setup an input mode that allows only the UI to respond to user input. */
 	UFUNCTION(BlueprintCallable, Category = "Input")
 	static void SetInputMode_UIOnly(APlayerController* Target, UWidget* InWidgetToFocus = nullptr, bool bLockMouseToViewport = false);
@@ -60,9 +63,11 @@ class UWidgetBlueprintLibrary : public UBlueprintFunctionLibrary
 	UFUNCTION(BlueprintCallable, Category="Painting")
 	static void DrawText(UPARAM(ref) FPaintContext& Context, const FString& InString, FVector2D Position, FLinearColor Tint = FLinearColor::White);
 
+	/** The default event reply when simply handling an event. */
 	UFUNCTION(BlueprintPure, Category="Widget|Event Reply")
 	static FEventReply Handled();
 
+	/** The event reply to use when you choose not to handle an event. */
 	UFUNCTION(BlueprintPure, Category="Widget|Event Reply")
 	static FEventReply Unhandled();
 
@@ -133,6 +138,15 @@ class UWidgetBlueprintLibrary : public UBlueprintFunctionLibrary
 	UFUNCTION(BlueprintPure, Category="Widget|Brush")
 	static FSlateBrush NoResourceBrush();
 
+	/**
+	 * Gets the material that allows changes to parameters at runtime.  The brush must already have a material assigned to it, 
+	 * if it does it will automatically be converted to a MID.
+	 *
+	 * @return A material that supports dynamic input from the game.
+	 */
+	UFUNCTION(BlueprintPure, Category="Widget|Brush")
+	static UMaterialInstanceDynamic* GetDynamicMaterial(FSlateBrush& Brush);
+
 	/** Closes any popup menu */
 	UFUNCTION(BlueprintCallable, Category="Widget|Brush")
 	static void DismissAllMenus();
