UMG - You can now call GetDynamicMaterial on a brush to automatically have it transform the material on it into a MID and return it. (part 2)

diff --git a/Engine/Source/Runtime/UMG/Private/WidgetBlueprintLibrary.cpp b/Engine/Source/Runtime/UMG/Private/WidgetBlueprintLibrary.cpp
--- a/Engine/Source/Runtime/UMG/Private/WidgetBlueprintLibrary.cpp
+++ b/Engine/Source/Runtime/UMG/Private/WidgetBlueprintLibrary.cpp
@@ -272,6 +272,29 @@ FSlateBrush UWidgetBlueprintLibrary::NoResourceBrush()
 	return FSlateNoResource();
 }
 
+UMaterialInstanceDynamic* UWidgetBlueprintLibrary::GetDynamicMaterial(FSlateBrush& Brush)
+{
+	UObject* Resource = Brush.GetResourceObject();
+
+	// If we already have a dynamic material, return it.
+	if ( UMaterialInstanceDynamic* DynamicMaterial = Cast<UMaterialInstanceDynamic>(Resource) )
+	{
+		return DynamicMaterial;
+	}
+	// If the resource has a material interface we'll just update the brush to have a dynamic material.
+	else if ( UMaterialInterface* Material = Cast<UMaterialInterface>(Resource) )
+	{
+		UMaterialInstanceDynamic* DynamicMaterial = UMaterialInstanceDynamic::Create(Material, nullptr);
+		Brush.SetResourceObject(DynamicMaterial);
+
+		return DynamicMaterial;
+	}
+
+	//TODO UMG can we do something for textures?  General purpose dynamic material for them?
+
+	return nullptr;
+}
+
 void UWidgetBlueprintLibrary::DismissAllMenus()
 {
 	FSlateApplication::Get().DismissAllMenus();
