UMG - Exposing the SetMousePosition event reply to blueprints.

diff --git a/Engine/Source/Runtime/UMG/Private/WidgetBlueprintLibrary.cpp b/Engine/Source/Runtime/UMG/Private/WidgetBlueprintLibrary.cpp
--- a/Engine/Source/Runtime/UMG/Private/WidgetBlueprintLibrary.cpp
+++ b/Engine/Source/Runtime/UMG/Private/WidgetBlueprintLibrary.cpp
@@ -196,6 +196,14 @@ FEventReply UWidgetBlueprintLibrary::ReleaseJoystickCapture(UPARAM(ref) FEventRe
 	return Reply;
 }
 
+FEventReply UWidgetBlueprintLibrary::SetMousePosition(UPARAM(ref) FEventReply& Reply, FVector2D NewMousePosition)
+{
+	FIntPoint NewPoint((int32)NewMousePosition.X, (int32)NewMousePosition.Y);
+	Reply.NativeReply = Reply.NativeReply.SetMousePos(NewPoint);
+
+	return Reply;
+}
+
 FEventReply UWidgetBlueprintLibrary::DetectDrag(UPARAM(ref) FEventReply& Reply, UWidget* WidgetDetectingDrag, FKey DragKey)
 {
 	if ( WidgetDetectingDrag )
diff --git a/Engine/Source/Runtime/UMG/Public/Blueprint/WidgetBlueprintLibrary.h b/Engine/Source/Runtime/UMG/Public/Blueprint/WidgetBlueprintLibrary.h
--- a/Engine/Source/Runtime/UMG/Public/Blueprint/WidgetBlueprintLibrary.h
+++ b/Engine/Source/Runtime/UMG/Public/Blueprint/WidgetBlueprintLibrary.h
@@ -83,6 +83,9 @@ class UWidgetBlueprintLibrary : public UBlueprintFunctionLibrary
 	UFUNCTION(BlueprintPure, Category="Widget|Event Reply")
 	static FEventReply ReleaseJoystickCapture(UPARAM(ref) FEventReply& Reply, bool bInAllJoysticks = false);
 
+	UFUNCTION(BlueprintPure, Category="Widget|Event Reply")
+	static FEventReply SetMousePosition(UPARAM(ref) FEventReply& Reply, FVector2D NewMousePosition);
+
 	/**
 	 * Ask Slate to detect if a user started dragging in this widget.
 	 * If a drag is detected, Slate will send an OnDragDetected event.
