#UMG

- Adding a Play Sound function to User Widgets

diff --git a/Engine/Source/Runtime/SlateCore/Public/Sound/SlateSound.h b/Engine/Source/Runtime/SlateCore/Public/Sound/SlateSound.h
--- a/Engine/Source/Runtime/SlateCore/Public/Sound/SlateSound.h
+++ b/Engine/Source/Runtime/SlateCore/Public/Sound/SlateSound.h
@@ -29,6 +29,14 @@ public:
 	 */
 	UObject* GetResourceObject( ) const;
 
+	/**
+	* Sets the sound that is supposed to be played.
+	*/
+	void SetResourceObject(class UObject* InResourceObject)
+	{
+		ResourceObject = InResourceObject;
+	}
+
 	/**
 	 * Used when updating the ResourceObject within an FSlateSound from the editor to clear out any legacy data that may be set.
 	 */
diff --git a/Engine/Source/Runtime/UMG/Private/UserWidget.cpp b/Engine/Source/Runtime/UMG/Private/UserWidget.cpp
--- a/Engine/Source/Runtime/UMG/Private/UserWidget.cpp
+++ b/Engine/Source/Runtime/UMG/Private/UserWidget.cpp
@@ -325,6 +325,16 @@ void UUserWidget::OnAnimationFinishedPlaying( UUMGSequencePlayer& Player )
 	StoppedSequencePlayers.Add( &Player );
 }
 
+void UUserWidget::PlaySound(USoundBase* SoundToPlay)
+{
+	if (SoundToPlay)
+	{
+		FSlateSound NewSound;
+		NewSound.SetResourceObject(SoundToPlay);
+		FSlateApplication::Get().PlaySound(NewSound);
+	}
+}
+
 UWidget* UUserWidget::GetWidgetHandle(TSharedRef<SWidget> InWidget)
 {
 	return WidgetTree->FindWidget(InWidget);
diff --git a/Engine/Source/Runtime/UMG/Public/Blueprint/UserWidget.h b/Engine/Source/Runtime/UMG/Public/Blueprint/UserWidget.h
--- a/Engine/Source/Runtime/UMG/Public/Blueprint/UserWidget.h
+++ b/Engine/Source/Runtime/UMG/Public/Blueprint/UserWidget.h
@@ -300,6 +300,14 @@ public:
 	/** Called when a sequence player is finished playing an animation */
 	void OnAnimationFinishedPlaying(UUMGSequencePlayer& Player );
 
+	/**
+	* Plays a sound through the UI
+	*
+	* @param The sound to play
+	*/
+	UFUNCTION(BlueprintCallable, Category = "Sound")
+	void PlaySound(class USoundBase* SoundToPlay);
+
 	/** @returns The UObject wrapper for a given SWidget */
 	UWidget* GetWidgetHandle(TSharedRef<SWidget> InWidget);
 
