UMG - Using displayname to clean up the name of several basic widgets.  Moving several less common widgets into Input.

diff --git a/Engine/Source/Runtime/UMG/Private/Components/MultiLineEditableTextBox.cpp b/Engine/Source/Runtime/UMG/Private/Components/MultiLineEditableTextBox.cpp
--- a/Engine/Source/Runtime/UMG/Private/Components/MultiLineEditableTextBox.cpp
+++ b/Engine/Source/Runtime/UMG/Private/Components/MultiLineEditableTextBox.cpp
@@ -138,7 +138,7 @@ const FSlateBrush* UMultiLineEditableTextBox::GetEditorIcon()
 
 const FText UMultiLineEditableTextBox::GetPaletteCategory()
 {
-	return LOCTEXT("Common", "Common");
+	return LOCTEXT("Input", "Input");
 }
 
 #endif
diff --git a/Engine/Source/Runtime/UMG/Private/Components/SizeBox.cpp b/Engine/Source/Runtime/UMG/Private/Components/SizeBox.cpp
--- a/Engine/Source/Runtime/UMG/Private/Components/SizeBox.cpp
+++ b/Engine/Source/Runtime/UMG/Private/Components/SizeBox.cpp
@@ -238,7 +238,7 @@ const FSlateBrush* USizeBox::GetEditorIcon()
 
 const FText USizeBox::GetPaletteCategory()
 {
-	return LOCTEXT("Common", "Common");
+	return LOCTEXT("Panel", "Panel");
 }
 
 #endif
diff --git a/Engine/Source/Runtime/UMG/Private/Components/SpinBox.cpp b/Engine/Source/Runtime/UMG/Private/Components/SpinBox.cpp
--- a/Engine/Source/Runtime/UMG/Private/Components/SpinBox.cpp
+++ b/Engine/Source/Runtime/UMG/Private/Components/SpinBox.cpp
@@ -315,7 +315,7 @@ const FSlateBrush* USpinBox::GetEditorIcon()
 
 const FText USpinBox::GetPaletteCategory()
 {
-	return LOCTEXT("Common", "Common");
+	return LOCTEXT("Input", "Input");
 }
 
 #endif
diff --git a/Engine/Source/Runtime/UMG/Public/Components/EditableText.h b/Engine/Source/Runtime/UMG/Public/Components/EditableText.h
--- a/Engine/Source/Runtime/UMG/Public/Components/EditableText.h
+++ b/Engine/Source/Runtime/UMG/Public/Components/EditableText.h
@@ -9,7 +9,7 @@
 /**
  * Editable text box widget
  */
-UCLASS(meta=( Category="Primitive" ), ClassGroup=UserInterface)
+UCLASS()
 class UMG_API UEditableText : public UWidget
 {
 	GENERATED_UCLASS_BODY()
diff --git a/Engine/Source/Runtime/UMG/Public/Components/EditableTextBox.h b/Engine/Source/Runtime/UMG/Public/Components/EditableTextBox.h
--- a/Engine/Source/Runtime/UMG/Public/Components/EditableTextBox.h
+++ b/Engine/Source/Runtime/UMG/Public/Components/EditableTextBox.h
@@ -4,8 +4,10 @@
 
 #include "EditableTextBox.generated.h"
 
-/** Editable text box widget */
-UCLASS(ClassGroup=UserInterface)
+/**
+ * Allows the user to type in custom text.  Only permits a single line of text to be entered.
+ */
+UCLASS(meta=(DisplayName="Text Box"))
 class UMG_API UEditableTextBox : public UWidget
 {
 	GENERATED_UCLASS_BODY()
diff --git a/Engine/Source/Runtime/UMG/Public/Components/MultiLineEditableText.h b/Engine/Source/Runtime/UMG/Public/Components/MultiLineEditableText.h
--- a/Engine/Source/Runtime/UMG/Public/Components/MultiLineEditableText.h
+++ b/Engine/Source/Runtime/UMG/Public/Components/MultiLineEditableText.h
@@ -7,7 +7,7 @@
 /**
  * Editable text box widget
  */
-UCLASS(ClassGroup=UserInterface)
+UCLASS(meta=( DisplayName="Editable Text (Multi-Line)" ))
 class UMG_API UMultiLineEditableText : public UWidget
 {
 	GENERATED_UCLASS_BODY()
diff --git a/Engine/Source/Runtime/UMG/Public/Components/MultiLineEditableTextBox.h b/Engine/Source/Runtime/UMG/Public/Components/MultiLineEditableTextBox.h
--- a/Engine/Source/Runtime/UMG/Public/Components/MultiLineEditableTextBox.h
+++ b/Engine/Source/Runtime/UMG/Public/Components/MultiLineEditableTextBox.h
@@ -4,8 +4,10 @@
 
 #include "MultiLineEditableTextBox.generated.h"
 
-/** Editable text box widget */
-UCLASS(ClassGroup=UserInterface)
+/**
+ * Allows a user to enter multiple lines of text
+ */
+UCLASS(meta=(DisplayName="Text Box (Multi-Line)"))
 class UMG_API UMultiLineEditableTextBox : public UWidget
 {
 	GENERATED_UCLASS_BODY()
diff --git a/Engine/Source/Runtime/UMG/Public/Components/TextBlock.h b/Engine/Source/Runtime/UMG/Public/Components/TextBlock.h
--- a/Engine/Source/Runtime/UMG/Public/Components/TextBlock.h
+++ b/Engine/Source/Runtime/UMG/Public/Components/TextBlock.h
@@ -7,7 +7,7 @@
 /**
  * A simple static text widget
  */
-UCLASS(ClassGroup=UserInterface)
+UCLASS(meta=(DisplayName="Text"))
 class UMG_API UTextBlock : public UWidget
 {
 	GENERATED_UCLASS_BODY()
