UMG - Removing the destruction event from user widgets, not ready for prime time.

diff --git a/Engine/Source/Runtime/UMG/Private/UserWidget.cpp b/Engine/Source/Runtime/UMG/Private/UserWidget.cpp
--- a/Engine/Source/Runtime/UMG/Private/UserWidget.cpp
+++ b/Engine/Source/Runtime/UMG/Private/UserWidget.cpp
@@ -78,12 +78,6 @@ void UUserWidget::ReleaseSlateResources(bool bReleaseChildren)
 {
 	Super::ReleaseSlateResources(bReleaseChildren);
 
-	if ( !HasAnyFlags(RF_ClassDefaultObject) )
-	{
-		// Notify the blueprint that it's been removed from the slate hierarchy.
-		//Destruct();
-	}
-
 	UWidget* RootWidget = GetRootWidgetComponent();
 	if ( RootWidget )
 	{
@@ -174,11 +168,6 @@ void UUserWidget::Construct_Implementation()
 
 }
 
-void UUserWidget::Destruct_Implementation()
-{
-
-}
-
 void UUserWidget::Tick_Implementation(FGeometry MyGeometry, float InDeltaTime)
 {
 
diff --git a/Engine/Source/Runtime/UMG/Public/Blueprint/UserWidget.h b/Engine/Source/Runtime/UMG/Public/Blueprint/UserWidget.h
--- a/Engine/Source/Runtime/UMG/Public/Blueprint/UserWidget.h
+++ b/Engine/Source/Runtime/UMG/Public/Blueprint/UserWidget.h
@@ -210,13 +210,6 @@ public:
 	UFUNCTION(BlueprintNativeEvent, Category="User Interface", meta=(Keywords="Begin Play"))
 	void Construct();
 
-	/**
-	 * Whenever the underlying slate widgets are torn down.  This is called to allow you to do any last minute
-	 * destruction work.
-	 */
-	UFUNCTION(BlueprintNativeEvent, Category="User Interface", meta=(Keywords="End Play"))
-	void Destruct();
-
 	UFUNCTION(BlueprintNativeEvent, Category="User Interface")
 	void Tick(FGeometry MyGeometry, float InDeltaTime);
 
