UMG - Exposing a ton of widget properties for readonly purposes to blueprints.

diff --git a/Engine/Source/Runtime/UMG/Private/Components/UniformGridPanel.cpp b/Engine/Source/Runtime/UMG/Private/Components/UniformGridPanel.cpp
--- a/Engine/Source/Runtime/UMG/Private/Components/UniformGridPanel.cpp
+++ b/Engine/Source/Runtime/UMG/Private/Components/UniformGridPanel.cpp
@@ -52,10 +52,7 @@ void UUniformGridPanel::OnSlotRemoved(UPanelSlot* Slot)
 
 TSharedRef<SWidget> UUniformGridPanel::RebuildWidget()
 {
-	MyUniformGridPanel =
-		SNew(SUniformGridPanel)
-		.MinDesiredSlotWidth(MinDesiredSlotWidth)
-		.MinDesiredSlotHeight(MinDesiredSlotHeight);
+	MyUniformGridPanel = SNew(SUniformGridPanel);
 
 	for ( UPanelSlot* Slot : Slots )
 	{
@@ -74,11 +71,40 @@ UUniformGridSlot* UUniformGridPanel::AddChildToUniformGrid(UWidget* Content)
 	return Cast<UUniformGridSlot>(Super::AddChild(Content));
 }
 
+void UUniformGridPanel::SetSlotPadding(FMargin InSlotPadding)
+{
+	SlotPadding = InSlotPadding;
+	if ( MyUniformGridPanel.IsValid() )
+	{
+		MyUniformGridPanel->SetSlotPadding(InSlotPadding);
+	}
+}
+
+void UUniformGridPanel::SetMinDesiredSlotWidth(float InMinDesiredSlotWidth)
+{
+	MinDesiredSlotWidth = InMinDesiredSlotWidth;
+	if ( MyUniformGridPanel.IsValid() )
+	{
+		MyUniformGridPanel->SetMinDesiredSlotWidth(InMinDesiredSlotWidth);
+	}
+}
+
+void UUniformGridPanel::SetMinDesiredSlotHeight(float InMinDesiredSlotHeight)
+{
+	MinDesiredSlotHeight = InMinDesiredSlotHeight;
+	if ( MyUniformGridPanel.IsValid() )
+	{
+		MyUniformGridPanel->SetMinDesiredSlotHeight(InMinDesiredSlotHeight);
+	}
+}
+
 void UUniformGridPanel::SynchronizeProperties()
 {
 	Super::SynchronizeProperties();
 
 	MyUniformGridPanel->SetSlotPadding(SlotPadding);
+	MyUniformGridPanel->SetMinDesiredSlotWidth(MinDesiredSlotWidth);
+	MyUniformGridPanel->SetMinDesiredSlotHeight(MinDesiredSlotHeight);
 }
 
 #if WITH_EDITOR
diff --git a/Engine/Source/Runtime/UMG/Private/Components/WrapBox.cpp b/Engine/Source/Runtime/UMG/Private/Components/WrapBox.cpp
--- a/Engine/Source/Runtime/UMG/Private/Components/WrapBox.cpp
+++ b/Engine/Source/Runtime/UMG/Private/Components/WrapBox.cpp
@@ -79,6 +79,16 @@ void UWrapBox::SynchronizeProperties()
 	MyWrapBox->SetInnerSlotPadding(InnerSlotPadding);
 }
 
+void UWrapBox::SetInnerSlotPadding(FVector2D InPadding)
+{
+	InnerSlotPadding = InPadding;
+
+	if ( MyWrapBox.IsValid() )
+	{
+		MyWrapBox->SetInnerSlotPadding(InPadding);
+	}
+}
+
 #if WITH_EDITOR
 
 const FSlateBrush* UWrapBox::GetEditorIcon()
diff --git a/Engine/Source/Runtime/UMG/Public/Components/Border.h b/Engine/Source/Runtime/UMG/Public/Components/Border.h
--- a/Engine/Source/Runtime/UMG/Public/Components/Border.h
+++ b/Engine/Source/Runtime/UMG/Public/Components/Border.h
@@ -10,7 +10,7 @@ class USlateBrushAsset;
  * A border is a container widget that can contain one child widget, providing an opportunity 
  * to surround it with a border image and adjustable padding.
  */
-UCLASS(ClassGroup=UserInterface)
+UCLASS()
 class UMG_API UBorder : public UContentWidget
 {
 	GENERATED_UCLASS_BODY()
@@ -18,7 +18,7 @@ class UMG_API UBorder : public UContentWidget
 public:
 
 	/** Color and opacity multiplier of content in the border */
-	UPROPERTY(EditDefaultsOnly, Category="Content")
+	UPROPERTY(EditDefaultsOnly, BlueprintReadOnly, Category="Content")
 	FLinearColor ContentColorAndOpacity;
 
 	/** A bindable delegate for the ContentColorAndOpacity. */
@@ -26,19 +26,19 @@ public:
 	FGetLinearColor ContentColorAndOpacityDelegate;
 
 	/** The padding area between the slot and the content it contains. */
-	UPROPERTY(EditDefaultsOnly, Category="Content")
+	UPROPERTY(EditDefaultsOnly, BlueprintReadOnly, Category="Content")
 	FMargin Padding;
 
 	/** The alignment of the content horizontally. */
-	UPROPERTY(EditDefaultsOnly, Category="Content")
+	UPROPERTY(EditDefaultsOnly, BlueprintReadOnly, Category="Content")
 	TEnumAsByte<EHorizontalAlignment> HorizontalAlignment;
 
 	/** The alignment of the content vertically. */
-	UPROPERTY(EditDefaultsOnly, Category="Content")
+	UPROPERTY(EditDefaultsOnly, BlueprintReadOnly, Category="Content")
 	TEnumAsByte<EVerticalAlignment> VerticalAlignment;
 
 	/** Brush to drag as the background */
-	UPROPERTY(EditDefaultsOnly, Category=Appearance, meta=(DisplayName="Brush"))
+	UPROPERTY(EditDefaultsOnly, BlueprintReadOnly, Category=Appearance, meta=( DisplayName="Brush" ))
 	FSlateBrush Background;
 
 	/** A bindable delegate for the Brush. */
@@ -46,7 +46,7 @@ public:
 	FGetSlateBrush BackgroundDelegate;
 
 	/** Color and opacity of the actual border image */
-	UPROPERTY(EditDefaultsOnly, Category=Appearance)
+	UPROPERTY(EditDefaultsOnly, BlueprintReadOnly, Category=Appearance)
 	FLinearColor BrushColor;
 
 	/** A bindable delegate for the BrushColor. */
@@ -54,9 +54,11 @@ public:
 	FGetLinearColor BrushColorDelegate;
 
 	/** Whether or not to show the disabled effect when this border is disabled */
-	UPROPERTY(EditDefaultsOnly, Category=Appearance, AdvancedDisplay)
+	UPROPERTY(EditDefaultsOnly, BlueprintReadOnly, Category=Appearance, AdvancedDisplay)
 	bool bShowEffectWhenDisabled;
 
+public:
+
 	UPROPERTY(EditDefaultsOnly, Category=Events)
 	FOnPointerEvent OnMouseButtonDownEvent;
 
diff --git a/Engine/Source/Runtime/UMG/Public/Components/Button.h b/Engine/Source/Runtime/UMG/Public/Components/Button.h
--- a/Engine/Source/Runtime/UMG/Public/Components/Button.h
+++ b/Engine/Source/Runtime/UMG/Public/Components/Button.h
@@ -11,7 +11,7 @@ DECLARE_DYNAMIC_MULTICAST_DELEGATE(FOnButtonClickedEvent);
 /**
  * The button is a clickable primitive widget to enable basic interaction.
  */
-UCLASS(ClassGroup=UserInterface)
+UCLASS()
 class UMG_API UButton : public UContentWidget
 {
 	GENERATED_UCLASS_BODY()
@@ -26,23 +26,23 @@ public:
 	FButtonStyle WidgetStyle;
 	
 	/** The color multiplier for the button content */
-	UPROPERTY(EditDefaultsOnly, Category=Appearance )
+	UPROPERTY(EditDefaultsOnly, BlueprintReadOnly, Category=Appearance)
 	FLinearColor ColorAndOpacity;
 	
 	/** The color multiplier for the button background */
-	UPROPERTY(EditDefaultsOnly, Category=Appearance )
+	UPROPERTY(EditDefaultsOnly, BlueprintReadOnly, Category=Appearance)
 	FLinearColor BackgroundColor;
 
 	/** The type of mouse action required by the user to trigger the buttons 'Click' */
-	UPROPERTY(EditDefaultsOnly, Category="Interaction", AdvancedDisplay)
+	UPROPERTY(EditDefaultsOnly, BlueprintReadOnly, Category="Interaction", AdvancedDisplay)
 	TEnumAsByte<EButtonClickMethod::Type> ClickMethod;
 
 	/** The type of touch action required by the user to trigger the buttons 'Click' */
-	UPROPERTY(EditDefaultsOnly, Category="Interaction", AdvancedDisplay)
+	UPROPERTY(EditDefaultsOnly, BlueprintReadOnly, Category="Interaction", AdvancedDisplay)
 	TEnumAsByte<EButtonTouchMethod::Type> TouchMethod;
 
 	/** Sometimes a button should only be mouse-clickable and never keyboard focusable. */
-	UPROPERTY(EditDefaultsOnly, Category="Interaction", AdvancedDisplay)
+	UPROPERTY(EditDefaultsOnly, BlueprintReadOnly, Category="Interaction", AdvancedDisplay)
 	bool IsFocusable;
 
 public:
@@ -50,6 +50,8 @@ public:
 	/** Called when the button is clicked */
 	UPROPERTY(BlueprintAssignable)
 	FOnButtonClickedEvent OnClicked;
+
+public:
 	
 	/** Sets the color multiplier for the button content */
 	UFUNCTION(BlueprintCallable, Category="Button|Appearance")
@@ -67,6 +69,8 @@ public:
 	UFUNCTION(BlueprintCallable, Category="Button")
 	bool IsPressed() const;
 
+public:
+
 	// UWidget interface
 	virtual void SynchronizeProperties() override;
 	// End of UWidget interface
diff --git a/Engine/Source/Runtime/UMG/Public/Components/ButtonSlot.h b/Engine/Source/Runtime/UMG/Public/Components/ButtonSlot.h
--- a/Engine/Source/Runtime/UMG/Public/Components/ButtonSlot.h
+++ b/Engine/Source/Runtime/UMG/Public/Components/ButtonSlot.h
@@ -26,6 +26,8 @@ public:
 	UPROPERTY(EditDefaultsOnly, Category="Layout (Button Slot)")
 	TEnumAsByte<EVerticalAlignment> VerticalAlignment;
 
+public:
+
 	UFUNCTION(BlueprintCallable, Category="Layout (Button Slot)")
 	void SetPadding(FMargin InPadding);
 
diff --git a/Engine/Source/Runtime/UMG/Public/Components/CanvasPanel.h b/Engine/Source/Runtime/UMG/Public/Components/CanvasPanel.h
--- a/Engine/Source/Runtime/UMG/Public/Components/CanvasPanel.h
+++ b/Engine/Source/Runtime/UMG/Public/Components/CanvasPanel.h
@@ -12,7 +12,7 @@ class UCanvasPanelSlot;
  * for manual layout, but bad when you want to procedurally just generate widgets and place them in a 
  * container (unless you want absolute layout).
  */
-UCLASS(ClassGroup=UserInterface)
+UCLASS()
 class UMG_API UCanvasPanel : public UPanelWidget
 {
 	GENERATED_UCLASS_BODY()
diff --git a/Engine/Source/Runtime/UMG/Public/Components/CanvasPanelSlot.h b/Engine/Source/Runtime/UMG/Public/Components/CanvasPanelSlot.h
--- a/Engine/Source/Runtime/UMG/Public/Components/CanvasPanelSlot.h
+++ b/Engine/Source/Runtime/UMG/Public/Components/CanvasPanelSlot.h
@@ -34,15 +34,15 @@ class UMG_API UCanvasPanelSlot : public UPanelSlot
 	GENERATED_UCLASS_BODY()
 
 	/** The anchoring information for the slot */
-	UPROPERTY(EditDefaultsOnly, Category="Layout (Canvas Slot)")
+	UPROPERTY(EditDefaultsOnly, BlueprintReadOnly, Category="Layout (Canvas Slot)")
 	FAnchorData LayoutData;
 
 	/** When AutoSize is true we use the widget's desired size */
-	UPROPERTY(EditDefaultsOnly, Category="Layout (Canvas Slot)", AdvancedDisplay, meta=(DisplayName="Size To Content"))
+	UPROPERTY(EditDefaultsOnly, BlueprintReadOnly, Category="Layout (Canvas Slot)", AdvancedDisplay, meta=( DisplayName="Size To Content" ))
 	bool bAutoSize;
 
 	/** The order priority this widget is rendered in.  Higher values are rendered last (and so they will appear to be on top). */
-	UPROPERTY(EditDefaultsOnly, Category="Layout (Canvas Slot)")
+	UPROPERTY(EditDefaultsOnly, BlueprintReadOnly, Category="Layout (Canvas Slot)")
 	int32 ZOrder;
 
 	void BuildSlot(TSharedRef<SConstraintCanvas> Canvas);
diff --git a/Engine/Source/Runtime/UMG/Public/Components/CheckBox.h b/Engine/Source/Runtime/UMG/Public/Components/CheckBox.h
--- a/Engine/Source/Runtime/UMG/Public/Components/CheckBox.h
+++ b/Engine/Source/Runtime/UMG/Public/Components/CheckBox.h
@@ -16,6 +16,15 @@ class UMG_API UCheckBox : public UContentWidget
 {
 	GENERATED_UCLASS_BODY()
 
+public:
+	/** Whether the check box is currently in a checked state */
+	UPROPERTY(EditDefaultsOnly, Category=Appearance)
+	TEnumAsByte<ESlateCheckBoxState::Type> CheckedState;
+
+	/** A bindable delegate for the IsChecked. */
+	UPROPERTY()
+	FGetCheckBoxState CheckedStateDelegate;
+
 public:
 	/** The checkbox bar style */
 	UPROPERTY(EditAnywhere, BlueprintReadWrite, Category="Style", meta=( DisplayName="Style" ))
@@ -61,24 +70,16 @@ public:
 	UPROPERTY()
 	USlateBrushAsset* UndeterminedPressedImage_DEPRECATED;
 
-	/** Whether the check box is currently in a checked state */
-	UPROPERTY(EditDefaultsOnly, Category=Appearance)
-	TEnumAsByte<ESlateCheckBoxState::Type> CheckedState;
-
-	/** A bindable delegate for the IsChecked. */
-	UPROPERTY()
-	FGetCheckBoxState CheckedStateDelegate;
-
 	/** How the content of the toggle button should align within the given space */
-	UPROPERTY(EditDefaultsOnly, Category=Appearance)
+	UPROPERTY(EditDefaultsOnly, BlueprintReadOnly, Category=Appearance)
 	TEnumAsByte<EHorizontalAlignment> HorizontalAlignment;
 
 	/** Spacing between the check box image and its content */
-	UPROPERTY(EditDefaultsOnly, Category=Appearance)
+	UPROPERTY(EditDefaultsOnly, BlueprintReadOnly, Category=Appearance)
 	FMargin Padding;
 
 	/** The color of the background border */
-	UPROPERTY(EditDefaultsOnly, Category=Appearance)
+	UPROPERTY(EditDefaultsOnly, BlueprintReadOnly, Category=Appearance)
 	FSlateColor BorderBackgroundColor;
 
 public:
@@ -87,6 +88,8 @@ public:
 	UPROPERTY(BlueprintAssignable)
 	FOnCheckBoxComponentStateChanged OnCheckStateChanged;
 
+public:
+
 	/** Returns true if this button is currently pressed */
 	UFUNCTION(BlueprintCallable, Category="Widget")
 	bool IsPressed() const;
@@ -106,6 +109,8 @@ public:
 	/** Sets the checked state. */
 	UFUNCTION(BlueprintCallable, Category="Widget")
 	void SetCheckedState(ESlateCheckBoxState::Type InCheckedState);
+
+public:
 	
 	// UWidget interface
 	virtual void SynchronizeProperties() override;
diff --git a/Engine/Source/Runtime/UMG/Public/Components/CircularThrobber.h b/Engine/Source/Runtime/UMG/Public/Components/CircularThrobber.h
--- a/Engine/Source/Runtime/UMG/Public/Components/CircularThrobber.h
+++ b/Engine/Source/Runtime/UMG/Public/Components/CircularThrobber.h
@@ -9,11 +9,32 @@ class USlateBrushAsset;
 /**
  * A throbber widget that orients images in a spinning circle.
  */
-UCLASS(ClassGroup=UserInterface)
+UCLASS()
 class UMG_API UCircularThrobber : public UWidget
 {
 	GENERATED_UCLASS_BODY()
 
+public:
+
+	/** How many pieces there are */
+	UPROPERTY(EditDefaultsOnly, BlueprintReadOnly, Category=Appearance, meta=( ClampMin = "1", UIMin = "1", UIMax = "8" ))
+	int32 NumberOfPieces;
+
+	/** The amount of time for a full circle (in seconds) */
+	UPROPERTY(EditDefaultsOnly, BlueprintReadOnly, Category=Appearance)
+	float Period;
+
+	/** The radius of the circle */
+	UPROPERTY(EditDefaultsOnly, BlueprintReadOnly, Category=Appearance)
+	float Radius;
+
+	/** Image to use for each segment of the throbber */
+	UPROPERTY()
+	USlateBrushAsset* PieceImage_DEPRECATED;
+
+	UPROPERTY(EditDefaultsOnly, BlueprintReadOnly, Category=Appearance)
+	FSlateBrush Image;
+
 public:
 
 	/** Sets how many pieces there are. */
@@ -45,27 +66,6 @@ public:
 	virtual const FText GetPaletteCategory() override;
 #endif
 
-protected:
-
-	/** How many pieces there are */
-	UPROPERTY(EditDefaultsOnly, Category=Appearance, meta=(ClampMin = "1", UIMin = "1", UIMax = "8"))
-	int32 NumberOfPieces;
-
-	/** The amount of time for a full circle (in seconds) */
-	UPROPERTY(EditDefaultsOnly, Category=Appearance)
-	float Period;
-
-	/** The radius of the circle */
-	UPROPERTY(EditDefaultsOnly, Category=Appearance)
-	float Radius;
-
-	/** Image to use for each segment of the throbber */
-	UPROPERTY()
-	USlateBrushAsset* PieceImage_DEPRECATED;
-
-	UPROPERTY(EditDefaultsOnly, Category=Appearance)
-	FSlateBrush Image;
-
 protected:
 	// UWidget interface
 	virtual TSharedRef<SWidget> RebuildWidget() override;
diff --git a/Engine/Source/Runtime/UMG/Public/Components/ComboBox.h b/Engine/Source/Runtime/UMG/Public/Components/ComboBox.h
--- a/Engine/Source/Runtime/UMG/Public/Components/ComboBox.h
+++ b/Engine/Source/Runtime/UMG/Public/Components/ComboBox.h
@@ -7,7 +7,7 @@
 /**
  * The combobox allows you to display a list of options to the user in a dropdown menu for them to select one.
  */
-UCLASS(Experimental, ClassGroup=UserInterface, meta=( DisplayName="ComboBox (Object)" ))
+UCLASS(Experimental, meta=( DisplayName="ComboBox (Object)" ))
 class UMG_API UComboBox : public UWidget
 {
 	GENERATED_UCLASS_BODY()
diff --git a/Engine/Source/Runtime/UMG/Public/Components/ComboBoxString.h b/Engine/Source/Runtime/UMG/Public/Components/ComboBoxString.h
--- a/Engine/Source/Runtime/UMG/Public/Components/ComboBoxString.h
+++ b/Engine/Source/Runtime/UMG/Public/Components/ComboBoxString.h
@@ -19,10 +19,6 @@ class UMG_API UComboBoxString : public UWidget
 
 public:
 
-	/** The style */
-	UPROPERTY(EditAnywhere, BlueprintReadWrite, Category="Style", meta=( DisplayName="Style" ))
-	FComboBoxStyle WidgetStyle;
-
 	/** The list of items to be displayed on the combobox. */
 	UPROPERTY(EditDefaultsOnly, Category=Content)
 	TArray<FString> DefaultOptions;
@@ -31,20 +27,28 @@ public:
 	UPROPERTY(EditDefaultsOnly, Category=Content)
 	FString SelectedOption;
 
-	UPROPERTY(EditDefaultsOnly, Category=Content)
+public:
+
+	/** The style */
+	UPROPERTY(EditAnywhere, BlueprintReadWrite, Category="Style", meta=( DisplayName="Style" ))
+	FComboBoxStyle WidgetStyle;
+
+	UPROPERTY(EditDefaultsOnly, BlueprintReadOnly, Category=Content)
 	FMargin ContentPadding;
 
 	/** The max height of the combobox list that opens */
-	UPROPERTY(EditDefaultsOnly, Category=Content, AdvancedDisplay)
+	UPROPERTY(EditDefaultsOnly, BlueprintReadOnly, Category=Content, AdvancedDisplay)
 	float MaxListHeight;
 
 	/**
 	 * When false, the down arrow is not generated and it is up to the API consumer
 	 * to make their own visual hint that this is a drop down.
 	 */
-	UPROPERTY(EditDefaultsOnly, Category=Content, AdvancedDisplay)
+	UPROPERTY(EditDefaultsOnly, BlueprintReadOnly, Category=Content, AdvancedDisplay)
 	bool HasDownArrow;
 
+public: // Events
+
 	/** Called when the widget is needed for the item. */
 	UPROPERTY(EditDefaultsOnly, Category=Events)
 	FGenerateWidgetForString OnGenerateWidgetEvent;
diff --git a/Engine/Source/Runtime/UMG/Public/Components/EditableText.h b/Engine/Source/Runtime/UMG/Public/Components/EditableText.h
--- a/Engine/Source/Runtime/UMG/Public/Components/EditableText.h
+++ b/Engine/Source/Runtime/UMG/Public/Components/EditableText.h
@@ -20,7 +20,6 @@ public:
 	DECLARE_DYNAMIC_MULTICAST_DELEGATE_TwoParams(FOnEditableTextCommittedEvent, const FText&, Text, ETextCommit::Type, CommitMethod);
 
 public:
-
 	/** The text content for this editable text box widget */
 	UPROPERTY(EditDefaultsOnly, Category=Content)
 	FText Text;
@@ -29,8 +28,10 @@ public:
 	UPROPERTY()
 	FGetText TextDelegate;
 
+public:
+
 	/** Hint text that appears when there is no text in the text box */
-	UPROPERTY(EditDefaultsOnly, Category=Content)
+	UPROPERTY(EditDefaultsOnly, BlueprintReadOnly, Category=Content)
 	FText HintText;
 
 	/** A bindable delegate to allow logic to drive the hint text of the widget */
diff --git a/Engine/Source/Runtime/UMG/Public/Components/EditableTextBox.h b/Engine/Source/Runtime/UMG/Public/Components/EditableTextBox.h
--- a/Engine/Source/Runtime/UMG/Public/Components/EditableTextBox.h
+++ b/Engine/Source/Runtime/UMG/Public/Components/EditableTextBox.h
@@ -17,6 +17,15 @@ public:
 	DECLARE_DYNAMIC_MULTICAST_DELEGATE_OneParam(FOnEditableTextBoxChangedEvent, const FText&, Text);
 	DECLARE_DYNAMIC_MULTICAST_DELEGATE_TwoParams(FOnEditableTextBoxCommittedEvent, const FText&, Text, ETextCommit::Type, CommitMethod);
 
+public:
+	/** The text content for this editable text box widget */
+	UPROPERTY(EditDefaultsOnly, Category=Content)
+	FText Text;
+
+	/** A bindable delegate to allow logic to drive the text of the widget */
+	UPROPERTY()
+	FGetText TextDelegate;
+
 public:
 
 	/** The style */
@@ -27,16 +36,8 @@ public:
 	UPROPERTY()
 	USlateWidgetStyleAsset* Style_DEPRECATED;
 
-	/** The text content for this editable text box widget */
-	UPROPERTY(EditDefaultsOnly, Category=Content)
-	FText Text;
-
-	/** A bindable delegate to allow logic to drive the text of the widget */
-	UPROPERTY()
-	FGetText TextDelegate;
-
 	/** Hint text that appears when there is no text in the text box */
-	UPROPERTY(EditDefaultsOnly, Category=Content)
+	UPROPERTY(EditDefaultsOnly, BlueprintReadOnly, Category=Content)
 	FText HintText;
 
 	/** A bindable delegate to allow logic to drive the hint text of the widget */
@@ -44,55 +45,55 @@ public:
 	FGetText HintTextDelegate;
 
 	/** Font color and opacity (overrides Style) */
-	UPROPERTY(EditDefaultsOnly, Category=Appearance)
+	UPROPERTY(EditDefaultsOnly, BlueprintReadOnly, Category=Appearance)
 	FSlateFontInfo Font;
 
 	/** Text color and opacity (overrides Style) */
-	UPROPERTY(EditDefaultsOnly, Category=Appearance)
+	UPROPERTY(EditDefaultsOnly, BlueprintReadOnly, Category=Appearance)
 	FLinearColor ForegroundColor;
 
 	/** The color of the background/border around the editable text (overrides Style) */
-	UPROPERTY(EditDefaultsOnly, Category=Appearance)
+	UPROPERTY(EditDefaultsOnly, BlueprintReadOnly, Category=Appearance)
 	FLinearColor BackgroundColor;
 
 	/** Text color and opacity when read-only (overrides Style) */
-	UPROPERTY(EditDefaultsOnly, Category=Appearance)
+	UPROPERTY(EditDefaultsOnly, BlueprintReadOnly, Category=Appearance)
 	FLinearColor ReadOnlyForegroundColor;
 
 	/** Sets whether this text box can actually be modified interactively by the user */
-	UPROPERTY(EditDefaultsOnly, Category=Appearance)
+	UPROPERTY(EditDefaultsOnly, BlueprintReadOnly, Category=Appearance)
 	bool IsReadOnly;
 
 	/** Sets whether this text box is for storing a password */
-	UPROPERTY(EditDefaultsOnly, Category=Appearance)
+	UPROPERTY(EditDefaultsOnly, BlueprintReadOnly, Category=Appearance)
 	bool IsPassword;
 
 	/** Minimum width that a text block should be */
-	UPROPERTY(EditDefaultsOnly, Category=Appearance)
+	UPROPERTY(EditDefaultsOnly, BlueprintReadOnly, Category=Appearance)
 	float MinimumDesiredWidth;
 
 	/** Padding between the box/border and the text widget inside (overrides Style) */
-	UPROPERTY(EditDefaultsOnly, Category=Appearance)
+	UPROPERTY(EditDefaultsOnly, BlueprintReadOnly, Category=Appearance)
 	FMargin Padding;
 
 	/** Workaround as we loose focus when the auto completion closes. */
-	UPROPERTY(EditDefaultsOnly, Category=Behavior, AdvancedDisplay)
+	UPROPERTY(EditDefaultsOnly, BlueprintReadOnly, Category=Behavior, AdvancedDisplay)
 	bool IsCaretMovedWhenGainFocus;
 
 	/** Whether to select all text when the user clicks to give focus on the widget */
-	UPROPERTY(EditDefaultsOnly, Category=Behavior, AdvancedDisplay)
+	UPROPERTY(EditDefaultsOnly, BlueprintReadOnly, Category=Behavior, AdvancedDisplay)
 	bool SelectAllTextWhenFocused;
 
 	/** Whether to allow the user to back out of changes when they press the escape key */
-	UPROPERTY(EditDefaultsOnly, Category=Behavior, AdvancedDisplay)
+	UPROPERTY(EditDefaultsOnly, BlueprintReadOnly, Category=Behavior, AdvancedDisplay)
 	bool RevertTextOnEscape;
 
 	/** Whether to clear keyboard focus when pressing enter to commit changes */
-	UPROPERTY(EditDefaultsOnly, Category=Behavior, AdvancedDisplay)
+	UPROPERTY(EditDefaultsOnly, BlueprintReadOnly, Category=Behavior, AdvancedDisplay)
 	bool ClearKeyboardFocusOnCommit;
 
 	/** Whether to select all text when pressing enter to commit changes */
-	UPROPERTY(EditDefaultsOnly, Category=Behavior, AdvancedDisplay)
+	UPROPERTY(EditDefaultsOnly, BlueprintReadOnly, Category=Behavior, AdvancedDisplay)
 	bool SelectAllTextOnCommit;
 
 	/** Called whenever the text is changed interactively by the user */
diff --git a/Engine/Source/Runtime/UMG/Public/Components/GridPanel.h b/Engine/Source/Runtime/UMG/Public/Components/GridPanel.h
--- a/Engine/Source/Runtime/UMG/Public/Components/GridPanel.h
+++ b/Engine/Source/Runtime/UMG/Public/Components/GridPanel.h
@@ -7,7 +7,7 @@
 /**
  * A panel that evenly divides up available space between all of its children.
  */
-UCLASS(ClassGroup=UserInterface)
+UCLASS()
 class UMG_API UGridPanel : public UPanelWidget
 {
 	GENERATED_UCLASS_BODY()
@@ -15,11 +15,11 @@ class UMG_API UGridPanel : public UPanelWidget
 public:
 
 	/** The column fill rules */
-	UPROPERTY(EditDefaultsOnly, Category="Fill Rules")
+	UPROPERTY(EditDefaultsOnly, BlueprintReadOnly, Category="Fill Rules")
 	TArray<float> ColumnFill;
 
 	/** The row fill rules */
-	UPROPERTY(EditDefaultsOnly, Category="Fill Rules")
+	UPROPERTY(EditDefaultsOnly, BlueprintReadOnly, Category="Fill Rules")
 	TArray<float> RowFill;
 
 public:
diff --git a/Engine/Source/Runtime/UMG/Public/Components/GridSlot.h b/Engine/Source/Runtime/UMG/Public/Components/GridSlot.h
--- a/Engine/Source/Runtime/UMG/Public/Components/GridSlot.h
+++ b/Engine/Source/Runtime/UMG/Public/Components/GridSlot.h
@@ -18,35 +18,35 @@ class UMG_API UGridSlot : public UPanelSlot
 public:
 
 	/** The alignment of the object horizontally. */
-	UPROPERTY(EditDefaultsOnly, Category="Layout (Grid Slot)")
+	UPROPERTY(EditDefaultsOnly, BlueprintReadOnly, Category="Layout (Grid Slot)")
 	TEnumAsByte<EHorizontalAlignment> HorizontalAlignment;
 
 	/** The alignment of the object vertically. */
-	UPROPERTY(EditDefaultsOnly, Category="Layout (Grid Slot)")
+	UPROPERTY(EditDefaultsOnly, BlueprintReadOnly, Category="Layout (Grid Slot)")
 	TEnumAsByte<EVerticalAlignment> VerticalAlignment;
 	
 	/** The row index of the cell this slot is in */
-	UPROPERTY(EditDefaultsOnly, meta=( UIMin = "0" ), Category="Layout (Grid Slot)")
+	UPROPERTY(EditDefaultsOnly, BlueprintReadOnly, meta=( UIMin = "0" ), Category="Layout (Grid Slot)")
 	int32 Row;
 	
 	/**  */
-	UPROPERTY(EditDefaultsOnly, Category="Layout (Grid Slot)")
+	UPROPERTY(EditDefaultsOnly, BlueprintReadOnly, Category="Layout (Grid Slot)")
 	int32 RowSpan;
 
 	/** The column index of the cell this slot is in */
-	UPROPERTY(EditDefaultsOnly, meta=( UIMin = "0" ), Category="Layout (Grid Slot)")
+	UPROPERTY(EditDefaultsOnly, BlueprintReadOnly, meta=( UIMin = "0" ), Category="Layout (Grid Slot)")
 	int32 Column;
 
 	/**  */
-	UPROPERTY(EditDefaultsOnly, Category="Layout (Grid Slot)")
+	UPROPERTY(EditDefaultsOnly, BlueprintReadOnly, Category="Layout (Grid Slot)")
 	int32 ColumnSpan;
 
 	/** Positive values offset this cell to be hit-tested and drawn on top of others. Default is 0; i.e. no offset. */
-	UPROPERTY(EditDefaultsOnly, Category="Layout (Grid Slot)")
+	UPROPERTY(EditDefaultsOnly, BlueprintReadOnly, Category="Layout (Grid Slot)")
 	int32 Layer;
 
 	/** Offset this slot's content by some amount; positive values offset to lower right*/
-	UPROPERTY(EditDefaultsOnly, Category="Layout (Grid Slot)")
+	UPROPERTY(EditDefaultsOnly, BlueprintReadOnly, Category="Layout (Grid Slot)")
 	FVector2D Nudge;
 
 	/** Sets the row index of the slot, this determines what cell the slot is in the panel */
diff --git a/Engine/Source/Runtime/UMG/Public/Components/Image.h b/Engine/Source/Runtime/UMG/Public/Components/Image.h
--- a/Engine/Source/Runtime/UMG/Public/Components/Image.h
+++ b/Engine/Source/Runtime/UMG/Public/Components/Image.h
@@ -9,7 +9,7 @@ class USlateBrushAsset;
 /**
  * The image widget allows you to display a Slate Brush, or texture or material in the UI.
  */
-UCLASS(ClassGroup=UserInterface)
+UCLASS()
 class UMG_API UImage : public UWidget
 {
 	GENERATED_UCLASS_BODY()
@@ -21,7 +21,7 @@ public:
 	USlateBrushAsset* Image_DEPRECATED;
 
 	/** Image to draw */
-	UPROPERTY(EditDefaultsOnly, Category=Appearance)
+	UPROPERTY(EditDefaultsOnly, BlueprintReadOnly, Category=Appearance)
 	FSlateBrush Brush;
 
 	/** A bindable delegate for the Image. */
@@ -29,13 +29,15 @@ public:
 	FGetSlateBrush BrushDelegate;
 
 	/** Color and opacity */
-	UPROPERTY(EditDefaultsOnly, Category=Appearance)
+	UPROPERTY(EditDefaultsOnly, BlueprintReadOnly, Category=Appearance)
 	FLinearColor ColorAndOpacity;
 
 	/** A bindable delegate for the ColorAndOpacity. */
 	UPROPERTY()
 	FGetLinearColor ColorAndOpacityDelegate;
 
+public:
+
 	UPROPERTY(EditDefaultsOnly, Category=Events)
 	FOnPointerEvent OnMouseButtonDownEvent;
 
diff --git a/Engine/Source/Runtime/UMG/Public/Components/MenuAnchor.h b/Engine/Source/Runtime/UMG/Public/Components/MenuAnchor.h
--- a/Engine/Source/Runtime/UMG/Public/Components/MenuAnchor.h
+++ b/Engine/Source/Runtime/UMG/Public/Components/MenuAnchor.h
@@ -20,7 +20,7 @@ public:
 	 * If you want to customize the creation of the popup, you should bind a function to OnGetMenuContentEvent 
 	 * instead.
 	 */
-	UPROPERTY(EditDefaultsOnly, Category="Menu")
+	UPROPERTY(EditDefaultsOnly, BlueprintReadOnly, Category="Menu")
 	TSubclassOf<class UUserWidget> MenuClass;
 
 	/** Called when the menu content is requested to allow a more customized handling over what to display */
@@ -28,9 +28,13 @@ public:
 	FGetContent OnGetMenuContentEvent;
 	
 	/** The placement location of the summoned widget. */
-	UPROPERTY(EditDefaultsOnly, Category="Menu")
+	UPROPERTY(EditDefaultsOnly, BlueprintReadOnly, Category="Menu")
 	TEnumAsByte<EMenuPlacement> Placement;
 
+public:
+	//TODO UMG Add Set MenuClass
+	//TODO UMG Add Set Placement
+
 public:
 
 	/**
diff --git a/Engine/Source/Runtime/UMG/Public/Components/MultiLineEditableText.h b/Engine/Source/Runtime/UMG/Public/Components/MultiLineEditableText.h
--- a/Engine/Source/Runtime/UMG/Public/Components/MultiLineEditableText.h
+++ b/Engine/Source/Runtime/UMG/Public/Components/MultiLineEditableText.h
@@ -17,22 +17,23 @@ public:
 	DECLARE_DYNAMIC_MULTICAST_DELEGATE_OneParam(FOnMultiLineEditableTextChangedEvent, const FText&, Text);
 	DECLARE_DYNAMIC_MULTICAST_DELEGATE_TwoParams(FOnMultiLineEditableTextCommittedEvent, const FText&, Text, ETextCommit::Type, CommitMethod);
 
+public:
+	/** The text content for this editable text box widget */
+	UPROPERTY(EditDefaultsOnly, Category=Content, meta=(MultiLine="true"))
+	FText Text;
+
 public:
 
 	/** The style */
 	UPROPERTY(EditAnywhere, BlueprintReadWrite, Category="Style", meta=( DisplayName="Style" ))
 	FTextBlockStyle WidgetStyle;
 
-	/** The text content for this editable text box widget */
-	UPROPERTY(EditDefaultsOnly, Category=Content, meta=(MultiLine="true"))
-	FText Text;
-
 	/** The justification of the text in the multilinebox */
-	UPROPERTY(EditDefaultsOnly, Category=Content)
+	UPROPERTY(EditAnywhere, BlueprintReadWrite, Category=Content)
 	TEnumAsByte<ETextJustify::Type> Justification;
 
 	/** Font color and opacity (overrides Style) */
-	UPROPERTY(EditDefaultsOnly, Category=Appearance)
+	UPROPERTY(EditAnywhere, BlueprintReadWrite, Category=Appearance)
 	FSlateFontInfo Font;
 
 	/** Called whenever the text is changed interactively by the user */
@@ -43,6 +44,8 @@ public:
 	UPROPERTY(BlueprintAssignable, Category="Widget Event")
 	FOnMultiLineEditableTextCommittedEvent OnTextCommitted;
 
+public:
+
 	/**  */
 	UFUNCTION(BlueprintCallable, Category="Widget")
 	FText GetText() const;
diff --git a/Engine/Source/Runtime/UMG/Public/Components/MultiLineEditableTextBox.h b/Engine/Source/Runtime/UMG/Public/Components/MultiLineEditableTextBox.h
--- a/Engine/Source/Runtime/UMG/Public/Components/MultiLineEditableTextBox.h
+++ b/Engine/Source/Runtime/UMG/Public/Components/MultiLineEditableTextBox.h
@@ -17,6 +17,12 @@ public:
 	DECLARE_DYNAMIC_MULTICAST_DELEGATE_OneParam(FOnMultiLineEditableTextBoxChangedEvent, const FText&, Text);
 	DECLARE_DYNAMIC_MULTICAST_DELEGATE_TwoParams(FOnMultiLineEditableTextBoxCommittedEvent, const FText&, Text, ETextCommit::Type, CommitMethod);
 
+public:
+	
+	/** The text content for this editable text box widget */
+	UPROPERTY(EditDefaultsOnly, Category=Content, meta=(MultiLine="true"))
+	FText Text;
+
 public:
 
 	/** The style */
@@ -30,28 +36,24 @@ public:
 	UPROPERTY()
 	USlateWidgetStyleAsset* Style_DEPRECATED;
 
-	/** The text content for this editable text box widget */
-	UPROPERTY(EditDefaultsOnly, Category=Content, meta=(MultiLine="true"))
-	FText Text;
-
 	/** The justification of the text in the multilinebox */
-	UPROPERTY(EditDefaultsOnly, Category=Content)
+	UPROPERTY(EditDefaultsOnly, BlueprintReadOnly, Category=Content)
 	TEnumAsByte<ETextJustify::Type> Justification;
 
 	/** Font color and opacity (overrides Style) */
-	UPROPERTY(EditDefaultsOnly, Category=Appearance)
+	UPROPERTY(EditDefaultsOnly, BlueprintReadOnly, Category=Appearance)
 	FSlateFontInfo Font;
 
 	/** Text color and opacity (overrides Style) */
-	UPROPERTY(EditDefaultsOnly, Category=Appearance)
+	UPROPERTY(EditDefaultsOnly, BlueprintReadOnly, Category=Appearance)
 	FLinearColor ForegroundColor;
 
 	/** The color of the background/border around the editable text (overrides Style) */
-	UPROPERTY(EditDefaultsOnly, Category=Appearance)
+	UPROPERTY(EditDefaultsOnly, BlueprintReadOnly, Category=Appearance)
 	FLinearColor BackgroundColor;
 
 	/** Text color and opacity when read-only (overrides Style) */
-	UPROPERTY(EditDefaultsOnly, Category=Appearance)
+	UPROPERTY(EditDefaultsOnly, BlueprintReadOnly, Category=Appearance)
 	FLinearColor ReadOnlyForegroundColor;
 
 	/** Called whenever the text is changed interactively by the user */
@@ -65,6 +67,8 @@ public:
 	/** Provide a alternative mechanism for error reporting. */
 	//SLATE_ARGUMENT(TSharedPtr<class IErrorReportingWidget>, ErrorReporting)
 
+public:
+
 	/**  */
 	UFUNCTION(BlueprintCallable, Category="Widget")
 	FText GetText() const;
@@ -76,6 +80,14 @@ public:
 	UFUNCTION(BlueprintCallable, Category="Widget")
 	void SetError(FText InError);
 
+	//TODO UMG Add Set ReadOnlyForegroundColor
+	//TODO UMG Add Set BackgroundColor
+	//TODO UMG Add Set ForegroundColor
+	//TODO UMG Add Set Font
+	//TODO UMG Add Set Justification
+
+public:
+
 	// UWidget interface
 	virtual void SynchronizeProperties() override;
 	// End of UWidget interface
diff --git a/Engine/Source/Runtime/UMG/Public/Components/NamedSlot.h b/Engine/Source/Runtime/UMG/Public/Components/NamedSlot.h
--- a/Engine/Source/Runtime/UMG/Public/Components/NamedSlot.h
+++ b/Engine/Source/Runtime/UMG/Public/Components/NamedSlot.h
@@ -8,7 +8,7 @@
  * Allows you to expose an external slot for your user widget.  When others reuse your user control, 
  * they can put whatever they want in this named slot.
  */
-UCLASS(ClassGroup=UserInterface)
+UCLASS()
 class UMG_API UNamedSlot : public UContentWidget
 {
 	GENERATED_UCLASS_BODY()
diff --git a/Engine/Source/Runtime/UMG/Public/Components/NativeWidgetHost.h b/Engine/Source/Runtime/UMG/Public/Components/NativeWidgetHost.h
--- a/Engine/Source/Runtime/UMG/Public/Components/NativeWidgetHost.h
+++ b/Engine/Source/Runtime/UMG/Public/Components/NativeWidgetHost.h
@@ -8,7 +8,7 @@
  * A NativeWidgetHost is a container widget that can contain one child slate widget.  This should
  * be used when all you need is to nest a native widget inside a UMG widget.
  */
-UCLASS(ClassGroup=UserInterface)
+UCLASS()
 class UMG_API UNativeWidgetHost : public UWidget
 {
 	GENERATED_UCLASS_BODY()
diff --git a/Engine/Source/Runtime/UMG/Public/Components/Overlay.h b/Engine/Source/Runtime/UMG/Public/Components/Overlay.h
--- a/Engine/Source/Runtime/UMG/Public/Components/Overlay.h
+++ b/Engine/Source/Runtime/UMG/Public/Components/Overlay.h
@@ -7,7 +7,7 @@
 /**
  * Allows widgets to be stacked on top of each other, uses simple flow layout for content on each layer.
  */
-UCLASS(ClassGroup=UserInterface)
+UCLASS()
 class UMG_API UOverlay : public UPanelWidget
 {
 	GENERATED_UCLASS_BODY()
diff --git a/Engine/Source/Runtime/UMG/Public/Components/OverlaySlot.h b/Engine/Source/Runtime/UMG/Public/Components/OverlaySlot.h
--- a/Engine/Source/Runtime/UMG/Public/Components/OverlaySlot.h
+++ b/Engine/Source/Runtime/UMG/Public/Components/OverlaySlot.h
@@ -15,15 +15,15 @@ class UMG_API UOverlaySlot : public UPanelSlot
 public:
 	
 	/** The padding area between the slot and the content it contains. */
-	UPROPERTY(EditDefaultsOnly, Category="Layout (Overlay Slot)")
+	UPROPERTY(EditDefaultsOnly, BlueprintReadOnly, Category="Layout (Overlay Slot)")
 	FMargin Padding;
 
 	/** The alignment of the object horizontally. */
-	UPROPERTY(EditDefaultsOnly, Category="Layout (Overlay Slot)")
+	UPROPERTY(EditDefaultsOnly, BlueprintReadOnly, Category="Layout (Overlay Slot)")
 	TEnumAsByte<EHorizontalAlignment> HorizontalAlignment;
 
 	/** The alignment of the object vertically. */
-	UPROPERTY(EditDefaultsOnly, Category="Layout (Overlay Slot)")
+	UPROPERTY(EditDefaultsOnly, BlueprintReadOnly, Category="Layout (Overlay Slot)")
 	TEnumAsByte<EVerticalAlignment> VerticalAlignment;
 
 public:
diff --git a/Engine/Source/Runtime/UMG/Public/Components/ProgressBar.h b/Engine/Source/Runtime/UMG/Public/Components/ProgressBar.h
--- a/Engine/Source/Runtime/UMG/Public/Components/ProgressBar.h
+++ b/Engine/Source/Runtime/UMG/Public/Components/ProgressBar.h
@@ -37,14 +37,14 @@ public:
 	USlateBrushAsset* MarqueeImage_DEPRECATED;
 
 	/** Defines if this progress bar fills Left to right or right to left */
-	UPROPERTY(EditDefaultsOnly, Category=Appearance)
+	UPROPERTY(EditDefaultsOnly, BlueprintReadOnly, Category=Appearance)
 	TEnumAsByte<EProgressBarFillType::Type> BarFillType;
 	
-	UPROPERTY(EditDefaultsOnly, Category=Appearance)
+	UPROPERTY(EditDefaultsOnly, BlueprintReadOnly, Category=Appearance)
 	bool bIsMarquee;
 
 	/** Used to determine the fill position of the progress bar ranging 0..1 */
-	UPROPERTY(EditDefaultsOnly, Category=Appearance, meta=( UIMin = "0", UIMax = "1" ))
+	UPROPERTY(EditDefaultsOnly, BlueprintReadOnly, Category=Appearance, meta=( UIMin = "0", UIMax = "1" ))
 	float Percent;
 
 	/** A bindable delegate to allow logic to drive the text of the widget */
@@ -52,7 +52,7 @@ public:
 	FGetFloat PercentDelegate;
 
 	/** Fill Color and Opacity */
-	UPROPERTY(EditDefaultsOnly, Category=Appearance)
+	UPROPERTY(EditDefaultsOnly, BlueprintReadOnly, Category=Appearance)
 	FLinearColor FillColorAndOpacity;
 
 	/** */
@@ -72,6 +72,10 @@ public:
 	/** Sets the progress bar to show as a marquee. */
 	UFUNCTION(BlueprintCallable, Category="Behavior")
 	void SetIsMarquee(bool InbIsMarquee);
+
+	//TODO UMG Add Set BarFillType.
+
+public:
 	
 	// UWidget interface
 	virtual void SynchronizeProperties() override;
diff --git a/Engine/Source/Runtime/UMG/Public/Components/ScaleBox.h b/Engine/Source/Runtime/UMG/Public/Components/ScaleBox.h
--- a/Engine/Source/Runtime/UMG/Public/Components/ScaleBox.h
+++ b/Engine/Source/Runtime/UMG/Public/Components/ScaleBox.h
@@ -9,7 +9,7 @@
  * you needed to have a background image scale to fill an area but not become distorted with different aspect ratios, or if you need
  * to auto fit some text to an area, this is the control for you.
  */
-UCLASS(ClassGroup=UserInterface)
+UCLASS()
 class UMG_API UScaleBox : public UContentWidget
 {
 	GENERATED_UCLASS_BODY()
@@ -17,13 +17,18 @@ class UMG_API UScaleBox : public UContentWidget
 public:
 
 	/** Controls in what direction content can be scaled */
-	UPROPERTY(EditDefaultsOnly, Category="Stretching")
+	UPROPERTY(EditDefaultsOnly, BlueprintReadOnly, Category="Stretching")
 	TEnumAsByte<EStretchDirection::Type> StretchDirection;
 
 	/** The stretching rule to apply when content is stretched */
-	UPROPERTY(EditDefaultsOnly, Category="Stretching")
+	UPROPERTY(EditDefaultsOnly, BlueprintReadOnly, Category="Stretching")
 	TEnumAsByte<EStretch::Type> Stretch;
 
+public:
+
+	//TODO UMG Add Set Stretch
+	//TODO UMG Add Set Stretch DIrection
+
 public:
 
 	// UWidget interface
diff --git a/Engine/Source/Runtime/UMG/Public/Components/ScaleBoxSlot.h b/Engine/Source/Runtime/UMG/Public/Components/ScaleBoxSlot.h
--- a/Engine/Source/Runtime/UMG/Public/Components/ScaleBoxSlot.h
+++ b/Engine/Source/Runtime/UMG/Public/Components/ScaleBoxSlot.h
@@ -15,15 +15,15 @@ class UMG_API UScaleBoxSlot : public UPanelSlot
 public:
 	
 	/** The padding area between the slot and the content it contains. */
-	UPROPERTY(EditDefaultsOnly, Category="Layout (ScaleBox Slot)")
+	UPROPERTY(EditDefaultsOnly, BlueprintReadOnly, Category="Layout (ScaleBox Slot)")
 	FMargin Padding;
 
 	/** The alignment of the object horizontally. */
-	UPROPERTY(EditDefaultsOnly, Category="Layout (ScaleBox Slot)")
+	UPROPERTY(EditDefaultsOnly, BlueprintReadOnly, Category="Layout (ScaleBox Slot)")
 	TEnumAsByte<EHorizontalAlignment> HorizontalAlignment;
 
 	/** The alignment of the object vertically. */
-	UPROPERTY(EditDefaultsOnly, Category="Layout (ScaleBox Slot)")
+	UPROPERTY(EditDefaultsOnly, BlueprintReadOnly, Category="Layout (ScaleBox Slot)" )
 	TEnumAsByte<EVerticalAlignment> VerticalAlignment;
 
 public:
diff --git a/Engine/Source/Runtime/UMG/Public/Components/ScrollBar.h b/Engine/Source/Runtime/UMG/Public/Components/ScrollBar.h
--- a/Engine/Source/Runtime/UMG/Public/Components/ScrollBar.h
+++ b/Engine/Source/Runtime/UMG/Public/Components/ScrollBar.h
@@ -5,7 +5,7 @@
 #include "ScrollBar.generated.h"
 
 /** */
-UCLASS(Experimental, ClassGroup=UserInterface, meta=( IsAdvanced = "True" ))
+UCLASS(Experimental)
 class UMG_API UScrollBar : public UPanelWidget
 {
 	GENERATED_UCLASS_BODY()
diff --git a/Engine/Source/Runtime/UMG/Public/Components/ScrollBox.h b/Engine/Source/Runtime/UMG/Public/Components/ScrollBox.h
--- a/Engine/Source/Runtime/UMG/Public/Components/ScrollBox.h
+++ b/Engine/Source/Runtime/UMG/Public/Components/ScrollBox.h
@@ -7,11 +7,13 @@
 /**
  * An arbitrary scrollable collection of widgets.  Great for presenting 10-100 widgets in a list.  Doesn't support virtualization.
  */
-UCLASS(ClassGroup=UserInterface)
+UCLASS()
 class UMG_API UScrollBox : public UPanelWidget
 {
 	GENERATED_UCLASS_BODY()
 
+public:
+
 	/** The style */
 	UPROPERTY(EditAnywhere, BlueprintReadWrite, Category="Style", meta=( DisplayName="Style" ))
 	FScrollBoxStyle WidgetStyle;
@@ -27,9 +29,13 @@ class UMG_API UScrollBox : public UPanelWidget
 	USlateWidgetStyleAsset* BarStyle_DEPRECATED;
 
 	/** The orientation of the scrolling and stacking in the box. */
-	UPROPERTY(EditDefaultsOnly, Category = "Behavior")
+	UPROPERTY(EditDefaultsOnly, BlueprintReadOnly, Category = "Behavior" )
 	TEnumAsByte<EOrientation> Orientation;
 
+	//TODO UMG Add SetOrientation
+
+public:
+
 	/** Updates the scroll offset of the scrollbox */
 	UFUNCTION(BlueprintCallable, Category="Widget")
 	void SetScrollOffset(float NewScrollOffset);
diff --git a/Engine/Source/Runtime/UMG/Public/Components/ScrollBoxSlot.h b/Engine/Source/Runtime/UMG/Public/Components/ScrollBoxSlot.h
--- a/Engine/Source/Runtime/UMG/Public/Components/ScrollBoxSlot.h
+++ b/Engine/Source/Runtime/UMG/Public/Components/ScrollBoxSlot.h
@@ -15,11 +15,11 @@ class UMG_API UScrollBoxSlot : public UPanelSlot
 public:
 	
 	/** The padding area between the slot and the content it contains. */
-	UPROPERTY(EditDefaultsOnly, Category="Layout (ScrollBox Slot)")
+	UPROPERTY(EditDefaultsOnly, BlueprintReadOnly, Category="Layout (ScrollBox Slot)")
 	FMargin Padding;
 
 	/** The alignment of the object horizontally. */
-	UPROPERTY(EditDefaultsOnly, Category="Layout (ScrollBox Slot)")
+	UPROPERTY(EditDefaultsOnly, BlueprintReadOnly, Category="Layout (ScrollBox Slot)")
 	TEnumAsByte<EHorizontalAlignment> HorizontalAlignment;
 
 public:
diff --git a/Engine/Source/Runtime/UMG/Public/Components/SizeBox.h b/Engine/Source/Runtime/UMG/Public/Components/SizeBox.h
--- a/Engine/Source/Runtime/UMG/Public/Components/SizeBox.h
+++ b/Engine/Source/Runtime/UMG/Public/Components/SizeBox.h
@@ -10,11 +10,62 @@
  * A widget that allows you to specify the size it reports to have and desire.  Not all widgets report a desired size
  * that you actually desire.  Wrapping them in a SizeBox lets you have the Size Box force them to be a particular size.
  */
-UCLASS(ClassGroup=UserInterface)
+UCLASS()
 class UMG_API USizeBox : public UContentWidget
 {
 	GENERATED_UCLASS_BODY()
 
 public:
+
+	/**  */
+	UPROPERTY()
+	uint32 bOverride_WidthOverride : 1;
+
+	/**  */
+	UPROPERTY()
+	uint32 bOverride_HeightOverride : 1;
+
+	/**  */
+	UPROPERTY()
+	uint32 bOverride_MinDesiredWidth : 1;
+
+	/**  */
+	UPROPERTY()
+	uint32 bOverride_MinDesiredHeight : 1;
+
+	/**  */
+	UPROPERTY()
+	uint32 bOverride_MaxDesiredWidth : 1;
+
+	/**  */
+	UPROPERTY()
+	uint32 bOverride_MaxDesiredHeight : 1;
+
+
+	/** When specified, ignore the content's desired size and report the WidthOverride as the Box's desired width. */
+	UPROPERTY(EditDefaultsOnly, BlueprintReadOnly, Category="Child Layout", meta=( editcondition="bOverride_WidthOverride" ))
+	float WidthOverride;
+
+	/** When specified, ignore the content's desired size and report the HeightOverride as the Box's desired height. */
+	UPROPERTY(EditDefaultsOnly, BlueprintReadOnly, Category="Child Layout", meta=( editcondition="bOverride_HeightOverride" ))
+	float HeightOverride;
+
+	/** When specified, will report the MinDesiredWidth if larger than the content's desired width. */
+	UPROPERTY(EditDefaultsOnly, BlueprintReadOnly, Category="Child Layout", meta=( editcondition="bOverride_MinDesiredWidth" ))
+	float MinDesiredWidth;
+
+	/** When specified, will report the MinDesiredHeight if larger than the content's desired height. */
+	UPROPERTY(EditDefaultsOnly, BlueprintReadOnly, Category="Child Layout", meta=( editcondition="bOverride_MinDesiredHeight" ))
+	float MinDesiredHeight;
+
+	/** When specified, will report the MaxDesiredWidth if smaller than the content's desired width. */
+	UPROPERTY(EditDefaultsOnly, BlueprintReadOnly, Category="Child Layout", meta=( editcondition="bOverride_MaxDesiredWidth" ))
+	float MaxDesiredWidth;
+
+	/** When specified, will report the MaxDesiredHeight if smaller than the content's desired height. */
+	UPROPERTY(EditDefaultsOnly, BlueprintReadOnly, Category="Child Layout", meta=( editcondition="bOverride_MaxDesiredHeight" ))
+	float MaxDesiredHeight;
+
+public:
 		
 	/** When specified, ignore the content's desired size and report the WidthOverride as the Box's desired width. */
@@ -76,55 +127,4 @@
 
 protected:
 
-	/**  */
-	UPROPERTY()
-	uint32 bOverride_WidthOverride : 1;
-
-	/**  */
-	UPROPERTY()
-	uint32 bOverride_HeightOverride : 1;
-
-	/**  */
-	UPROPERTY()
-	uint32 bOverride_MinDesiredWidth : 1;
-
-	/**  */
-	UPROPERTY()
-	uint32 bOverride_MinDesiredHeight : 1;
-
-	/**  */
-	UPROPERTY()
-	uint32 bOverride_MaxDesiredWidth : 1;
-
-	/**  */
-	UPROPERTY()
-	uint32 bOverride_MaxDesiredHeight : 1;
-
-
-	/** When specified, ignore the content's desired size and report the WidthOverride as the Box's desired width. */
-	UPROPERTY(EditDefaultsOnly, Category="Child Layout", meta=( editcondition="bOverride_WidthOverride" ))
-	float WidthOverride;
-
-	/** When specified, ignore the content's desired size and report the HeightOverride as the Box's desired height. */
-	UPROPERTY(EditDefaultsOnly, Category="Child Layout", meta=( editcondition="bOverride_HeightOverride" ))
-	float HeightOverride;
-
-	/** When specified, will report the MinDesiredWidth if larger than the content's desired width. */
-	UPROPERTY(EditDefaultsOnly, Category="Child Layout", meta=( editcondition="bOverride_MinDesiredWidth" ))
-	float MinDesiredWidth;
-
-	/** When specified, will report the MinDesiredHeight if larger than the content's desired height. */
-	UPROPERTY(EditDefaultsOnly, Category="Child Layout", meta=( editcondition="bOverride_MinDesiredHeight" ))
-	float MinDesiredHeight;
-
-	/** When specified, will report the MaxDesiredWidth if smaller than the content's desired width. */
-	UPROPERTY(EditDefaultsOnly, Category="Child Layout", meta=( editcondition="bOverride_MaxDesiredWidth" ))
-	float MaxDesiredWidth;
-
-	/** When specified, will report the MaxDesiredHeight if smaller than the content's desired height. */
-	UPROPERTY(EditDefaultsOnly, Category="Child Layout", meta=( editcondition="bOverride_MaxDesiredHeight" ))
-	float MaxDesiredHeight;
-
-protected:
-
 	// UPanelWidget
diff --git a/Engine/Source/Runtime/UMG/Public/Components/SizeBoxSlot.h b/Engine/Source/Runtime/UMG/Public/Components/SizeBoxSlot.h
--- a/Engine/Source/Runtime/UMG/Public/Components/SizeBoxSlot.h
+++ b/Engine/Source/Runtime/UMG/Public/Components/SizeBoxSlot.h
@@ -12,6 +12,19 @@ class UMG_API USizeBoxSlot : public UPanelSlot
 {
 	GENERATED_UCLASS_BODY()
 
+public:
+	/** The padding area between the slot and the content it contains. */
+	UPROPERTY(EditDefaultsOnly, BlueprintReadOnly, Category="Layout (SizeBox Slot)")
+	FMargin Padding;
+
+	/** The alignment of the object horizontally. */
+	UPROPERTY(EditDefaultsOnly, BlueprintReadOnly, Category="Layout (SizeBox Slot)")
+	TEnumAsByte<EHorizontalAlignment> HorizontalAlignment;
+
+	/** The alignment of the object vertically. */
+	UPROPERTY(EditDefaultsOnly, BlueprintReadOnly, Category="Layout (SizeBox Slot)")
+	TEnumAsByte<EVerticalAlignment> VerticalAlignment;
+
 public:
 
 	UFUNCTION(BlueprintCallable, Category="Layout (SizeBox Slot)")
@@ -34,19 +47,6 @@ public:
 
 	virtual void ReleaseSlateResources(bool bReleaseChildren) override;
 
-protected:
-	/** The padding area between the slot and the content it contains. */
-	UPROPERTY(EditDefaultsOnly, Category="Layout (SizeBox Slot)")
-	FMargin Padding;
-
-	/** The alignment of the object horizontally. */
-	UPROPERTY(EditDefaultsOnly, Category="Layout (SizeBox Slot)")
-	TEnumAsByte<EHorizontalAlignment> HorizontalAlignment;
-
-	/** The alignment of the object vertically. */
-	UPROPERTY(EditDefaultsOnly, Category="Layout (SizeBox Slot)")
-	TEnumAsByte<EVerticalAlignment> VerticalAlignment;
-
 private:
 
 	/** A pointer to the button to allow us to adjust the size, padding...etc at runtime. */
diff --git a/Engine/Source/Runtime/UMG/Public/Components/Slider.h b/Engine/Source/Runtime/UMG/Public/Components/Slider.h
--- a/Engine/Source/Runtime/UMG/Public/Components/Slider.h
+++ b/Engine/Source/Runtime/UMG/Public/Components/Slider.h
@@ -11,17 +11,12 @@ DECLARE_DYNAMIC_MULTICAST_DELEGATE_OneParam(FOnFloatValueChangedEvent, float, Va
 /**
  * A simple widget that shows a sliding bar with a handle that allows you to control the value between 0..1.
  */
-UCLASS(ClassGroup=UserInterface)
+UCLASS()
 class UMG_API USlider : public UWidget
 {
 	GENERATED_UCLASS_BODY()
 
 public:
-	
-	/** The progress bar style */
-	UPROPERTY(EditAnywhere, BlueprintReadWrite, Category="Style", meta=( DisplayName="Style" ))
-	FSliderStyle WidgetStyle;
-
 	/** The volume value to display. */
 	UPROPERTY(EditDefaultsOnly, Category=Appearance)
 	float Value;
@@ -30,26 +25,34 @@ public:
 	UPROPERTY()
 	FGetFloat ValueDelegate;
 
+public:
+	
+	/** The progress bar style */
+	UPROPERTY(EditAnywhere, BlueprintReadWrite, Category="Style", meta=( DisplayName="Style" ))
+	FSliderStyle WidgetStyle;
+
 	/** The slider's orientation. */
-	UPROPERTY(EditDefaultsOnly, Category=Appearance)
+	UPROPERTY(EditDefaultsOnly, BlueprintReadOnly, Category=Appearance)
 	TEnumAsByte<EOrientation> Orientation;
 
 	/** The color to draw the slider bar in. */
-	UPROPERTY(EditDefaultsOnly, Category=Appearance)
+	UPROPERTY(EditDefaultsOnly, BlueprintReadOnly, Category=Appearance)
 	FLinearColor SliderBarColor;
 
 	/** The color to draw the slider handle in. */
-	UPROPERTY(EditDefaultsOnly, Category=Appearance)
+	UPROPERTY(EditDefaultsOnly, BlueprintReadOnly, Category=Appearance)
 	FLinearColor SliderHandleColor;
 
 	/** Whether the slidable area should be indented to fit the handle. */
-	UPROPERTY(EditDefaultsOnly, Category=Appearance, AdvancedDisplay)
+	UPROPERTY(EditDefaultsOnly, BlueprintReadOnly, Category=Appearance, AdvancedDisplay)
 	bool IndentHandle;
 
 	/** Whether the handle is interactive or fixed. */
-	UPROPERTY(EditDefaultsOnly, Category=Appearance, AdvancedDisplay)
+	UPROPERTY(EditDefaultsOnly, BlueprintReadOnly, Category=Appearance, AdvancedDisplay)
 	bool Locked;
 
+public:
+
 	/** Invoked when the mouse is pressed and a capture begins. */
 	UPROPERTY(BlueprintAssignable, Category="Widget Event")
 	FOnMouseCaptureBeginEvent OnMouseCaptureBegin;
diff --git a/Engine/Source/Runtime/UMG/Public/Components/Spacer.h b/Engine/Source/Runtime/UMG/Public/Components/Spacer.h
--- a/Engine/Source/Runtime/UMG/Public/Components/Spacer.h
+++ b/Engine/Source/Runtime/UMG/Public/Components/Spacer.h
@@ -7,7 +7,7 @@
 /**
  * A spacer widget; it does not have a visual representation, and just provides padding between other widgets.
  */
-UCLASS(ClassGroup=UserInterface)
+UCLASS()
 class UMG_API USpacer : public UWidget
 {
 	GENERATED_UCLASS_BODY()
@@ -15,7 +15,7 @@ class UMG_API USpacer : public UWidget
 public:
 
 	/** The size of the spacer */
-	UPROPERTY(EditDefaultsOnly, Category=Appearance)
+	UPROPERTY(EditDefaultsOnly, BlueprintReadOnly, Category=Appearance)
 	FVector2D Size;
 
 public:
diff --git a/Engine/Source/Runtime/UMG/Public/Components/SpinBox.h b/Engine/Source/Runtime/UMG/Public/Components/SpinBox.h
--- a/Engine/Source/Runtime/UMG/Public/Components/SpinBox.h
+++ b/Engine/Source/Runtime/UMG/Public/Components/SpinBox.h
@@ -7,7 +7,7 @@
 /**
  * A numerical entry box that allows for direct entry of the number or allows the user to click and slide the number.
  */
-UCLASS(ClassGroup = UserInterface)
+UCLASS()
 class UMG_API USpinBox : public UWidget
 {
 	GENERATED_UCLASS_BODY()
@@ -17,6 +17,16 @@ public:
 	DECLARE_DYNAMIC_MULTICAST_DELEGATE_TwoParams(FOnSpinBoxValueCommittedEvent, float, InValue, ETextCommit::Type, CommitMethod);
 	DECLARE_DYNAMIC_MULTICAST_DELEGATE(FOnSpinBoxBeginSliderMovement);
 
+public:
+
+	/** Value stored in this spin box */
+	UPROPERTY(EditDefaultsOnly, Category=Content)
+	float Value;
+
+	/** A bindable delegate to allow logic to drive the value of the widget */
+	UPROPERTY()
+	FGetFloat ValueDelegate;
+
 public:
 	/** The Style */
 	UPROPERTY(EditAnywhere, BlueprintReadWrite, Category="Style", meta=( DisplayName="Style" ))
@@ -25,65 +35,64 @@ public:
 	UPROPERTY()
 	USlateWidgetStyleAsset* Style_DEPRECATED;
 
-	/** Value stored in this spin box */
-	UPROPERTY(EditDefaultsOnly, Category = Content)
-	float Value;
-
 	/** The amount by which to change the spin box value as the slider moves. */
-	UPROPERTY(EditDefaultsOnly, Category = Content)
+	UPROPERTY(EditDefaultsOnly, BlueprintReadOnly, Category="Slider")
 	float Delta;
 
 	/** The exponent by which to increase the delta as the mouse moves. 1 is constant (never increases the delta). */
-	UPROPERTY(EditDefaultsOnly, Category = Content)
+	UPROPERTY(EditDefaultsOnly, BlueprintReadOnly, Category="Slider")
 	float SliderExponent;
 	
 	/** Font color and opacity (overrides style) */
-	UPROPERTY(EditDefaultsOnly, Category = Content)
+	UPROPERTY(EditDefaultsOnly, BlueprintReadOnly, Category="Display")
 	FSlateFontInfo Font;
 
 	/** The minimum width of the spin box */
-	UPROPERTY(EditDefaultsOnly, Category = Content, AdvancedDisplay, DisplayName = "Minimum Desired Width")
+	UPROPERTY(EditDefaultsOnly, BlueprintReadOnly, Category="Display", AdvancedDisplay, DisplayName = "Minimum Desired Width")
 	float MinDesiredWidth;
 
 	/** Whether to remove the keyboard focus from the spin box when the value is committed */
-	UPROPERTY(EditDefaultsOnly, Category = Content, AdvancedDisplay)
+	UPROPERTY(EditDefaultsOnly, BlueprintReadOnly, Category="Input", AdvancedDisplay)
 	bool ClearKeyboardFocusOnCommit;
 
 	/** Whether to select the text in the spin box when the value is committed */
-	UPROPERTY(EditDefaultsOnly, Category = Content, AdvancedDisplay)
+	UPROPERTY(EditDefaultsOnly, BlueprintReadOnly, Category="Input", AdvancedDisplay)
 	bool SelectAllTextOnCommit;
 
-	UPROPERTY(EditDefaultsOnly, Category="Style")
+	UPROPERTY(EditDefaultsOnly, BlueprintReadOnly, Category="Style")
 	FSlateColor ForegroundColor;
 
 public:
 	/** Called when the value is changed interactively by the user */
-	UPROPERTY(BlueprintAssignable, Category = "Widget Event")
+	UPROPERTY(BlueprintAssignable, Category="SpinBox|Events")
 	FOnSpinBoxValueChangedEvent OnValueChanged;
 
 	/** Called when the value is committed. Occurs when the user presses Enter or the text box loses focus. */
-	UPROPERTY(BlueprintAssignable, Category = "Widget Event")
+	UPROPERTY(BlueprintAssignable, Category="SpinBox|Events")
 	FOnSpinBoxValueCommittedEvent OnValueCommitted;
 
 	/** Called right before the slider begins to move */
-	UPROPERTY(BlueprintAssignable, Category = "Widget Event")
+	UPROPERTY(BlueprintAssignable, Category="SpinBox|Events")
 	FOnSpinBoxBeginSliderMovement OnBeginSliderMovement;
 
 	/** Called right after the slider handle is released by the user */
-	UPROPERTY(BlueprintAssignable, Category = "Widget Event")
+	UPROPERTY(BlueprintAssignable, Category="SpinBox|Events")
 	FOnSpinBoxValueChangedEvent OnEndSliderMovement;
 
+public:
+
 	/** Get the current value of the spin box. */
-	UFUNCTION(BlueprintCallable, Category = "Behavior")
+	UFUNCTION(BlueprintCallable, Category="Behavior")
 	float GetValue() const;
 
 	/** Set the value of the spin box. */
-	UFUNCTION(BlueprintCallable, Category = "Behavior")
+	UFUNCTION(BlueprintCallable, Category="Behavior")
 	void SetValue(float NewValue);
 
-	// MIN VALUE
+public:
+
 	/** Get the current minimum value that can be manually set in the spin box. */
-	UFUNCTION(BlueprintCallable, Category = "Behavior")
+	UFUNCTION(BlueprintCallable, Category="Behavior")
 	float GetMinValue() const;
 
 	/** Set the minimum value that can be manually set in the spin box. */
@@ -94,7 +103,6 @@ public:
 	UFUNCTION(BlueprintCallable, Category = "Behavior")
 	void ClearMinValue();
 
-	// MAX VALUE
 	/** Get the current maximum value that can be manually set in the spin box. */
 	UFUNCTION(BlueprintCallable, Category = "Behavior")
 	float GetMaxValue() const;
@@ -107,7 +115,6 @@ public:
 	UFUNCTION(BlueprintCallable, Category = "Behavior")
 	void ClearMaxValue();
 
-	// MIN SLIDER VALUE
 	/** Get the current minimum value that can be specified using the slider. */
 	UFUNCTION(BlueprintCallable, Category = "Behavior")
 	float GetMinSliderValue() const;
@@ -120,7 +127,6 @@ public:
 	UFUNCTION(BlueprintCallable, Category = "Behavior")
 	void ClearMinSliderValue();
 
-	// MAX SLIDER VALUE
 	/** Get the current maximum value that can be specified using the slider. */
 	UFUNCTION(BlueprintCallable, Category = "Behavior")
 	float GetMaxSliderValue() const;
@@ -137,6 +143,8 @@ public:
 	UFUNCTION(BlueprintCallable, Category = "Appearance")
 	void SetForegroundColor(FSlateColor InForegroundColor);
 
+public:
+
 	// UWidget interface
 	virtual void SynchronizeProperties() override;
 	// End of UWidget interface
@@ -155,10 +163,6 @@ public:
 #endif
 	// End of UWidget interface
 
-	/** A bindable delegate to allow logic to drive the value of the widget */
-	UPROPERTY()
-	FGetFloat ValueDelegate;
-
 protected:
 	// UWidget interface
 	virtual TSharedRef<SWidget> RebuildWidget() override;
diff --git a/Engine/Source/Runtime/UMG/Public/Components/TextBlock.h b/Engine/Source/Runtime/UMG/Public/Components/TextBlock.h
--- a/Engine/Source/Runtime/UMG/Public/Components/TextBlock.h
+++ b/Engine/Source/Runtime/UMG/Public/Components/TextBlock.h
@@ -47,7 +47,7 @@ public:
 	USlateWidgetStyleAsset* Style_DEPRECATED;
 
 	/** The text to display */
-	UPROPERTY(EditDefaultsOnly, Category=Content, meta=(MultiLine="true"))
+	UPROPERTY(EditDefaultsOnly, Category=Content, meta=( MultiLine="true" ))
 	FText Text;
 	
 	/** A bindable delegate to allow logic to drive the text of the widget */
@@ -55,7 +55,7 @@ public:
 	FGetText TextDelegate;
 
 	/** The color of the text */
-	UPROPERTY(EditDefaultsOnly, Category=Appearance)
+	UPROPERTY(EditDefaultsOnly, BlueprintReadOnly, Category=Appearance)
 	FSlateColor ColorAndOpacity;
 
 	/** A bindable delegate for the ColorAndOpacity. */
@@ -63,15 +63,15 @@ public:
 	FGetSlateColor ColorAndOpacityDelegate;
 	
 	/** The font to render the text with */
-	UPROPERTY(EditDefaultsOnly, Category=Appearance)
+	UPROPERTY(EditDefaultsOnly, BlueprintReadOnly, Category=Appearance)
 	FSlateFontInfo Font;
 
 	/** The direction the shadow is cast */
-	UPROPERTY(EditDefaultsOnly, Category=Appearance)
+	UPROPERTY(EditDefaultsOnly, BlueprintReadOnly, Category=Appearance)
 	FVector2D ShadowOffset;
 
 	/** The color of the shadow */
-	UPROPERTY(EditDefaultsOnly, Category=Appearance, meta=( DisplayName="Shadow Color" ))
+	UPROPERTY(EditDefaultsOnly, BlueprintReadOnly, Category=Appearance, meta=( DisplayName="Shadow Color" ))
 	FLinearColor ShadowColorAndOpacity;
 
 	/** A bindable delegate for the ShadowColorAndOpacity. */
@@ -79,27 +79,27 @@ public:
 	FGetLinearColor ShadowColorAndOpacityDelegate;
 
 	/** How the text should be aligned with the margin. */
-	UPROPERTY(EditDefaultsOnly, Category=Appearance)
+	UPROPERTY(EditDefaultsOnly, BlueprintReadOnly, Category=Appearance)
 	TEnumAsByte<ETextJustify::Type> Justification;
 
 	/** True if we're wrapping text automatically based on the computed horizontal space for this widget */
-	UPROPERTY(EditDefaultsOnly, Category=Appearance)
+	UPROPERTY(EditDefaultsOnly, BlueprintReadOnly, Category=Appearance)
 	bool AutoWrapText;
 
 	/** Whether text wraps onto a new line when it's length exceeds this width; if this value is zero or negative, no wrapping occurs. */
-	UPROPERTY(EditDefaultsOnly, Category=Appearance, AdvancedDisplay)
+	UPROPERTY(EditDefaultsOnly, BlueprintReadOnly, Category=Appearance, AdvancedDisplay)
 	float WrapTextAt;
 
 	/** The minimum desired size for the text */
-	UPROPERTY(EditDefaultsOnly, Category=Appearance, AdvancedDisplay)
+	UPROPERTY(EditDefaultsOnly, BlueprintReadOnly, Category=Appearance, AdvancedDisplay)
 	float MinDesiredWidth;
 
 	/** The amount of blank space left around the edges of text area. */
-	UPROPERTY(EditDefaultsOnly, Category=Appearance, AdvancedDisplay)
+	UPROPERTY(EditDefaultsOnly, BlueprintReadOnly, Category=Appearance, AdvancedDisplay)
 	FMargin Margin;
 
 	/** The amount to scale each lines height by. */
-	UPROPERTY(EditDefaultsOnly, Category=Appearance, AdvancedDisplay)
+	UPROPERTY(EditDefaultsOnly, BlueprintReadOnly, Category=Appearance, AdvancedDisplay)
 	float LineHeightPercentage;
 
 	///** Called when this text is double clicked */
diff --git a/Engine/Source/Runtime/UMG/Public/Components/Throbber.h b/Engine/Source/Runtime/UMG/Public/Components/Throbber.h
--- a/Engine/Source/Runtime/UMG/Public/Components/Throbber.h
+++ b/Engine/Source/Runtime/UMG/Public/Components/Throbber.h
@@ -15,18 +15,43 @@ class UMG_API UThrobber : public UWidget
 {
 	GENERATED_UCLASS_BODY()
 
+public:
+	
+	/** How many pieces there are */
+	UPROPERTY(EditDefaultsOnly, BlueprintReadOnly, Category=Appearance, meta=( ClampMin="1", UIMin="1", UIMax="8" ))
+	int32 NumberOfPieces;
+
+	/** Should the pieces animate horizontally? */
+	UPROPERTY(EditDefaultsOnly, BlueprintReadOnly, Category=Appearance)
+	bool bAnimateHorizontally;
+
+	/** Should the pieces animate vertically? */
+	UPROPERTY(EditDefaultsOnly, BlueprintReadOnly, Category=Appearance)
+	bool bAnimateVertically;
+
+	/** Should the pieces animate their opacity? */
+	UPROPERTY(EditDefaultsOnly, BlueprintReadOnly, Category=Appearance)
+	bool bAnimateOpacity;
+
+	/** Image to use for each segment of the throbber */
+	UPROPERTY()
+	USlateBrushAsset* PieceImage_DEPRECATED;
+
+	UPROPERTY(EditDefaultsOnly, BlueprintReadOnly, Category=Appearance)
+	FSlateBrush Image;
+
 public:
 
 	/** Sets how many pieces there are */
-	UFUNCTION(BlueprintCallable, Category = "Appearance")
+	UFUNCTION(BlueprintCallable, Category="Appearance")
 	void SetNumberOfPieces(int32 InNumberOfPieces);
 
 	/** Sets whether the pieces animate horizontally. */
-	UFUNCTION(BlueprintCallable, Category = "Appearance")
+	UFUNCTION(BlueprintCallable, Category="Appearance")
 	void SetAnimateHorizontally(bool bInAnimateHorizontally);
 
 	/** Sets whether the pieces animate vertically. */
-	UFUNCTION(BlueprintCallable, Category = "Appearance")
+	UFUNCTION(BlueprintCallable, Category="Appearance")
 	void SetAnimateVertically(bool bInAnimateVertically);
 
 	/** Sets whether the pieces animate their opacity. */
@@ -50,31 +75,6 @@ public:
 	virtual const FText GetPaletteCategory() override;
 #endif
 
-protected:
-
-	/** How many pieces there are */
-	UPROPERTY(EditDefaultsOnly, Category=Appearance, meta=(ClampMin = "1", UIMin = "1", UIMax = "8"))
-	int32 NumberOfPieces;
-
-	/** Should the pieces animate horizontally? */
-	UPROPERTY(EditDefaultsOnly, Category=Appearance)
-	bool bAnimateHorizontally;
-
-	/** Should the pieces animate vertically? */
-	UPROPERTY(EditDefaultsOnly, Category=Appearance)
-	bool bAnimateVertically;
-
-	/** Should the pieces animate their opacity? */
-	UPROPERTY(EditDefaultsOnly, Category=Appearance)
-	bool bAnimateOpacity;
-
-	/** Image to use for each segment of the throbber */
-	UPROPERTY()
-	USlateBrushAsset* PieceImage_DEPRECATED;
-
-	UPROPERTY(EditDefaultsOnly, Category=Appearance)
-	FSlateBrush Image;
-
 protected:
 	// UWidget interface
 	virtual TSharedRef<SWidget> RebuildWidget() override;
diff --git a/Engine/Source/Runtime/UMG/Public/Components/UniformGridPanel.h b/Engine/Source/Runtime/UMG/Public/Components/UniformGridPanel.h
--- a/Engine/Source/Runtime/UMG/Public/Components/UniformGridPanel.h
+++ b/Engine/Source/Runtime/UMG/Public/Components/UniformGridPanel.h
@@ -7,7 +7,7 @@
 /**
  * A panel that evenly divides up available space between all of its children.
  */
-UCLASS(ClassGroup=UserInterface)
+UCLASS()
 class UMG_API UUniformGridPanel : public UPanelWidget
 {
 	GENERATED_UCLASS_BODY()
@@ -15,17 +15,31 @@ class UMG_API UUniformGridPanel : public UPanelWidget
 public:
 
 	/** Padding given to each slot */
-	UPROPERTY(EditDefaultsOnly, Category="Child Layout")
+	UPROPERTY(EditDefaultsOnly, BlueprintReadOnly, Category="Child Layout")
 	FMargin SlotPadding;
 
 	/** The minimum desired width of the slots */
-	UPROPERTY(EditDefaultsOnly, Category="Child Layout")
+	UPROPERTY(EditDefaultsOnly, BlueprintReadOnly, Category="Child Layout")
 	float MinDesiredSlotWidth;
 
 	/** The minimum desired height of the slots */
-	UPROPERTY(EditDefaultsOnly, Category="Child Layout")
+	UPROPERTY(EditDefaultsOnly, BlueprintReadOnly, Category="Child Layout")
 	float MinDesiredSlotHeight;
 
+public:
+
+	/**  */
+	UFUNCTION(BlueprintCallable, Category="Child Layout")
+	void SetSlotPadding(FMargin InSlotPadding);
+
+	/**  */
+	UFUNCTION(BlueprintCallable, Category="Child Layout")
+	void SetMinDesiredSlotWidth(float InMinDesiredSlotWidth);
+
+	/**  */
+	UFUNCTION(BlueprintCallable, Category="Child Layout")
+	void SetMinDesiredSlotHeight(float InMinDesiredSlotHeight);
+
 public:
 
 	/**  */
diff --git a/Engine/Source/Runtime/UMG/Public/Components/UniformGridSlot.h b/Engine/Source/Runtime/UMG/Public/Components/UniformGridSlot.h
--- a/Engine/Source/Runtime/UMG/Public/Components/UniformGridSlot.h
+++ b/Engine/Source/Runtime/UMG/Public/Components/UniformGridSlot.h
@@ -18,19 +18,19 @@ class UMG_API UUniformGridSlot : public UPanelSlot
 public:
 
 	/** The alignment of the object horizontally. */
-	UPROPERTY(EditDefaultsOnly, Category="Layout (Uniform Grid Slot)")
+	UPROPERTY(EditDefaultsOnly, BlueprintReadOnly, Category="Layout (Uniform Grid Slot)")
 	TEnumAsByte<EHorizontalAlignment> HorizontalAlignment;
 
 	/** The alignment of the object vertically. */
-	UPROPERTY(EditDefaultsOnly, Category="Layout (Uniform Grid Slot)")
+	UPROPERTY(EditDefaultsOnly, BlueprintReadOnly, Category="Layout (Uniform Grid Slot)")
 	TEnumAsByte<EVerticalAlignment> VerticalAlignment;
 	
 	/** The row index of the cell this slot is in */
-	UPROPERTY(EditDefaultsOnly, meta=( UIMin = "0" ), Category="Layout (Uniform Grid Slot)")
+	UPROPERTY(EditDefaultsOnly, BlueprintReadOnly, meta=( UIMin = "0" ), Category="Layout (Uniform Grid Slot)")
 	int32 Row;
 	
 	/** The column index of the cell this slot is in */
-	UPROPERTY(EditDefaultsOnly, meta=( UIMin = "0" ), Category="Layout (Uniform Grid Slot)")
+	UPROPERTY(EditDefaultsOnly, BlueprintReadOnly, meta=( UIMin = "0" ), Category="Layout (Uniform Grid Slot)")
 	int32 Column;
 
 public:
diff --git a/Engine/Source/Runtime/UMG/Public/Components/VerticalBox.h b/Engine/Source/Runtime/UMG/Public/Components/VerticalBox.h
--- a/Engine/Source/Runtime/UMG/Public/Components/VerticalBox.h
+++ b/Engine/Source/Runtime/UMG/Public/Components/VerticalBox.h
@@ -10,7 +10,7 @@ class UVerticalBoxSlot;
  * A vertical box widget is a layout panel allowing child widgets to be automatically laid out
  * vertically.
  */
-UCLASS(ClassGroup=UserInterface)
+UCLASS()
 class UMG_API UVerticalBox : public UPanelWidget
 {
 	GENERATED_UCLASS_BODY()
diff --git a/Engine/Source/Runtime/UMG/Public/Components/VerticalBoxSlot.h b/Engine/Source/Runtime/UMG/Public/Components/VerticalBoxSlot.h
--- a/Engine/Source/Runtime/UMG/Public/Components/VerticalBoxSlot.h
+++ b/Engine/Source/Runtime/UMG/Public/Components/VerticalBoxSlot.h
@@ -15,34 +15,34 @@ class UMG_API UVerticalBoxSlot : public UPanelSlot
 public:
 	
 	/** The padding area between the slot and the content it contains. */
-	UPROPERTY(EditDefaultsOnly, Category="Layout (Vertical Box Slot)")
+	UPROPERTY(EditDefaultsOnly, BlueprintReadOnly, Category="Layout (Vertical Box Slot)")
 	FMargin Padding;
 
 	/** How much space this slot should occupy in the direction of the panel. */
-	UPROPERTY(EditDefaultsOnly, Category="Layout (Vertical Box Slot)")
+	UPROPERTY(EditDefaultsOnly, BlueprintReadOnly, Category="Layout (Vertical Box Slot)")
 	FSlateChildSize Size;
 
 	/** The alignment of the object horizontally. */
-	UPROPERTY(EditDefaultsOnly, Category="Layout (Vertical Box Slot)")
+	UPROPERTY(EditDefaultsOnly, BlueprintReadOnly, Category="Layout (Vertical Box Slot)")
 	TEnumAsByte<EHorizontalAlignment> HorizontalAlignment;
 
 	/** The alignment of the object vertically. */
-	UPROPERTY(EditDefaultsOnly, Category="Layout (Vertical Box Slot)")
+	UPROPERTY(EditDefaultsOnly, BlueprintReadOnly, Category="Layout (Vertical Box Slot)")
 	TEnumAsByte<EVerticalAlignment> VerticalAlignment;
 
 public:
 
 	UFUNCTION(BlueprintCallable, Category="Layout (Vertical Box Slot)")
-		void SetPadding(FMargin InPadding);
+	void SetPadding(FMargin InPadding);
 
 	UFUNCTION(BlueprintCallable, Category="Layout (Vertical Box Slot)")
-		void SetSize(FSlateChildSize InSize);
+	void SetSize(FSlateChildSize InSize);
 
 	UFUNCTION(BlueprintCallable, Category="Layout (Vertical Box Slot)")
-		void SetHorizontalAlignment(EHorizontalAlignment InHorizontalAlignment);
+	void SetHorizontalAlignment(EHorizontalAlignment InHorizontalAlignment);
 
 	UFUNCTION(BlueprintCallable, Category="Layout (Vertical Box Slot)")
-		void SetVerticalAlignment(EVerticalAlignment InVerticalAlignment);
+	void SetVerticalAlignment(EVerticalAlignment InVerticalAlignment);
 
 public:
 
diff --git a/Engine/Source/Runtime/UMG/Public/Components/Widget.h b/Engine/Source/Runtime/UMG/Public/Components/Widget.h
--- a/Engine/Source/Runtime/UMG/Public/Components/Widget.h
+++ b/Engine/Source/Runtime/UMG/Public/Components/Widget.h
@@ -73,7 +73,7 @@ public:
 	UPanelSlot* Slot;
 
 	/** Sets whether this widget can be modified interactively by the user */
-	UPROPERTY(EditDefaultsOnly, Category=Behavior)
+	UPROPERTY(EditDefaultsOnly, BlueprintReadOnly, Category=Behavior)
 	bool bIsEnabled;
 
 	/** A bindable delegate for bIsEnabled */
@@ -81,7 +81,7 @@ public:
 	FGetBool bIsEnabledDelegate;
 
 	/** Tooltip text to show when the user hovers over the widget with the mouse */
-	UPROPERTY(EditDefaultsOnly, Category=Behavior)
+	UPROPERTY(EditDefaultsOnly, BlueprintReadOnly, Category=Behavior)
 	FText ToolTipText;
 
 	/** A bindable delegate for ToolTipText */
@@ -89,7 +89,7 @@ public:
 	FGetText ToolTipTextDelegate;
 
 	/** The visibility of the widget */
-	UPROPERTY(EditDefaultsOnly, Category=Behavior)
+	UPROPERTY(EditDefaultsOnly, BlueprintReadOnly, Category=Behavior)
 	TEnumAsByte<ESlateVisibility::Type> Visiblity;
 
 	/** A bindable delegate for Visibility */
@@ -97,7 +97,7 @@ public:
 	FGetSlateVisibility VisiblityDelegate;
 
 	/** The cursor to show when the mouse is over the widget */
-	UPROPERTY(EditDefaultsOnly, Category=Behavior, AdvancedDisplay)
+	UPROPERTY(EditDefaultsOnly, BlueprintReadOnly, Category=Behavior, AdvancedDisplay)
 	TEnumAsByte<EMouseCursor::Type> Cursor;
 
 	/** A bindable delegate for Cursor */
diff --git a/Engine/Source/Runtime/UMG/Public/Components/WidgetSwitcher.h b/Engine/Source/Runtime/UMG/Public/Components/WidgetSwitcher.h
--- a/Engine/Source/Runtime/UMG/Public/Components/WidgetSwitcher.h
+++ b/Engine/Source/Runtime/UMG/Public/Components/WidgetSwitcher.h
@@ -7,14 +7,14 @@
 /**
  * A widget switcher is like a tab control, but without tabs. At most one widget is visible at time.
  */
-UCLASS(ClassGroup=UserInterface)
+UCLASS()
 class UMG_API UWidgetSwitcher : public UPanelWidget
 {
 	GENERATED_UCLASS_BODY()
 
 public:
 	/** The slot index to display */
-	UPROPERTY(EditDefaultsOnly, Category="Switcher", meta=( UIMin=0, ClampMin=0 ))
+	UPROPERTY(EditDefaultsOnly, BlueprintReadOnly, Category="Switcher", meta=( UIMin=0, ClampMin=0 ))
 	int32 ActiveWidgetIndex;
 
 public:
diff --git a/Engine/Source/Runtime/UMG/Public/Components/WidgetSwitcherSlot.h b/Engine/Source/Runtime/UMG/Public/Components/WidgetSwitcherSlot.h
--- a/Engine/Source/Runtime/UMG/Public/Components/WidgetSwitcherSlot.h
+++ b/Engine/Source/Runtime/UMG/Public/Components/WidgetSwitcherSlot.h
@@ -16,15 +16,15 @@ class UMG_API UWidgetSwitcherSlot : public UPanelSlot
 public:
 	
 	/** The padding area between the slot and the content it contains. */
-	UPROPERTY(EditDefaultsOnly, Category="Layout (Widget Switcher Slot)")
+	UPROPERTY(EditDefaultsOnly, BlueprintReadOnly, Category="Layout (Widget Switcher Slot)")
 	FMargin Padding;
 
 	/** The alignment of the object horizontally. */
-	UPROPERTY(EditDefaultsOnly, Category="Layout (Widget Switcher Slot)")
+	UPROPERTY(EditDefaultsOnly, BlueprintReadOnly, Category="Layout (Widget Switcher Slot)")
 	TEnumAsByte<EHorizontalAlignment> HorizontalAlignment;
 
 	/** The alignment of the object vertically. */
-	UPROPERTY(EditDefaultsOnly, Category="Layout (Widget Switcher Slot)")
+	UPROPERTY(EditDefaultsOnly, BlueprintReadOnly, Category="Layout (Widget Switcher Slot)")
 	TEnumAsByte<EVerticalAlignment> VerticalAlignment;
 
 public:
diff --git a/Engine/Source/Runtime/UMG/Public/Components/WrapBox.h b/Engine/Source/Runtime/UMG/Public/Components/WrapBox.h
--- a/Engine/Source/Runtime/UMG/Public/Components/WrapBox.h
+++ b/Engine/Source/Runtime/UMG/Public/Components/WrapBox.h
@@ -7,16 +7,20 @@
 /**
  * Arranges widgets left-to-right.  When the widgets exceed the Width it will place widgets on the next line.
  */
-UCLASS(ClassGroup=UserInterface)
+UCLASS()
 class UMG_API UWrapBox : public UPanelWidget
 {
 	GENERATED_UCLASS_BODY()
 
 public:
 	/** The inner slot padding goes between slots sharing borders */
-	UPROPERTY(EditDefaultsOnly, Category="Conetnt Layout")
+	UPROPERTY(EditDefaultsOnly, BlueprintReadOnly, Category="Content Layout")
 	FVector2D InnerSlotPadding;
 
+	/** Sets the inner slot padding goes between slots sharing borders */
+	UFUNCTION(BlueprintCallable, Category="Content Layout")
+	void SetInnerSlotPadding(FVector2D InPadding);
+
 public:
 
 	/**  */
diff --git a/Engine/Source/Runtime/UMG/Public/Components/WrapBoxSlot.h b/Engine/Source/Runtime/UMG/Public/Components/WrapBoxSlot.h
--- a/Engine/Source/Runtime/UMG/Public/Components/WrapBoxSlot.h
+++ b/Engine/Source/Runtime/UMG/Public/Components/WrapBoxSlot.h
@@ -17,26 +17,26 @@ class UMG_API UWrapBoxSlot : public UPanelSlot
 public:
 
 	/** The padding area between the slot and the content it contains. */
-	UPROPERTY(EditDefaultsOnly, Category="Layout (Wrap Box Slot)")
+	UPROPERTY(EditDefaultsOnly, BlueprintReadOnly, Category="Layout (Wrap Box Slot)")
 	FMargin Padding;
 
 	/** Should this slot fill the remaining space on the line? */
-	UPROPERTY(EditDefaultsOnly, Category="Layout (Wrap Box Slot)")
+	UPROPERTY(EditDefaultsOnly, BlueprintReadOnly, Category="Layout (Wrap Box Slot)")
 	bool bFillEmptySpace;
 
 	/**
 	 * If the total available space in the wrap panel drops below this threshold, this slot will attempt to fill an entire line.
 	 * NOTE: A value of 0, denotes no filling will occur.
 	 */
-	UPROPERTY(EditDefaultsOnly, Category="Layout (Wrap Box Slot)")
+	UPROPERTY(EditDefaultsOnly, BlueprintReadOnly, Category="Layout (Wrap Box Slot)")
 	float FillSpanWhenLessThan;
 
 	/** The alignment of the object horizontally. */
-	UPROPERTY(EditDefaultsOnly, Category="Layout (Wrap Box Slot)")
+	UPROPERTY(EditDefaultsOnly, BlueprintReadOnly, Category="Layout (Wrap Box Slot)")
 	TEnumAsByte<EHorizontalAlignment> HorizontalAlignment;
 
 	/** The alignment of the object vertically. */
-	UPROPERTY(EditDefaultsOnly, Category="Layout (Wrap Box Slot)")
+	UPROPERTY(EditDefaultsOnly, BlueprintReadOnly, Category="Layout (Wrap Box Slot)")
 	TEnumAsByte<EVerticalAlignment> VerticalAlignment;
 
 public:
