UMG - Fixing some categories for ufields on border slot.

diff --git a/Engine/Source/Runtime/UMG/Public/Components/BorderSlot.h b/Engine/Source/Runtime/UMG/Public/Components/BorderSlot.h
--- a/Engine/Source/Runtime/UMG/Public/Components/BorderSlot.h
+++ b/Engine/Source/Runtime/UMG/Public/Components/BorderSlot.h
@@ -8,7 +8,9 @@
 
 class UBorder;
 
-/** The Slot for the UBorderSlot, contains the widget displayed in a button's single slot */
+/**
+ * The Slot for the UBorderSlot, contains the widget displayed in a border's single slot
+ */
 UCLASS()
 class UMG_API UBorderSlot : public UPanelSlot
 {
@@ -16,26 +18,26 @@ class UMG_API UBorderSlot : public UPanelSlot
 
 public:
 
-	UFUNCTION(BlueprintCallable, Category="Layout (Border Slot)")
+	UFUNCTION(BlueprintCallable, Category="Layout|Border Slot")
 	void SetPadding(FMargin InPadding);
 
-	UFUNCTION(BlueprintCallable, Category="Layout (Border Slot)")
+	UFUNCTION(BlueprintCallable, Category="Layout|Border Slot")
 	void SetHorizontalAlignment(EHorizontalAlignment InHorizontalAlignment);
 
-	UFUNCTION(BlueprintCallable, Category="Layout (Border Slot)")
+	UFUNCTION(BlueprintCallable, Category="Layout|Border Slot")
 	void SetVerticalAlignment(EVerticalAlignment InVerticalAlignment);
 
 protected:
 	/** The padding area between the slot and the content it contains. */
-	UPROPERTY(EditDefaultsOnly, Category="Layout (Border Slot)")
+	UPROPERTY(EditDefaultsOnly, Category="Layout|Border Slot")
 	FMargin Padding;
 
 	/** The alignment of the object horizontally. */
-	UPROPERTY(EditDefaultsOnly, Category="Layout (Border Slot)")
+	UPROPERTY(EditDefaultsOnly, Category="Layout|Border Slot")
 	TEnumAsByte<EHorizontalAlignment> HorizontalAlignment;
 
 	/** The alignment of the object vertically. */
-	UPROPERTY(EditDefaultsOnly, Category="Layout (Border Slot)")
+	UPROPERTY(EditDefaultsOnly, Category="Layout|Border Slot")
 	TEnumAsByte<EVerticalAlignment> VerticalAlignment;
 
 public:
