UMG - Removing ClassGroup from CheckBox.

diff --git a/Engine/Source/Runtime/UMG/Public/Components/CheckBox.h b/Engine/Source/Runtime/UMG/Public/Components/CheckBox.h
--- a/Engine/Source/Runtime/UMG/Public/Components/CheckBox.h
+++ b/Engine/Source/Runtime/UMG/Public/Components/CheckBox.h
@@ -11,7 +11,7 @@ DECLARE_DYNAMIC_MULTICAST_DELEGATE_OneParam( FOnCheckBoxComponentStateChanged, b
 /**
  * Check box widget
  */
-UCLASS(ClassGroup=UserInterface)
+UCLASS()
 class UMG_API UCheckBox : public UContentWidget
 {
 	GENERATED_UCLASS_BODY()
