Fix unable to save due to GLEO error with widget animations

diff --git a/Engine/Source/Runtime/UMG/Private/WidgetBlueprintGeneratedClass.cpp b/Engine/Source/Runtime/UMG/Private/WidgetBlueprintGeneratedClass.cpp
--- a/Engine/Source/Runtime/UMG/Private/WidgetBlueprintGeneratedClass.cpp
+++ b/Engine/Source/Runtime/UMG/Private/WidgetBlueprintGeneratedClass.cpp
@@ -67,13 +67,15 @@ void UWidgetBlueprintGeneratedClass::InitializeWidget(UUserWidget* UserWidget) c
 
 		for(UWidgetAnimation* Animation : Animations)
 		{
-			if( Animation->MovieScene )
+			UWidgetAnimation* Anim = DuplicateObject<UWidgetAnimation>( Animation, UserWidget );
+
+			if( Anim->MovieScene )
 			{
 				// Find property with the same name as the template and assign the new widget to it.
-				UObjectPropertyBase* Prop = FindField<UObjectPropertyBase>(WidgetBlueprintClass, Animation->MovieScene->GetFName());
+				UObjectPropertyBase* Prop = FindField<UObjectPropertyBase>(WidgetBlueprintClass, Anim->MovieScene->GetFName());
 				if(Prop)
 				{
-					Prop->SetObjectPropertyValue_InContainer(UserWidget, Animation);
+					Prop->SetObjectPropertyValue_InContainer(UserWidget, Anim);
 				}
 			}
 
