UMG - Adding a lot more documentation to several functions on a user widget.  Removing GetVisibility from UserWidget, it can do nothing but confuse people, there's already a GetVisibility on UWidget.  Little bug fix on named slots to allow them to be filled at runtime.

diff --git a/Engine/Source/Runtime/UMG/Private/UserWidget.cpp b/Engine/Source/Runtime/UMG/Private/UserWidget.cpp
--- a/Engine/Source/Runtime/UMG/Private/UserWidget.cpp
+++ b/Engine/Source/Runtime/UMG/Private/UserWidget.cpp
@@ -510,6 +510,14 @@ void UUserWidget::SetContentForSlot(FName SlotName, UWidget* Content)
 
 		NamedSlotBindings.Add(NewBinding);
 	}
+
+	// Dynamically insert the new widget into the hierarchy if it exists.
+	if ( WidgetTree )
+	{
+		UNamedSlot* NamedSlot = Cast<UNamedSlot>(WidgetTree->FindWidget(SlotName));
+		NamedSlot->ClearChildren();
+		NamedSlot->AddChild(Content);
+	}
 }
 
 void UUserWidget::NativeTick(const FGeometry& MyGeometry, float InDeltaTime )
@@ -625,18 +633,6 @@ bool UUserWidget::IsInViewport() const
 	return FullScreenWidget.IsValid();
 }
 
-TEnumAsByte<ESlateVisibility::Type> UUserWidget::GetVisiblity() const
-{
-	if ( FullScreenWidget.IsValid() )
-	{
-		TSharedPtr<SWidget> RootWidget = FullScreenWidget.Pin();
-
-		return UWidget::ConvertRuntimeToSerializedVisiblity(RootWidget->GetVisibility());
-	}
-
-	return ESlateVisibility::Collapsed;
-}
-
 void UUserWidget::SetPlayerContext(FLocalPlayerContext InPlayerContext)
 {
 	PlayerContext = InPlayerContext;
diff --git a/Engine/Source/Runtime/UMG/Public/Blueprint/UserWidget.h b/Engine/Source/Runtime/UMG/Public/Blueprint/UserWidget.h
--- a/Engine/Source/Runtime/UMG/Public/Blueprint/UserWidget.h
+++ b/Engine/Source/Runtime/UMG/Public/Blueprint/UserWidget.h
@@ -182,24 +182,30 @@ public:
 	UFUNCTION(BlueprintPure, Category="Appearance")
 	bool IsInViewport() const;
 
-	UFUNCTION(BlueprintPure, Category="Appearance")
-	TEnumAsByte<ESlateVisibility::Type> GetVisiblity() const;
-
 	/** Sets the player context associated with this UI. */
 	void SetPlayerContext(FLocalPlayerContext InPlayerContext);
 
 	/** Gets the player context associated with this UI. */
 	const FLocalPlayerContext& GetPlayerContext() const;
 
-	/** Gets the local player associated with this UI. */
+	/**
+	 * Gets the local player associated with this UI.
+	 * @return The owning local player.
+	 */
 	UFUNCTION(BlueprintCallable, Category="Player")
 	class ULocalPlayer* GetOwningLocalPlayer() const;
 
-	/** Gets the player controller associated with this UI. */
+	/**
+	 * Gets the player controller associated with this UI.
+	 * @return The player controller that owns the UI.
+	 */
 	UFUNCTION(BlueprintCallable, Category="Player")
 	class APlayerController* GetOwningPlayer() const;
 
-	/** Gets the player pawn associated with this UI. */
+	/**
+	 * Gets the player pawn associated with this UI.
+	 * @return Gets the owning player pawn that's owned by the player controller assigned to this widget.
+	 */
 	UFUNCTION(BlueprintCallable, Category="Player")
 	class APawn* GetOwningPlayerPawn() const;
 
@@ -210,90 +216,320 @@ public:
 	UFUNCTION(BlueprintNativeEvent, Category="User Interface", meta=(Keywords="Begin Play"))
 	void Construct();
 
+	/**
+	 * Ticks this widget.  Override in derived classes, but always call the parent implementation.
+	 *
+	 * @param  MyGeometry The space allotted for this widget
+	 * @param  InDeltaTime  Real time passed since last tick
+	 */
 	UFUNCTION(BlueprintNativeEvent, Category="User Interface")
 	void Tick(FGeometry MyGeometry, float InDeltaTime);
 
-	//TODO UMG HitTest
-
 	UFUNCTION(BlueprintNativeEvent, Category="User Interface | Painting")
 	void OnPaint(UPARAM(ref) FPaintContext& Context) const;
 
+	/**
+	 * Called when keyboard focus is given to this widget.  This event does not bubble.
+	 *
+	 * @param MyGeometry The Geometry of the widget receiving the event
+	 * @param InKeyboardFocusEvent  KeyboardFocusEvent
+	 * @return  Returns whether the event was handled, along with other possible actions
+	 */
 	UFUNCTION(BlueprintNativeEvent, Category="Keyboard")
 	FEventReply OnKeyboardFocusReceived(FGeometry MyGeometry, FKeyboardFocusEvent InKeyboardFocusEvent);
+
+	/**
+	 * Called when this widget loses the keyboard focus.  This event does not bubble.
+	 *
+	 * @param  InKeyboardFocusEvent  KeyboardFocusEvent
+	 */
 	UFUNCTION(BlueprintNativeEvent, Category="Keyboard")
 	void OnKeyboardFocusLost(FKeyboardFocusEvent InKeyboardFocusEvent);
-	//UFUNCTION(BlueprintNativeEvent, Category="User Interface")
-	//void OnKeyboardFocusChanging(FWeakWidgetPath PreviousFocusPath, FWidgetPath NewWidgetPath);
 
+	/**
+	 * Called after a character is entered while this widget has keyboard focus
+	 *
+	 * @param MyGeometry The Geometry of the widget receiving the event
+	 * @param  InCharacterEvent  Character event
+	 * @return  Returns whether the event was handled, along with other possible actions
+	 */
 	UFUNCTION(BlueprintNativeEvent, Category="Keyboard")
 	FEventReply OnKeyChar(FGeometry MyGeometry, FCharacterEvent InCharacterEvent);
+
+	/**
+	 * Called after a key is pressed when this widget or a child of this widget has keyboard focus
+	 * If a widget handles this event, OnKeyDown will *not* be passed to the focused widget.
+	 *
+	 * This event is primarily to allow parent widgets to consume an event before a child widget processes
+	 * it and it should be used only when there is no better design alternative.
+	 *
+	 * @param MyGeometry The Geometry of the widget receiving the event
+	 * @param  InKeyboardEvent  Keyboard event
+	 * @return  Returns whether the event was handled, along with other possible actions
+	 */
 	UFUNCTION(BlueprintNativeEvent, Category="Keyboard")
 	FEventReply OnPreviewKeyDown(FGeometry MyGeometry, FKeyboardEvent InKeyboardEvent);
+
+	/**
+	 * Called after a key is pressed when this widget has keyboard focus (this event bubbles if not handled)
+	 *
+	 * @param MyGeometry The Geometry of the widget receiving the event
+	 * @param  InKeyboardEvent  Keyboard event
+	 * @return  Returns whether the event was handled, along with other possible actions
+	 */
 	UFUNCTION(BlueprintNativeEvent, Category="Keyboard")
 	FEventReply OnKeyDown(FGeometry MyGeometry, FKeyboardEvent InKeyboardEvent);
+
+	/**
+	 * Called after a key is released when this widget has keyboard focus
+	 *
+	 * @param MyGeometry The Geometry of the widget receiving the event
+	 * @param  InKeyboardEvent  Keyboard event
+	 * @return  Returns whether the event was handled, along with other possible actions
+	 */
 	UFUNCTION(BlueprintNativeEvent, Category="Keyboard")
 	FEventReply OnKeyUp(FGeometry MyGeometry, FKeyboardEvent InKeyboardEvent);
 
+	/**
+	 * The system calls this method to notify the widget that a mouse button was pressed within it. This event is bubbled.
+	 *
+	 * @param MyGeometry The Geometry of the widget receiving the event
+	 * @param MouseEvent Information about the input event
+	 * @return Whether the event was handled along with possible requests for the system to take action.
+	 */
 	UFUNCTION(BlueprintNativeEvent, Category="Mouse")
 	FEventReply OnMouseButtonDown(FGeometry MyGeometry, const FPointerEvent& MouseEvent);
+
+	/**
+	 * Just like OnMouseButtonDown, but tunnels instead of bubbling.
+	 * If this even is handled, OnMouseButtonDown will not be sent.
+	 * 
+	 * Use this event sparingly as preview events generally make UIs more
+	 * difficult to reason about.
+	 *
+	 * @param MyGeometry The Geometry of the widget receiving the event
+	 * @param MouseEvent Information about the input event
+	 * @return Whether the event was handled along with possible requests for the system to take action.
+	 */
 	UFUNCTION(BlueprintNativeEvent, Category="Mouse")
 	FEventReply OnPreviewMouseButtonDown(FGeometry MyGeometry, const FPointerEvent& MouseEvent);
+
+	/**
+	 * The system calls this method to notify the widget that a mouse button was release within it. This event is bubbled.
+	 *
+	 * @param MyGeometry The Geometry of the widget receiving the event
+	 * @param MouseEvent Information about the input event
+	 * @return Whether the event was handled along with possible requests for the system to take action.
+	 */
 	UFUNCTION(BlueprintNativeEvent, Category="Mouse")
 	FEventReply OnMouseButtonUp(FGeometry MyGeometry, const FPointerEvent& MouseEvent);
+
+	/**
+	 * The system calls this method to notify the widget that a mouse moved within it. This event is bubbled.
+	 *
+	 * @param MyGeometry The Geometry of the widget receiving the event
+	 * @param MouseEvent Information about the input event
+	 * @return Whether the event was handled along with possible requests for the system to take action.
+	 */
 	UFUNCTION(BlueprintNativeEvent, Category="Mouse")
 	FEventReply OnMouseMove(FGeometry MyGeometry, const FPointerEvent& MouseEvent);
+
+	/**
+	 * The system will use this event to notify a widget that the cursor has entered it. This event is NOT bubbled.
+	 *
+	 * @param MyGeometry The Geometry of the widget receiving the event
+	 * @param MouseEvent Information about the input event
+	 */
 	UFUNCTION(BlueprintNativeEvent, Category="Mouse")
 	void OnMouseEnter(FGeometry MyGeometry, const FPointerEvent& MouseEvent);
+
+	/**
+	 * The system will use this event to notify a widget that the cursor has left it. This event is NOT bubbled.
+	 *
+	 * @param MouseEvent Information about the input event
+	 */
 	UFUNCTION(BlueprintNativeEvent, Category="Mouse")
 	void OnMouseLeave(const FPointerEvent& MouseEvent);
+
+	/**
+	 * Called when the mouse wheel is spun. This event is bubbled.
+	 *
+	 * @param  MouseEvent  Mouse event
+	 * @return  Returns whether the event was handled, along with other possible actions
+	 */
 	UFUNCTION(BlueprintNativeEvent, Category="Mouse")
 	FEventReply OnMouseWheel(FGeometry MyGeometry, const FPointerEvent& MouseEvent);
+
+	/**
+	 * Called when a mouse button is double clicked.  Override this in derived classes.
+	 *
+	 * @param  InMyGeometry  Widget geometry
+	 * @param  InMouseEvent  Mouse button event
+	 * @return  Returns whether the event was handled, along with other possible actions
+	 */
 	UFUNCTION(BlueprintNativeEvent, Category="Mouse")
 	FEventReply OnMouseButtonDoubleClick(FGeometry InMyGeometry, const FPointerEvent& InMouseEvent);
 
+	// TODO
 	//UFUNCTION(BlueprintNativeEvent, Category="Mouse")
 	//FCursorReply OnCursorQuery(FGeometry MyGeometry, const FPointerEvent& CursorEvent) const;
 
+	// TODO
+	//virtual bool OnVisualizeTooltip(const TSharedPtr<SWidget>& TooltipContent);
+
+	/**
+	 * Called when Slate detects that a widget started to be dragged.
+	 *
+	 * @param  InMyGeometry  Widget geometry
+	 * @param  PointerEvent  MouseMove that triggered the drag
+	 * @param  Operation     The drag operation that was detected.
+	 */
 	UFUNCTION(BlueprintNativeEvent, Category="Drag and Drop")
 	void OnDragDetected(FGeometry MyGeometry, const FPointerEvent& PointerEvent, UDragDropOperation*& Operation);
 
+	/**
+	 * Called when the user cancels the drag operation, typically when they simply release the mouse button after
+	 * beginning a drag operation, but failing to complete the drag.
+	 *
+	 * @param  PointerEvent  Last mouse event from when the drag was canceled.
+	 * @param  Operation     The drag operation that was canceled.
+	 */
 	UFUNCTION(BlueprintNativeEvent, Category="Drag and Drop")
 	void OnDragCancelled(const FPointerEvent& PointerEvent, UDragDropOperation* Operation);
 	
+	/**
+	 * Called during drag and drop when the drag enters the widget.
+	 *
+	 * @param MyGeometry     The geometry of the widget receiving the event.
+	 * @param PointerEvent   The mouse event from when the drag entered the widget.
+	 * @param Operation      The drag operation that entered the widget.
+	 */
 	UFUNCTION(BlueprintNativeEvent, Category="Drag and Drop")
 	void OnDragEnter(FGeometry MyGeometry, FPointerEvent PointerEvent, UDragDropOperation* Operation);
+
+	/**
+	 * Called during drag and drop when the drag leaves the widget.
+	 *
+	 * @param PointerEvent   The mouse event from when the drag left the widget.
+	 * @param Operation      The drag operation that entered the widget.
+	 */
 	UFUNCTION(BlueprintNativeEvent, Category="Drag and Drop")
 	void OnDragLeave(FPointerEvent PointerEvent, UDragDropOperation* Operation);
+
+	/**
+	 * Called during drag and drop when the the mouse is being dragged over a widget.
+	 *
+	 * @param MyGeometry     The geometry of the widget receiving the event.
+	 * @param PointerEvent   The mouse event from when the drag occurred over the widget.
+	 * @param Operation      The drag operation over the widget.
+	 *
+	 * @return 'true' to indicate that you handled the drag over operation.  Returning 'false' will cause the operation to continue to bubble up.
+	 */
 	UFUNCTION(BlueprintNativeEvent, Category="Drag and Drop")
 	bool OnDragOver(FGeometry MyGeometry, FPointerEvent PointerEvent, UDragDropOperation* Operation);
+
+	/**
+	 * Called when the user is dropping something onto a widget.  Ends the drag and drop operation, even if no widget handles this.
+	 *
+	 * @param MyGeometry     The geometry of the widget receiving the event.
+	 * @param PointerEvent   The mouse event from when the drag occurred over the widget.
+	 * @param Operation      The drag operation over the widget.
+	 * 
+	 * @return 'true' to indicate you handled the drop operation.
+	 */
 	UFUNCTION(BlueprintNativeEvent, Category="Drag and Drop")
 	bool OnDrop(FGeometry MyGeometry, FPointerEvent PointerEvent, UDragDropOperation* Operation);
 
+	/**
+	 * Called when a controller button is pressed
+	 * 
+	 * @param MyGeometry        The geometry of the widget receiving the event.
+	 * @param ControllerEvent	The controller event generated
+	 */
 	UFUNCTION(BlueprintNativeEvent, Category="Gamepad Input")
 	FEventReply OnControllerButtonPressed(FGeometry MyGeometry, FControllerEvent ControllerEvent);
+
+	/**
+	 * Called when a controller button is released
+	 * 
+	 * @param MyGeometry        The geometry of the widget receiving the event.
+	 * @param ControllerEvent	The controller event generated
+	 */
 	UFUNCTION(BlueprintNativeEvent, Category="Gamepad Input")
 	FEventReply OnControllerButtonReleased(FGeometry MyGeometry, FControllerEvent ControllerEvent);
+
+	/**
+	 * Called when an analog value on a controller changes, like a joystick.
+	 * 
+	 * @param MyGeometry        The geometry of the widget receiving the event.
+	 * @param ControllerEvent	The controller event generated
+	 */
 	UFUNCTION(BlueprintNativeEvent, Category="Gamepad Input")
 	FEventReply OnControllerAnalogValueChanged(FGeometry MyGeometry, FControllerEvent ControllerEvent);
 
+	/**
+	 * Called when the user performs a gesture on trackpad. This event is bubbled.
+	 *
+	 * @param MyGeometry     The geometry of the widget receiving the event.
+	 * @param  GestureEvent  gesture event
+	 * @return  Returns whether the event was handled, along with other possible actions
+	 */
 	UFUNCTION(BlueprintNativeEvent, Category="Touch Input")
 	FEventReply OnTouchGesture(FGeometry MyGeometry, const FPointerEvent& GestureEvent);
+
+	/**
+	 * Called when a touchpad touch is started (finger down)
+	 * 
+	 * @param MyGeometry    The geometry of the widget receiving the event.
+	 * @param InTouchEvent	The touch event generated
+	 */
 	UFUNCTION(BlueprintNativeEvent, Category="Touch Input")
 	FEventReply OnTouchStarted(FGeometry MyGeometry, const FPointerEvent& InTouchEvent);
+
+	/**
+	 * Called when a touchpad touch is moved  (finger moved)
+	 * 
+	 * @param MyGeometry    The geometry of the widget receiving the event.
+	 * @param InTouchEvent	The touch event generated
+	 */
 	UFUNCTION(BlueprintNativeEvent, Category="Touch Input")
 	FEventReply OnTouchMoved(FGeometry MyGeometry, const FPointerEvent& InTouchEvent);
+
+	/**
+	 * Called when a touchpad touch is ended (finger lifted)
+	 * 
+	 * @param MyGeometry    The geometry of the widget receiving the event.
+	 * @param InTouchEvent	The touch event generated
+	 */
 	UFUNCTION(BlueprintNativeEvent, Category="Touch Input")
 	FEventReply OnTouchEnded(FGeometry MyGeometry, const FPointerEvent& InTouchEvent);
+
+	/**
+	 * Called when motion is detected (controller or device)
+	 * e.g. Someone tilts or shakes their controller.
+	 * 
+	 * @param MyGeometry    The geometry of the widget receiving the event.
+	 * @param MotionEvent	The motion event generated
+	 */
 	UFUNCTION(BlueprintNativeEvent, Category="Touch Input")
 	FEventReply OnMotionDetected(FGeometry MyGeometry, FMotionEvent InMotionEvent);
 
-	//virtual bool OnVisualizeTooltip(const TSharedPtr<SWidget>& TooltipContent);
+public:
 
-	/**  */
+	/**
+	 * Sets the tint of the widget, this affects all child widgets.
+	 * 
+	 * @param InColorAndOpacity	The tint to apply to all child widgets.
+	 */
 	UFUNCTION(BlueprintCallable, Category="Appearance")
 	void SetColorAndOpacity(FLinearColor InColorAndOpacity);
 
-	/**  */
+	/**
+	 * Sets the foreground color of the widget, this is inherited by sub widgets.  Any color property 
+	 * that is marked as inherit will use this color.
+	 * 
+	 * @param InForegroundColor	The foreground color.
+	 */
 	UFUNCTION(BlueprintCallable, Category="Appearance")
 	void SetForegroundColor(FSlateColor InForegroundColor);
 
@@ -350,19 +586,20 @@ public:
 #endif
 
 public:
-	/**  */
+	/** The color and opacity of this widget.  Tints all child widgets. */
 	UPROPERTY(EditDefaultsOnly, BlueprintReadOnly, Category="Style")
 	FLinearColor ColorAndOpacity;
 
-	/**  */
 	UPROPERTY()
 	FGetLinearColor ColorAndOpacityDelegate;
 
-	/**  */
+	/**
+	 * The foreground color of the widget, this is inherited by sub widgets.  Any color property
+	 * that is marked as inherit will use this color.
+	 */
 	UPROPERTY(EditDefaultsOnly, BlueprintReadOnly, Category="Style")
 	FSlateColor ForegroundColor;
 
-	/**  */
 	UPROPERTY()
 	FGetSlateColor ForegroundColorDelegate;
 
