UMG - Adding more documentation to the top of several widgets.

diff --git a/Engine/Source/Runtime/UMG/Public/Components/Border.h b/Engine/Source/Runtime/UMG/Public/Components/Border.h
--- a/Engine/Source/Runtime/UMG/Public/Components/Border.h
+++ b/Engine/Source/Runtime/UMG/Public/Components/Border.h
@@ -1,4 +1,4 @@
-// Copyright 1998-2014 Epic Games, Inc. All Rights Reserved.
+﻿// Copyright 1998-2014 Epic Games, Inc. All Rights Reserved.
 
 #pragma once
 
@@ -8,7 +8,9 @@ class USlateBrushAsset;
 
 /**
  * A border is a container widget that can contain one child widget, providing an opportunity 
- * to surround it with a border image and adjustable padding.
+ * to surround it with a background image and adjustable padding.
+ *
+ * ● Single Child
  */
 UCLASS()
 class UMG_API UBorder : public UContentWidget
diff --git a/Engine/Source/Runtime/UMG/Public/Components/Button.h b/Engine/Source/Runtime/UMG/Public/Components/Button.h
--- a/Engine/Source/Runtime/UMG/Public/Components/Button.h
+++ b/Engine/Source/Runtime/UMG/Public/Components/Button.h
@@ -1,4 +1,4 @@
-// Copyright 1998-2014 Epic Games, Inc. All Rights Reserved.
+﻿// Copyright 1998-2014 Epic Games, Inc. All Rights Reserved.
 
 #pragma once
 
@@ -9,7 +9,12 @@
 DECLARE_DYNAMIC_MULTICAST_DELEGATE(FOnButtonClickedEvent);
 
 /**
- * The button is a clickable primitive widget to enable basic interaction.
+ * The button is a click-able primitive widget to enable basic interaction, you
+ * can place any other widget inside a button to make a more complex and 
+ * interesting click-able element in your UI.
+ *
+ * ● Single Child
+ * ● Clickable
  */
 UCLASS()
 class UMG_API UButton : public UContentWidget
diff --git a/Engine/Source/Runtime/UMG/Public/Components/CanvasPanel.h b/Engine/Source/Runtime/UMG/Public/Components/CanvasPanel.h
--- a/Engine/Source/Runtime/UMG/Public/Components/CanvasPanel.h
+++ b/Engine/Source/Runtime/UMG/Public/Components/CanvasPanel.h
@@ -1,4 +1,4 @@
-// Copyright 1998-2014 Epic Games, Inc. All Rights Reserved.
+﻿// Copyright 1998-2014 Epic Games, Inc. All Rights Reserved.
 
 #pragma once
 #include "Components/PanelWidget.h"
@@ -11,6 +11,10 @@ class UCanvasPanelSlot;
  * locations, anchored and z-ordered with other children of the canvas.  The canvas is a great widget
  * for manual layout, but bad when you want to procedurally just generate widgets and place them in a 
  * container (unless you want absolute layout).
+ *
+ * ● Many Children
+ * ● Absolute Layout
+ * ● Anchors
  */
 UCLASS()
 class UMG_API UCanvasPanel : public UPanelWidget
diff --git a/Engine/Source/Runtime/UMG/Public/Components/CheckBox.h b/Engine/Source/Runtime/UMG/Public/Components/CheckBox.h
--- a/Engine/Source/Runtime/UMG/Public/Components/CheckBox.h
+++ b/Engine/Source/Runtime/UMG/Public/Components/CheckBox.h
@@ -1,4 +1,4 @@
-// Copyright 1998-2014 Epic Games, Inc. All Rights Reserved.
+﻿// Copyright 1998-2014 Epic Games, Inc. All Rights Reserved.
 
 #pragma once
 
@@ -9,7 +9,12 @@
 DECLARE_DYNAMIC_MULTICAST_DELEGATE_OneParam( FOnCheckBoxComponentStateChanged, bool, bIsChecked );
 
 /**
- * Check box widget
+ * The checkbox widget allows you to display a toggled state of 'unchecked', 'checked' and 
+ * 'indeterminable.  You can use the checkbox for a classic checkbox, or as a toggle button,
+ * or as radio buttons.
+ * 
+ * ● Single Child
+ * ● Toggle
  */
 UCLASS()
 class UMG_API UCheckBox : public UContentWidget
diff --git a/Engine/Source/Runtime/UMG/Public/Components/CircularThrobber.h b/Engine/Source/Runtime/UMG/Public/Components/CircularThrobber.h
--- a/Engine/Source/Runtime/UMG/Public/Components/CircularThrobber.h
+++ b/Engine/Source/Runtime/UMG/Public/Components/CircularThrobber.h
@@ -1,4 +1,4 @@
-// Copyright 1998-2014 Epic Games, Inc. All Rights Reserved.
+﻿// Copyright 1998-2014 Epic Games, Inc. All Rights Reserved.
 
 #pragma once
 
@@ -8,6 +8,9 @@ class USlateBrushAsset;
 
 /**
  * A throbber widget that orients images in a spinning circle.
+ * 
+ * ● No Children
+ * ● Spinner Progress
  */
 UCLASS()
 class UMG_API UCircularThrobber : public UWidget
diff --git a/Engine/Source/Runtime/UMG/Public/Components/ContentWidget.h b/Engine/Source/Runtime/UMG/Public/Components/ContentWidget.h
--- a/Engine/Source/Runtime/UMG/Public/Components/ContentWidget.h
+++ b/Engine/Source/Runtime/UMG/Public/Components/ContentWidget.h
@@ -4,6 +4,7 @@
 
 #include "ContentWidget.generated.h"
 
+/**  */
 UCLASS(Abstract)
 class UMG_API UContentWidget : public UPanelWidget
 {
diff --git a/Engine/Source/Runtime/UMG/Public/Components/EditableTextBox.h b/Engine/Source/Runtime/UMG/Public/Components/EditableTextBox.h
--- a/Engine/Source/Runtime/UMG/Public/Components/EditableTextBox.h
+++ b/Engine/Source/Runtime/UMG/Public/Components/EditableTextBox.h
@@ -1,4 +1,4 @@
-// Copyright 1998-2014 Epic Games, Inc. All Rights Reserved.
+﻿// Copyright 1998-2014 Epic Games, Inc. All Rights Reserved.
 
 #pragma once
 
@@ -6,6 +6,9 @@
 
 /**
  * Allows the user to type in custom text.  Only permits a single line of text to be entered.
+ * 
+ * ● No Children
+ * ● Text Entry
  */
 UCLASS(meta=(DisplayName="Text Box"))
 class UMG_API UEditableTextBox : public UWidget
diff --git a/Engine/Source/Runtime/UMG/Public/Components/GridPanel.h b/Engine/Source/Runtime/UMG/Public/Components/GridPanel.h
--- a/Engine/Source/Runtime/UMG/Public/Components/GridPanel.h
+++ b/Engine/Source/Runtime/UMG/Public/Components/GridPanel.h
@@ -1,4 +1,4 @@
-// Copyright 1998-2014 Epic Games, Inc. All Rights Reserved.
+﻿// Copyright 1998-2014 Epic Games, Inc. All Rights Reserved.
 
 #pragma once
 
@@ -6,6 +6,8 @@
 
 /**
  * A panel that evenly divides up available space between all of its children.
+ * 
+ * ● Many Children
  */
 UCLASS()
 class UMG_API UGridPanel : public UPanelWidget
diff --git a/Engine/Source/Runtime/UMG/Public/Components/HorizontalBox.h b/Engine/Source/Runtime/UMG/Public/Components/HorizontalBox.h
--- a/Engine/Source/Runtime/UMG/Public/Components/HorizontalBox.h
+++ b/Engine/Source/Runtime/UMG/Public/Components/HorizontalBox.h
@@ -1,4 +1,4 @@
-// Copyright 1998-2014 Epic Games, Inc. All Rights Reserved.
+﻿// Copyright 1998-2014 Epic Games, Inc. All Rights Reserved.
 
 #pragma once
 #include "Components/PanelWidget.h"
@@ -8,6 +8,9 @@ class UHorizontalBoxSlot;
 
 /**
  * Allows widgets to be laid out in a flow horizontally.
+ *
+ * ● Many Children
+ * ● Flow Horizontal
  */
 UCLASS(ClassGroup=UserInterface)
 class UMG_API UHorizontalBox : public UPanelWidget
diff --git a/Engine/Source/Runtime/UMG/Public/Components/Image.h b/Engine/Source/Runtime/UMG/Public/Components/Image.h
--- a/Engine/Source/Runtime/UMG/Public/Components/Image.h
+++ b/Engine/Source/Runtime/UMG/Public/Components/Image.h
@@ -1,4 +1,4 @@
-// Copyright 1998-2014 Epic Games, Inc. All Rights Reserved.
+﻿// Copyright 1998-2014 Epic Games, Inc. All Rights Reserved.
 
 #pragma once
 
@@ -8,6 +8,8 @@ class USlateBrushAsset;
 
 /**
  * The image widget allows you to display a Slate Brush, or texture or material in the UI.
+ *
+ * ● No Children
  */
 UCLASS()
 class UMG_API UImage : public UWidget
diff --git a/Engine/Source/Runtime/UMG/Public/Components/ProgressBar.h b/Engine/Source/Runtime/UMG/Public/Components/ProgressBar.h
--- a/Engine/Source/Runtime/UMG/Public/Components/ProgressBar.h
+++ b/Engine/Source/Runtime/UMG/Public/Components/ProgressBar.h
@@ -1,4 +1,4 @@
-// Copyright 1998-2014 Epic Games, Inc. All Rights Reserved.
+﻿// Copyright 1998-2014 Epic Games, Inc. All Rights Reserved.
 
 #pragma once
 
@@ -8,6 +8,8 @@
 
 /**
  * The progress bar widget is a simple bar that fills up that can be restyled to fit any number of uses.
+ *
+ * ● No Children
  */
 UCLASS(ClassGroup=UserInterface)
 class UMG_API UProgressBar : public UWidget
diff --git a/Engine/Source/Runtime/UMG/Public/Components/Spacer.h b/Engine/Source/Runtime/UMG/Public/Components/Spacer.h
--- a/Engine/Source/Runtime/UMG/Public/Components/Spacer.h
+++ b/Engine/Source/Runtime/UMG/Public/Components/Spacer.h
@@ -1,4 +1,4 @@
-// Copyright 1998-2014 Epic Games, Inc. All Rights Reserved.
+﻿// Copyright 1998-2014 Epic Games, Inc. All Rights Reserved.
 
 #pragma once
 
@@ -6,6 +6,8 @@
 
 /**
  * A spacer widget; it does not have a visual representation, and just provides padding between other widgets.
+ *
+ * ● No Children
  */
 UCLASS()
 class UMG_API USpacer : public UWidget
diff --git a/Engine/Source/Runtime/UMG/Public/Components/TextBlock.h b/Engine/Source/Runtime/UMG/Public/Components/TextBlock.h
--- a/Engine/Source/Runtime/UMG/Public/Components/TextBlock.h
+++ b/Engine/Source/Runtime/UMG/Public/Components/TextBlock.h
@@ -1,11 +1,14 @@
-// Copyright 1998-2014 Epic Games, Inc. All Rights Reserved.
+﻿// Copyright 1998-2014 Epic Games, Inc. All Rights Reserved.
 
 #pragma once
 
 #include "TextBlock.generated.h"
 
 /**
- * A simple static text widget
+ * A simple static text widget.
+ *
+ * ● No Children
+ * ● Text
  */
 UCLASS(meta=(DisplayName="Text"))
 class UMG_API UTextBlock : public UWidget
@@ -106,17 +109,17 @@ public:
 	//SLATE_EVENT(FOnClicked, OnDoubleClicked)
 
 	/** 
-	Gets the widget text
-	@return The widget text
-	*/
-	UFUNCTION(BlueprintCallable, Category = "Widget")
+	 * Gets the widget text
+	 * @return The widget text
+	 */
+	UFUNCTION(BlueprintCallable, Category="Widget")
 	FText GetText() const;
 
 	/**
-	Directly sets the widget text.
-	Warning: This will wipe any binding created for the Text property!
-	@param InText The text to assign to the widget
-	*/
+	 * Directly sets the widget text.
+	 * Warning: This will wipe any binding created for the Text property!
+	 * @param InText The text to assign to the widget
+	 */
 	UFUNCTION(BlueprintCallable, Category="Widget")
 	void SetText(FText InText);
 
diff --git a/Engine/Source/Runtime/UMG/Public/Components/VerticalBox.h b/Engine/Source/Runtime/UMG/Public/Components/VerticalBox.h
--- a/Engine/Source/Runtime/UMG/Public/Components/VerticalBox.h
+++ b/Engine/Source/Runtime/UMG/Public/Components/VerticalBox.h
@@ -1,4 +1,4 @@
-// Copyright 1998-2014 Epic Games, Inc. All Rights Reserved.
+﻿// Copyright 1998-2014 Epic Games, Inc. All Rights Reserved.
 
 #pragma once
 #include "Components/PanelWidget.h"
@@ -9,6 +9,9 @@ class UVerticalBoxSlot;
 /**
  * A vertical box widget is a layout panel allowing child widgets to be automatically laid out
  * vertically.
+ *
+ * ● Many Children
+ * ● Flows Vertical
  */
 UCLASS()
 class UMG_API UVerticalBox : public UPanelWidget
diff --git a/Engine/Source/Runtime/UMG/Public/Components/WrapBox.h b/Engine/Source/Runtime/UMG/Public/Components/WrapBox.h
--- a/Engine/Source/Runtime/UMG/Public/Components/WrapBox.h
+++ b/Engine/Source/Runtime/UMG/Public/Components/WrapBox.h
@@ -1,4 +1,4 @@
-// Copyright 1998-2014 Epic Games, Inc. All Rights Reserved.
+﻿// Copyright 1998-2014 Epic Games, Inc. All Rights Reserved.
 
 #pragma once
 
@@ -6,6 +6,10 @@
 
 /**
  * Arranges widgets left-to-right.  When the widgets exceed the Width it will place widgets on the next line.
+ * 
+ * ● Many Children
+ * ● Flows
+ * ● Wraps
  */
 UCLASS()
 class UMG_API UWrapBox : public UPanelWidget
