UMG - Adding a way to query slate application's drag/drop status.  Adds a layer on top of slate application's, only returns valid data if it's a UMG drag drop operation.

diff --git a/Engine/Source/Runtime/UMG/Private/WidgetBlueprintLibrary.cpp b/Engine/Source/Runtime/UMG/Private/WidgetBlueprintLibrary.cpp
--- a/Engine/Source/Runtime/UMG/Private/WidgetBlueprintLibrary.cpp
+++ b/Engine/Source/Runtime/UMG/Private/WidgetBlueprintLibrary.cpp
@@ -2,6 +2,7 @@
 
 #include "UMGPrivatePCH.h"
 #include "Slate/SlateBrushAsset.h"
+#include "Slate/UMGDragDropOp.h"
 #include "WidgetBlueprintLibrary.h"
 
 #define LOCTEXT_NAMESPACE "UMG"
@@ -236,6 +237,32 @@ FEventReply UWidgetBlueprintLibrary::EndDragDrop(UPARAM(ref) FEventReply& Reply)
 	return Reply;
 }
 
+bool UWidgetBlueprintLibrary::IsDragDropping()
+{
+	if ( FSlateApplication::Get().IsDragDropping() )
+	{
+		TSharedPtr<FDragDropOperation> SlateDragOp = FSlateApplication::Get().GetDragDroppingContent();
+		if ( SlateDragOp.IsValid() && SlateDragOp->IsOfType<FUMGDragDropOp>() )
+		{
+			return true;
+		}
+	}
+
+	return false;
+}
+
+UDragDropOperation* UWidgetBlueprintLibrary::GetDragDroppingContent()
+{
+	TSharedPtr<FDragDropOperation> SlateDragOp = FSlateApplication::Get().GetDragDroppingContent();
+	if ( SlateDragOp.IsValid() && SlateDragOp->IsOfType<FUMGDragDropOp>() )
+	{
+		TSharedPtr<FUMGDragDropOp> UMGDragDropOp = StaticCastSharedPtr<FUMGDragDropOp>(SlateDragOp);
+		return UMGDragDropOp->GetOperation();
+	}
+
+	return nullptr;
+}
+
 FSlateBrush UWidgetBlueprintLibrary::MakeBrushFromAsset(USlateBrushAsset* BrushAsset)
 {
 	if ( BrushAsset )
diff --git a/Engine/Source/Runtime/UMG/Public/Blueprint/WidgetBlueprintLibrary.h b/Engine/Source/Runtime/UMG/Public/Blueprint/WidgetBlueprintLibrary.h
--- a/Engine/Source/Runtime/UMG/Public/Blueprint/WidgetBlueprintLibrary.h
+++ b/Engine/Source/Runtime/UMG/Public/Blueprint/WidgetBlueprintLibrary.h
@@ -16,14 +16,14 @@ class UWidgetBlueprintLibrary : public UBlueprintFunctionLibrary
 	GENERATED_UCLASS_BODY()
 
 	/** Creates a widget */
-	UFUNCTION(BlueprintCallable, meta=( WorldContext="WorldContextObject", FriendlyName = "Create Widget", BlueprintInternalUseOnly = "true" ), Category="User Interface|Widget")
+	UFUNCTION(BlueprintCallable, meta=( WorldContext="WorldContextObject", FriendlyName = "Create Widget", BlueprintInternalUseOnly = "true" ), Category="Widget")
 	static class UUserWidget* Create(UObject* WorldContextObject, TSubclassOf<class UUserWidget> WidgetType, APlayerController* OwningPlayer);
 
 	/**
 	 * Creates a new drag and drop operation that can be returned from a drag begin to inform the UI what i
 	 * being dragged and dropped and what it looks like.
 	 */
-	UFUNCTION(BlueprintCallable, Category="User Interface|Drag and Drop", meta=( BlueprintInternalUseOnly = "true" ))
+	UFUNCTION(BlueprintCallable, Category="Widget|Drag and Drop", meta=( BlueprintInternalUseOnly="true" ))
 	static UDragDropOperation* CreateDragDropOperation(TSubclassOf<UDragDropOperation> OperationClass);
 	
 	/** Setup an input mode that allows only the UI to respond to user input. */
@@ -77,18 +77,23 @@ class UWidgetBlueprintLibrary : public UBlueprintFunctionLibrary
 	UFUNCTION(BlueprintPure, Category="Widget|Event Reply")
 	static FEventReply Unhandled();
 
+	/**  */
 	UFUNCTION(BlueprintPure, meta=( HidePin="CapturingWidget", DefaultToSelf="CapturingWidget" ), Category="Widget|Event Reply")
 	static FEventReply CaptureMouse(UPARAM(ref) FEventReply& Reply, UWidget* CapturingWidget);
 
+	/**  */
 	UFUNCTION(BlueprintPure, Category="Widget|Event Reply")
 	static FEventReply ReleaseMouseCapture(UPARAM(ref) FEventReply& Reply);
 
+	/**  */
 	UFUNCTION(BlueprintPure, meta=( HidePin="CapturingWidget", DefaultToSelf="CapturingWidget" ), Category="Widget|Event Reply")
 	static FEventReply CaptureJoystick(UPARAM(ref) FEventReply& Reply, UWidget* CapturingWidget, bool bInAllJoysticks = false);
 
+	/**  */
 	UFUNCTION(BlueprintPure, Category="Widget|Event Reply")
 	static FEventReply ReleaseJoystickCapture(UPARAM(ref) FEventReply& Reply, bool bInAllJoysticks = false);
 
+	/**  */
 	UFUNCTION(BlueprintPure, Category="Widget|Event Reply")
 	static FEventReply SetMousePosition(UPARAM(ref) FEventReply& Reply, FVector2D NewMousePosition);
 
@@ -99,18 +104,30 @@ class UWidgetBlueprintLibrary : public UBlueprintFunctionLibrary
 	 * @param WidgetDetectingDrag  Detect dragging in this widget
 	 * @param DragKey		       This button should be pressed to detect the drag
 	 */
-	UFUNCTION(BlueprintPure, meta=( HidePin="WidgetDetectingDrag", DefaultToSelf="WidgetDetectingDrag" ), Category="Widget|Event Reply")
+	UFUNCTION(BlueprintPure, meta=( HidePin="WidgetDetectingDrag", DefaultToSelf="WidgetDetectingDrag" ), Category="Widget|Drag and Drop|Event Reply")
 	static FEventReply DetectDrag(UPARAM(ref) FEventReply& Reply, UWidget* WidgetDetectingDrag, FKey DragKey);
 
-	UFUNCTION(BlueprintCallable, meta=( HidePin="WidgetDetectingDrag", DefaultToSelf="WidgetDetectingDrag" ), Category="Widget|Event Reply")
+	UFUNCTION(BlueprintCallable, meta=( HidePin="WidgetDetectingDrag", DefaultToSelf="WidgetDetectingDrag" ), Category="Widget|Drag and Drop|Event Reply")
 	static FEventReply DetectDragIfPressed(const FPointerEvent& PointerEvent, UWidget* WidgetDetectingDrag, FKey DragKey);
 
 	/**
 	 * An event should return FReply::Handled().EndDragDrop() to request that the current drag/drop operation be terminated.
 	 */
-	UFUNCTION(BlueprintPure, Category="Widget|Event Reply")
+	UFUNCTION(BlueprintPure, Category="Widget|Drag and Drop|Event Reply")
 	static FEventReply EndDragDrop(UPARAM(ref) FEventReply& Reply);
 
+	/**
+	 * Returns true if a drag/drop event is occurring that a widget can handle.
+	 */
+	UFUNCTION(BlueprintPure, Category="Widget|Drag and Drop")
+	static bool IsDragDropping();
+
+	/**
+	 * Returns the drag and drop operation that is currently occuring if any, otherwise nothing.
+	 */
+	UFUNCTION(BlueprintPure, Category="Widget|Drag and Drop")
+	static UDragDropOperation* GetDragDroppingContent();
+
 	/**
 	 * Creates a Slate Brush from a Slate Brush Asset
 	 *
@@ -157,6 +174,6 @@ class UWidgetBlueprintLibrary : public UBlueprintFunctionLibrary
 	static UMaterialInstanceDynamic* GetDynamicMaterial(FSlateBrush& Brush);
 
 	/** Closes any popup menu */
-	UFUNCTION(BlueprintCallable, Category="Widget|Brush")
+	UFUNCTION(BlueprintCallable, Category="Widget|Menu")
 	static void DismissAllMenus();
 };
