UMG - UWidgets now have a simple way to determine if it's the widget is visible in any way, IsVisible.

diff --git a/Engine/Source/Runtime/UMG/Private/Components/Widget.cpp b/Engine/Source/Runtime/UMG/Private/Components/Widget.cpp
--- a/Engine/Source/Runtime/UMG/Private/Components/Widget.cpp
+++ b/Engine/Source/Runtime/UMG/Private/Components/Widget.cpp
@@ -93,6 +93,17 @@ void UWidget::SetIsEnabled(bool bInIsEnabled)
 	}
 }
 
+bool UWidget::IsVisible() const
+{
+	TSharedPtr<SWidget> SafeWidget = GetCachedWidget();
+	if ( SafeWidget.IsValid() )
+	{
+		return SafeWidget->GetVisibility().IsVisible();
+	}
+
+	return false;
+}
+
 TEnumAsByte<ESlateVisibility::Type> UWidget::GetVisibility()
 {
 	TSharedPtr<SWidget> SafeWidget = GetCachedWidget();
diff --git a/Engine/Source/Runtime/UMG/Public/Components/Widget.h b/Engine/Source/Runtime/UMG/Public/Components/Widget.h
--- a/Engine/Source/Runtime/UMG/Public/Components/Widget.h
+++ b/Engine/Source/Runtime/UMG/Public/Components/Widget.h
@@ -163,6 +163,10 @@ public:
 	UFUNCTION(BlueprintCallable, Category="Widget")
 	void SetToolTipText(const FText& InToolTipText);
 
+	/** @return true if the widget is Visible, HitTestInvisible or SelfHitTestInvisible. */
+	UFUNCTION(BlueprintCallable, Category="Widget")
+	bool IsVisible() const;
+
 	/** Gets the current visibility of the widget. */
 	UFUNCTION(BlueprintCallable, Category="Widget")
 	TEnumAsByte<ESlateVisibility::Type> GetVisibility();
