UMG - Adding more docuementation to a few more widgets.

diff --git a/Engine/Source/Runtime/UMG/Public/Components/ScaleBox.h b/Engine/Source/Runtime/UMG/Public/Components/ScaleBox.h
--- a/Engine/Source/Runtime/UMG/Public/Components/ScaleBox.h
+++ b/Engine/Source/Runtime/UMG/Public/Components/ScaleBox.h
@@ -1,4 +1,4 @@
-// Copyright 1998-2014 Epic Games, Inc. All Rights Reserved.
+﻿// Copyright 1998-2014 Epic Games, Inc. All Rights Reserved.
 
 #pragma once
 
@@ -8,6 +8,9 @@
  * Allows you to place content with a desired size and have it scale to meet the constraints placed on this box's alloted area.  If
  * you needed to have a background image scale to fill an area but not become distorted with different aspect ratios, or if you need
  * to auto fit some text to an area, this is the control for you.
+ *
+ * ● Single Child
+ * ● Aspect Ratio
  */
 UCLASS()
 class UMG_API UScaleBox : public UContentWidget
diff --git a/Engine/Source/Runtime/UMG/Public/Components/SizeBox.h b/Engine/Source/Runtime/UMG/Public/Components/SizeBox.h
--- a/Engine/Source/Runtime/UMG/Public/Components/SizeBox.h
+++ b/Engine/Source/Runtime/UMG/Public/Components/SizeBox.h
@@ -1,4 +1,4 @@
-// Copyright 1998-2014 Epic Games, Inc. All Rights Reserved.
+﻿// Copyright 1998-2014 Epic Games, Inc. All Rights Reserved.
 
 #pragma once
 
@@ -9,6 +9,9 @@ class USlateBrushAsset;
 /**
  * A widget that allows you to specify the size it reports to have and desire.  Not all widgets report a desired size
  * that you actually desire.  Wrapping them in a SizeBox lets you have the Size Box force them to be a particular size.
+ *
+ * ● Single Child
+ * ● Fixed Size
  */
 UCLASS()
 class UMG_API USizeBox : public UContentWidget
diff --git a/Engine/Source/Runtime/UMG/Public/Components/Slider.h b/Engine/Source/Runtime/UMG/Public/Components/Slider.h
--- a/Engine/Source/Runtime/UMG/Public/Components/Slider.h
+++ b/Engine/Source/Runtime/UMG/Public/Components/Slider.h
@@ -1,4 +1,4 @@
-// Copyright 1998-2014 Epic Games, Inc. All Rights Reserved.
+﻿// Copyright 1998-2014 Epic Games, Inc. All Rights Reserved.
 
 #pragma once
 
@@ -10,6 +10,8 @@ DECLARE_DYNAMIC_MULTICAST_DELEGATE_OneParam(FOnFloatValueChangedEvent, float, Va
 
 /**
  * A simple widget that shows a sliding bar with a handle that allows you to control the value between 0..1.
+ *
+ * ● No Children
  */
 UCLASS()
 class UMG_API USlider : public UWidget
