[Github PR 528] leobenaducci - Fixed UPanelWidget array-out-of-bounds crash on GetChildAt, RemoveChildAt and ReplaceChildAt.

diff --git a/Engine/Source/Runtime/UMG/Private/Components/PanelWidget.cpp b/Engine/Source/Runtime/UMG/Private/Components/PanelWidget.cpp
--- a/Engine/Source/Runtime/UMG/Private/Components/PanelWidget.cpp
+++ b/Engine/Source/Runtime/UMG/Private/Components/PanelWidget.cpp
@@ -41,6 +41,11 @@ int32 UPanelWidget::GetChildrenCount() const
 
 UWidget* UPanelWidget::GetChildAt(int32 Index) const
 {
+	if ( Index < 0 || Index >= Slots.Num() )
+	{
+		return nullptr;
+	}
+
 	return Slots[Index]->Content;
 }
 
@@ -60,6 +65,11 @@ int32 UPanelWidget::GetChildIndex(UWidget* Content) const
 
 bool UPanelWidget::RemoveChildAt(int32 Index)
 {
+	if ( Index < 0 || Index >= Slots.Num() )
+	{
+		return false;
+	}
+
 	UPanelSlot* Slot = Slots[Index];
 	if ( Slot->Content )
 	{
@@ -105,8 +115,13 @@ UPanelSlot* UPanelWidget::AddChild(UWidget* Content)
 	return Slot;
 }
 
-void UPanelWidget::ReplaceChildAt(int32 Index, UWidget* Content)
+bool UPanelWidget::ReplaceChildAt(int32 Index, UWidget* Content)
 {
+	if ( Index < 0 || Index >= Slots.Num() )
+	{
+		return false;
+	}
+
 	UPanelSlot* Slot = Slots[Index];
 	Slot->Content = Content;
 
@@ -116,6 +131,8 @@ void UPanelWidget::ReplaceChildAt(int32 Index, UWidget* Content)
 	}
 
 	Slot->SynchronizeProperties();
+
+	return true;
 }
 
 void UPanelWidget::InsertChildAt(int32 Index, UWidget* Content)
diff --git a/Engine/Source/Runtime/UMG/Public/Components/PanelWidget.h b/Engine/Source/Runtime/UMG/Public/Components/PanelWidget.h
--- a/Engine/Source/Runtime/UMG/Public/Components/PanelWidget.h
+++ b/Engine/Source/Runtime/UMG/Public/Components/PanelWidget.h
@@ -46,8 +46,14 @@ public:
 	UFUNCTION(BlueprintCallable, Category="Widget|Panel")
 	UPanelSlot* AddChild(UWidget* Content);
 
-	/** Swaps the widget out of the slot at the given index, replacing it with a different widget. */
-	void ReplaceChildAt(int32 Index, UWidget* Content);
+	/**
+	 * Swaps the widget out of the slot at the given index, replacing it with a different widget.
+	 *
+	 * @param Index the index at which to replace the child content with this new Content.
+	 * @param Content The content to replace the child with.
+	 * @return true if the index existed and the child could be replaced.
+	 */
+	bool ReplaceChildAt(int32 Index, UWidget* Content);
 
 	/** Inserts a widget at a specific index */
 	void InsertChildAt(int32 Index, UWidget* Content);
