UMG - Little old code cleanup in UPanelWidget.

diff --git a/Engine/Source/Runtime/UMG/Private/Components/PanelWidget.cpp b/Engine/Source/Runtime/UMG/Private/Components/PanelWidget.cpp
--- a/Engine/Source/Runtime/UMG/Private/Components/PanelWidget.cpp
+++ b/Engine/Source/Runtime/UMG/Private/Components/PanelWidget.cpp
@@ -26,7 +26,7 @@ void UPanelWidget::ReleaseSlateResources(bool bReleaseChildren)
 	{
 		for ( int32 SlotIndex = 0; SlotIndex < Slots.Num(); SlotIndex++ )
 		{
-			if ( Slots[SlotIndex]->Content != NULL )
+			if ( Slots[SlotIndex]->Content != nullptr )
 			{
 				Slots[SlotIndex]->ReleaseSlateResources(bReleaseChildren);
 			}
@@ -73,10 +73,10 @@ bool UPanelWidget::RemoveChildAt(int32 Index)
 	UPanelSlot* Slot = Slots[Index];
 	if ( Slot->Content )
 	{
-		Slot->Content->Slot = NULL;
+		Slot->Content->Slot = nullptr;
 	}
 
-	Slot->Parent = NULL;
+	Slot->Parent = nullptr;
 	Slots.RemoveAt(Index);
 
 	OnSlotRemoved(Slot);
@@ -88,12 +88,12 @@ UPanelSlot* UPanelWidget::AddChild(UWidget* Content)
 {
 	if ( Content == nullptr )
 	{
-		return NULL;
+		return nullptr;
 	}
 
 	if ( !bCanHaveMultipleChildren && GetChildrenCount() > 0 )
 	{
-		return NULL;
+		return nullptr;
 	}
 
 	Content->RemoveFromParent();
@@ -188,7 +188,7 @@ void UPanelWidget::SetIsDesignTime(bool bInDesignTime)
 	int32 Children = GetChildrenCount();
 	for ( int32 SlotIndex = 0; SlotIndex < Children; SlotIndex++ )
 	{
-		if ( Slots[SlotIndex]->Content != NULL )
+		if ( Slots[SlotIndex]->Content != nullptr )
 		{
 			Slots[SlotIndex]->Content->SetIsDesignTime(bInDesignTime);
 		}
@@ -202,7 +202,7 @@ void UPanelWidget::PostLoad()
 	for ( int32 SlotIndex = 0; SlotIndex < Slots.Num(); SlotIndex++ )
 	{
 		// Remove any slots where their content is null, we don't support content-less slots.
-		if ( Slots[SlotIndex]->Content == NULL )
+		if ( Slots[SlotIndex]->Content == nullptr )
 		{
 			Slots.RemoveAt(SlotIndex);
 			SlotIndex--;
