UMG - Named slots now have better support for deletion, selection.  Renamed a few methods in UUserWidget to better capture what they do.  Several performance improvements, we now have a predicate based way on WidgetTree to perform an operation on every widget in the tree, useful for all the cases when you need to do a lot of searching in the tree.

diff --git a/Engine/Source/Editor/UMGEditor/Private/Designer/SDesignerView.cpp b/Engine/Source/Editor/UMGEditor/Private/Designer/SDesignerView.cpp
--- a/Engine/Source/Editor/UMGEditor/Private/Designer/SDesignerView.cpp
+++ b/Engine/Source/Editor/UMGEditor/Private/Designer/SDesignerView.cpp
@@ -52,7 +52,8 @@ public:
 
 	TMap<FName, FString> ExportedSlotProperties;
 
-	FWidgetReference Widget;
+	UWidget* Template;
+	UWidget* Preview;
 
 	bool bStayingInParent;
 	FWidgetReference ParentWidget;
@@ -69,13 +70,17 @@ TSharedRef<FSelectedWidgetDragDropOp> FSelectedWidgetDragDropOp::New(TSharedPtr<
 	}
 
 	TSharedRef<FSelectedWidgetDragDropOp> Operation = MakeShareable(new FSelectedWidgetDragDropOp());
-	Operation->Widget = InWidget;
 	Operation->bStayingInParent = bStayInParent;
 	Operation->ParentWidget = Editor->GetReferenceFromTemplate(InWidget.GetTemplate()->GetParent());
 	Operation->DefaultHoverText = FText::FromString( InWidget.GetTemplate()->GetLabel() );
 	Operation->CurrentHoverText = FText::FromString( InWidget.GetTemplate()->GetLabel() );
 	Operation->Construct();
 
+	// Cache the preview and template, it's not safe to query the preview/template while dragging the widget as it no longer
+	// exists in the tree.
+	Operation->Preview = InWidget.GetPreview();
+	Operation->Template = InWidget.GetTemplate();
+
 	FWidgetBlueprintEditorUtils::ExportPropertiesToText(InWidget.GetTemplate()->Slot, Operation->ExportedSlotProperties);
 
 	return Operation;
@@ -1228,7 +1233,7 @@ void SDesignerView::Tick(const FGeometry& AllottedGeometry, const double InCurre
 		if ( SelectedWidget.IsValid() )
 		{
 			// Set the selected widget so that we can draw the highlight
-			SelectedSlateWidget = PreviewWidget->GetWidgetFromName(SelectedWidget.GetTemplate()->GetFName());
+			SelectedSlateWidget = PreviewWidget->GetSlateWidgetFromName(SelectedWidget.GetTemplate()->GetFName());
 		}
 		else
 		{
@@ -1237,7 +1242,7 @@ void SDesignerView::Tick(const FGeometry& AllottedGeometry, const double InCurre
 
 		if ( HoveredWidget.IsValid() )
 		{
-			HoveredSlateWidget = PreviewWidget->GetWidgetFromName(HoveredWidget.GetTemplate()->GetFName());
+			HoveredSlateWidget = PreviewWidget->GetSlateWidgetFromName(HoveredWidget.GetTemplate()->GetFName());
 		}
 		else
 		{
@@ -1341,7 +1346,7 @@ UWidget* SDesignerView::ProcessDropAndAddWidget(const FGeometry& MyGeometry, con
 	TSharedPtr<FSelectedWidgetDragDropOp> SelectedDragDropOp = DragDropEvent.GetOperationAs<FSelectedWidgetDragDropOp>();
 	if ( SelectedDragDropOp.IsValid() )
 	{
-		DropPreviewWidget = SelectedDragDropOp->Widget.GetPreview();
+		DropPreviewWidget = SelectedDragDropOp->Preview;
 	}
 
 	UWidgetBlueprint* BP = GetBlueprint();
@@ -1517,10 +1522,13 @@ UWidget* SDesignerView::ProcessDropAndAddWidget(const FGeometry& MyGeometry, con
 				BP->WidgetTree->Modify();
 			}
 
-			UWidget* Widget = bIsPreview ? SelectedDragDropOp->Widget.GetPreview() : SelectedDragDropOp->Widget.GetTemplate();
-			check(Widget);
+			UWidget* Widget = bIsPreview ? SelectedDragDropOp->Preview : SelectedDragDropOp->Template;
+			if ( !Widget )
+			{
+				Widget = bIsPreview ? SelectedDragDropOp->Preview : SelectedDragDropOp->Template;
+			}
 
-			if ( Widget )
+			if ( ensure(Widget) )
 			{
 				if ( Widget->GetParent() )
 				{
diff --git a/Engine/Source/Editor/UMGEditor/Private/WidgetBlueprint.cpp b/Engine/Source/Editor/UMGEditor/Private/WidgetBlueprint.cpp
--- a/Engine/Source/Editor/UMGEditor/Private/WidgetBlueprint.cpp
+++ b/Engine/Source/Editor/UMGEditor/Private/WidgetBlueprint.cpp
@@ -113,12 +113,9 @@ void UWidgetBlueprint::PostLoad()
 {
 	Super::PostLoad();
 
-	TArray<UWidget*> Widgets;
-	WidgetTree->GetAllWidgets(Widgets);
-	for ( UWidget* Widget : Widgets )
-	{
+	WidgetTree->ForEachWidget([&] (UWidget* Widget) {
 		Widget->ConnectEditorData();
-	}
+	});
 
 	if( GetLinkerUE4Version() < VER_UE4_FIXUP_WIDGET_ANIMATION_CLASS )
 	{
diff --git a/Engine/Source/Editor/UMGEditor/Private/WidgetBlueprintCompiler.cpp b/Engine/Source/Editor/UMGEditor/Private/WidgetBlueprintCompiler.cpp
--- a/Engine/Source/Editor/UMGEditor/Private/WidgetBlueprintCompiler.cpp
+++ b/Engine/Source/Editor/UMGEditor/Private/WidgetBlueprintCompiler.cpp
@@ -107,10 +107,7 @@ void FWidgetBlueprintCompiler::ValidateWidgetNames()
 		}
 	}
 
-	TArray<UWidget*> Widgets;
-	Blueprint->WidgetTree->GetAllWidgets(Widgets);
-	for ( UWidget* Widget : Widgets )
-	{
+	Blueprint->WidgetTree->ForEachWidget([&] (UWidget* Widget) {
 		if ( ParentBPNameValidator.IsValid() && ParentBPNameValidator->IsValid(Widget->GetName()) != EValidatorResult::Ok )
 		{
 			// TODO Support renaming items, similar to timelines.
@@ -122,7 +119,7 @@ void FWidgetBlueprintCompiler::ValidateWidgetNames()
 			//MessageLog.Warning(*FString::Printf(*LOCTEXT("TimelineConflictWarning", "Found a timeline with a conflicting name (%s) - changed to %s.").ToString(), *TimelineTemplate->GetName(), *NewName.ToString()));
 			//FBlueprintEditorUtils::RenameTimeline(Blueprint, FName(*TimelineName), NewName);
 		}
-	}
+	});
 }
 
 template<typename TOBJ>
@@ -269,6 +266,15 @@ void FWidgetBlueprintCompiler::FinishCompilingClass(UClass* Class)
 		}
 	}
 
+	// Add all the names of the named slot widgets to the slot names structure.
+	BPGClass->NamedSlots.Reset();
+	BPGClass->WidgetTree->ForEachWidget([&] (UWidget* Widget) {
+		if ( Widget && Widget->IsA<UNamedSlot>() )
+		{
+			BPGClass->NamedSlots.Add(Widget->GetFName());
+		}
+	});
+
 	Super::FinishCompilingClass(Class);
 }
 
diff --git a/Engine/Source/Editor/UMGEditor/Private/WidgetBlueprintEditorUtils.cpp b/Engine/Source/Editor/UMGEditor/Private/WidgetBlueprintEditorUtils.cpp
--- a/Engine/Source/Editor/UMGEditor/Private/WidgetBlueprintEditorUtils.cpp
+++ b/Engine/Source/Editor/UMGEditor/Private/WidgetBlueprintEditorUtils.cpp
@@ -235,11 +235,11 @@ void FWidgetBlueprintEditorUtils::DeleteWidgets(UWidgetBlueprint* BP, TSet<FWidg
 
 			bRemoved = BP->WidgetTree->RemoveWidget(WidgetTemplate);
 
-			// If we failed to remove the widget from the tree, it may be a named slot widget that
-			// we need to search for in the named slot bindings.
-			if ( !bRemoved && WidgetTemplate->GetParent() == nullptr )
+			// If the widget we're removing doesn't have a parent it may be rooted in a named slot,
+			// so check there as well.
+			if ( WidgetTemplate->GetParent() == nullptr )
 			{
-				bRemoved = FindAndRemoveNamedSlotContent(WidgetTemplate, BP->WidgetTree);
+				bRemoved |= FindAndRemoveNamedSlotContent(WidgetTemplate, BP->WidgetTree);
 			}
 
 			// Rename the removed widget to the transient package so that it doesn't conflict with future widgets sharing the same name.
@@ -265,11 +265,15 @@ void FWidgetBlueprintEditorUtils::DeleteWidgets(UWidgetBlueprint* BP, TSet<FWidg
 
 bool FWidgetBlueprintEditorUtils::FindAndRemoveNamedSlotContent(UWidget* WidgetTemplate, UWidgetTree* WidgetTree)
 {
-	TArray<UWidget*> AllWidgets;
-	WidgetTree->GetAllWidgets(AllWidgets);
+	bool bSuccess = false;
+
+	WidgetTree->ForEachWidget([&] (UWidget* Widget) {
+
+		if ( bSuccess )
+		{
+			return;
+		}
 
-	for ( UWidget* Widget : AllWidgets )
-	{
 		if ( INamedSlotInterface* NamedSlotHost = Cast<INamedSlotInterface>(Widget) )
 		{
 			TArray<FName> SlotNames;
@@ -282,14 +286,14 @@ bool FWidgetBlueprintEditorUtils::FindAndRemoveNamedSlotContent(UWidget* WidgetT
 					if ( SlotContent == WidgetTemplate )
 					{
 						NamedSlotHost->SetContentForSlot(SlotName, nullptr);
-						return true;
+						bSuccess = true;
 					}
 				}
 			}
 		}
-	}
+	});
 
-	return false;
+	return bSuccess;
 }
 
 void FWidgetBlueprintEditorUtils::BuildWrapWithMenu(FMenuBuilder& Menu, UWidgetBlueprint* BP, TSet<FWidgetReference> Widgets)
diff --git a/Engine/Source/Editor/UMGEditor/Private/WidgetReference.cpp b/Engine/Source/Editor/UMGEditor/Private/WidgetReference.cpp
--- a/Engine/Source/Editor/UMGEditor/Private/WidgetReference.cpp
+++ b/Engine/Source/Editor/UMGEditor/Private/WidgetReference.cpp
@@ -33,7 +33,7 @@ bool FWidgetReference::IsValid() const
 {
 	if ( TemplateHandle.IsValid() )
 	{
-		return TemplateHandle->Widget.Get() != NULL && GetPreview();
+		return TemplateHandle->Widget.Get() != nullptr && GetPreview();
 	}
 
 	return false;
@@ -46,7 +46,7 @@ UWidget* FWidgetReference::GetTemplate() const
 		return TemplateHandle->Widget.Get();
 	}
 
-	return NULL;
+	return nullptr;
 }
 
 UWidget* FWidgetReference::GetPreview() const
@@ -55,14 +55,17 @@ UWidget* FWidgetReference::GetPreview() const
 	{
 		UUserWidget* PreviewRoot = WidgetEditor.Pin()->GetPreview();
 
-		if ( PreviewRoot && TemplateHandle->Widget.Get() )
+		if ( PreviewRoot )
 		{
-			UWidget* PreviewWidget = PreviewRoot->GetHandleFromName(TemplateHandle->Widget.Get()->GetName());
-			return PreviewWidget;
+			if ( UWidget* TemplateWidget = TemplateHandle->Widget.Get() )
+			{
+				UWidget* PreviewWidget = PreviewRoot->GetWidgetFromName(TemplateWidget->GetFName());
+				return PreviewWidget;
+			}
 		}
 	}
 
-	return NULL;
+	return nullptr;
 }
 
 #undef LOCTEXT_NAMESPACE
diff --git a/Engine/Source/Runtime/UMG/Private/UserWidget.cpp b/Engine/Source/Runtime/UMG/Private/UserWidget.cpp
--- a/Engine/Source/Runtime/UMG/Private/UserWidget.cpp
+++ b/Engine/Source/Runtime/UMG/Private/UserWidget.cpp
@@ -41,13 +41,12 @@ void UUserWidget::Initialize()
 
 		// Only do this if this widget is of a blueprint class
 		UWidgetBlueprintGeneratedClass* BGClass = Cast<UWidgetBlueprintGeneratedClass>(GetClass());
-		if ( BGClass != NULL )
+		if ( BGClass != nullptr )
 		{
 			BGClass->InitializeWidget(this);
 		}
 
-		for ( UWidget* Widget : Components )
-		{
+		WidgetTree->ForEachWidget([&] (UWidget* Widget) {
 			if ( UNamedSlot* NamedWidet = Cast<UNamedSlot>(Widget) )
 			{
 				for ( FNamedSlotBinding& Binding : NamedSlotBindings )
@@ -60,7 +59,7 @@ void UUserWidget::Initialize()
 					}
 				}
 			}
-		}
+		});
 	}
 }
 
@@ -78,7 +77,7 @@ void UUserWidget::ReleaseSlateResources(bool bReleaseChildren)
 {
 	Super::ReleaseSlateResources(bReleaseChildren);
 
-	UWidget* RootWidget = GetRootWidgetComponent();
+	UWidget* RootWidget = GetRootWidget();
 	if ( RootWidget )
 	{
 		RootWidget->ReleaseSlateResources(bReleaseChildren);
@@ -125,9 +124,6 @@ void UUserWidget::SetForegroundColor(FSlateColor InForegroundColor)
 void UUserWidget::PostInitProperties()
 {
 	Super::PostInitProperties();
-
-	Components.Reset();
-	//TODO UMG For non-BP versions how do we generate the Components list?
 }
 
 UWorld* UUserWidget::GetWorld() const
@@ -157,10 +153,9 @@ void UUserWidget::SetIsDesignTime(bool bInDesignTime)
 {
 	Super::SetIsDesignTime(bInDesignTime);
 
-	for ( UWidget* Widget : Components )
-	{
+	WidgetTree->ForEachWidget([&] (UWidget* Widget) {
 		Widget->SetIsDesignTime(bInDesignTime);
-	}
+	});
 }
 
 void UUserWidget::Construct_Implementation()
@@ -389,28 +384,19 @@ TSharedRef<SWidget> UUserWidget::RebuildWidget()
 
 	TSharedPtr<SWidget> UserRootWidget;
 
-	//setup the player context on sub user widgets, if we have a valid context
+	// Setup the player context on sub user widgets, if we have a valid context
 	if (PlayerContext.IsValid())
 	{
-		for (UWidget* Widget : Components)
-		{
-			UUserWidget* UserWidget = Cast<UUserWidget>(Widget);
-			if (UserWidget)
+		WidgetTree->ForEachWidget([&] (UWidget* Widget) {
+			if ( UUserWidget* UserWidget = Cast<UUserWidget>(Widget) )
 			{
 				UserWidget->SetPlayerContext(PlayerContext);
 			}
-		}
+		});
 	}
 
 	// Add the first component to the root of the widget surface.
-	if ( Components.Num() > 0 && Components[0] != NULL )
-	{
-		UserRootWidget = Components[0]->TakeWidget();
-	}
-	else
-	{
-		UserRootWidget = SNew(SSpacer);
-	}
+	UserRootWidget = WidgetTree->RootWidget ? WidgetTree->RootWidget->TakeWidget() : SNew(SSpacer);
 
 	if ( !IsDesignTime() )
 	{
@@ -421,7 +407,7 @@ TSharedRef<SWidget> UUserWidget::RebuildWidget()
 	return UserRootWidget.ToSharedRef();
 }
 
-TSharedPtr<SWidget> UUserWidget::GetWidgetFromName(const FName& Name) const
+TSharedPtr<SWidget> UUserWidget::GetSlateWidgetFromName(const FName& Name) const
 {
 	UWidget* WidgetObject = WidgetTree->FindWidget(Name);
 	if ( WidgetObject )
@@ -432,30 +418,18 @@ TSharedPtr<SWidget> UUserWidget::GetWidgetFromName(const FName& Name) const
 	return TSharedPtr<SWidget>();
 }
 
-UWidget* UUserWidget::GetHandleFromName(const FString& Name) const
+UWidget* UUserWidget::GetWidgetFromName(const FName& Name) const
 {
-	for ( UWidget* Widget : Components )
-	{
-		if ( Widget->GetName().Equals(Name, ESearchCase::IgnoreCase) )
-		{
-			return Widget;
-		}
-	}
-
-	return nullptr;
+	return WidgetTree->FindWidget(Name);
 }
 
 void UUserWidget::GetSlotNames(TArray<FName>& SlotNames) const
 {
-	//TODO UMG This should be static data computed and stored on the class at compile time.
-
-	// Add all the names of the named slot widgets to the slot names structure.
-	for ( UWidget* Widget : Components )
+	// Only do this if this widget is of a blueprint class
+	UWidgetBlueprintGeneratedClass* BGClass = Cast<UWidgetBlueprintGeneratedClass>(GetClass());
+	if ( BGClass != nullptr )
 	{
-		if ( Widget && Widget->IsA<UNamedSlot>() )
-		{
-			SlotNames.Add( Widget->GetFName() );
-		}
+		SlotNames.Append(BGClass->NamedSlots);
 	}
 
 	// Also add any existing bindings uniquely, templates don't have any components
@@ -565,11 +539,11 @@ TSharedRef<SWidget> UUserWidget::MakeViewportWidget(TSharedPtr<SWidget>& UserSla
 		];
 }
 
-UWidget* UUserWidget::GetRootWidgetComponent()
+UWidget* UUserWidget::GetRootWidget() const
 {
-	if ( Components.Num() > 0 )
+	if ( WidgetTree )
 	{
-		return Components[0];
+		return WidgetTree->RootWidget;
 	}
 
 	return nullptr;
diff --git a/Engine/Source/Runtime/UMG/Private/WidgetBlueprintGeneratedClass.cpp b/Engine/Source/Runtime/UMG/Private/WidgetBlueprintGeneratedClass.cpp
--- a/Engine/Source/Runtime/UMG/Private/WidgetBlueprintGeneratedClass.cpp
+++ b/Engine/Source/Runtime/UMG/Private/WidgetBlueprintGeneratedClass.cpp
@@ -62,9 +62,6 @@ void UWidgetBlueprintGeneratedClass::InitializeWidget(UUserWidget* UserWidget) c
 
 		UClass* WidgetBlueprintClass = UserWidget->GetClass();
 
-		TArray<UWidget*> ClonedWidgets;
-		ClonedTree->GetAllWidgets(ClonedWidgets);
-
 		for(UWidgetAnimation* Animation : Animations)
 		{
 			UWidgetAnimation* Anim = DuplicateObject<UWidgetAnimation>( Animation, UserWidget );
@@ -78,23 +75,18 @@ void UWidgetBlueprintGeneratedClass::InitializeWidget(UUserWidget* UserWidget) c
 					Prop->SetObjectPropertyValue_InContainer(UserWidget, Anim);
 				}
 			}
-
 		}
 
-		UserWidget->Components.Reset();
-
-		for ( UWidget* Widget : ClonedWidgets )
-		{
+		ClonedTree->ForEachWidget([&] (UWidget* Widget) {
 			// Not fatal if NULL, but shouldn't happen
 			if ( !ensure(Widget != NULL) )
 			{
-				continue;
+				return;
 			}
 
 			FString VariableName = Widget->GetName();
 
 			Widget->bCreatedByConstructionScript = true; // Indicate it comes from a blueprint so it gets cleared when we rerun construction scripts
-			UserWidget->Components.Add(Widget); // Add to array so it gets saved
 
 			// Find property with the same name as the template and assign the new widget to it.
 			UObjectPropertyBase* Prop = FindField<UObjectPropertyBase>(WidgetBlueprintClass, *VariableName);
@@ -155,7 +147,7 @@ void UWidgetBlueprintGeneratedClass::InitializeWidget(UUserWidget* UserWidget) c
 	#if WITH_EDITOR
 			Widget->ConnectEditorData();
 	#endif
-		}
+		});
 
 		// Bind any delegates on widgets
 		BindDynamicDelegates(UserWidget);
diff --git a/Engine/Source/Runtime/UMG/Private/WidgetTree.cpp b/Engine/Source/Runtime/UMG/Private/WidgetTree.cpp
--- a/Engine/Source/Runtime/UMG/Private/WidgetTree.cpp
+++ b/Engine/Source/Runtime/UMG/Private/WidgetTree.cpp
@@ -12,38 +12,30 @@ UWidgetTree::UWidgetTree(const FObjectInitializer& ObjectInitializer)
 
 UWidget* UWidgetTree::FindWidget(const FName& Name) const
 {
-	// TODO UMG Hacky, remove this find widget function, or make it faster.
-	TArray<UWidget*> Widgets;
-	GetAllWidgets(Widgets);
+	UWidget* FoundWidget = nullptr;
 
-	for ( UWidget* Widget : Widgets )
-	{
+	ForEachWidget([&] (UWidget* Widget) {
 		if ( Widget->GetFName() == Name )
 		{
-			return Widget;
+			FoundWidget = Widget;
 		}
-	}
+	});
 
-	return NULL;
+	return FoundWidget;
 }
 
 UWidget* UWidgetTree::FindWidget(TSharedRef<SWidget> InWidget) const
 {
-	FString ExistingName;
-
-	// TODO UMG Hacky, remove this find widget function, or make it faster.
-	TArray<UWidget*> Widgets;
-	GetAllWidgets(Widgets);
+	UWidget* FoundWidget = nullptr;
 
-	for ( UWidget* Widget : Widgets )
-	{
+	ForEachWidget([&] (UWidget* Widget) {
 		if ( Widget->GetCachedWidget() == InWidget )
 		{
-			return Widget;
+			FoundWidget = Widget;
 		}
-	}
+	});
 
-	return nullptr;
+	return FoundWidget;
 }
 
 UPanelWidget* UWidgetTree::FindWidgetParent(UWidget* Widget, int32& OutChildIndex)
@@ -85,39 +77,10 @@ bool UWidgetTree::RemoveWidget(UWidget* InRemovedWidget)
 
 void UWidgetTree::GetAllWidgets(TArray<UWidget*>& Widgets) const
 {
-	if ( RootWidget )
-	{
-		Widgets.Add(RootWidget);
-		GetChildWidgets(RootWidget, Widgets);
-	}
+	ForEachWidget([&] (UWidget* Widget) { Widgets.Add(Widget); });
 }
 
 void UWidgetTree::GetChildWidgets(UWidget* Parent, TArray<UWidget*>& Widgets) const
 {
-	if ( INamedSlotInterface* NamedSlotHost = Cast<INamedSlotInterface>(Parent) )
-	{
-		TArray<FName> SlotNames;
-		NamedSlotHost->GetSlotNames(SlotNames);
-
-		for ( FName SlotName : SlotNames )
-		{
-			if ( UWidget* SlotContent = NamedSlotHost->GetContentForSlot(SlotName) )
-			{
-				Widgets.Add(SlotContent);
-			}
-		}
-	}
-
-	if ( UPanelWidget* PanelParent = Cast<UPanelWidget>(Parent) )
-	{
-		for ( int32 ChildIndex = 0; ChildIndex < PanelParent->GetChildrenCount(); ChildIndex++ )
-		{
-			if ( UWidget* ChildWidget = PanelParent->GetChildAt(ChildIndex) )
-			{
-				Widgets.Add(ChildWidget);
-
-				GetChildWidgets(ChildWidget, Widgets);
-			}
-		}
-	}
+	ForWidgetAndChildren(Parent, [&] (UWidget* Widget) { Widgets.Add(Widget); });
 }
diff --git a/Engine/Source/Runtime/UMG/Public/Blueprint/UserWidget.h b/Engine/Source/Runtime/UMG/Public/Blueprint/UserWidget.h
--- a/Engine/Source/Runtime/UMG/Public/Blueprint/UserWidget.h
+++ b/Engine/Source/Runtime/UMG/Public/Blueprint/UserWidget.h
@@ -567,13 +567,13 @@ public:
 	TSharedRef<SWidget> MakeViewportWidget(TSharedPtr<SWidget>& UserSlateWidget);
 
 	/** @returns The root UObject widget wrapper */
-	UWidget* GetRootWidgetComponent();
+	UWidget* GetRootWidget() const;
 
 	/** @returns The slate widget corresponding to a given name */
-	TSharedPtr<SWidget> GetWidgetFromName(const FName& Name) const;
+	TSharedPtr<SWidget> GetSlateWidgetFromName(const FName& Name) const;
 
 	/** @returns The uobject widget corresponding to a given name */
-	UWidget* GetHandleFromName(const FString& Name) const;
+	UWidget* GetWidgetFromName(const FName& Name) const;
 
 	/** Ticks this widget and forwards to the underlying widget to tick */
 	void NativeTick(const FGeometry& MyGeometry, float InDeltaTime );
@@ -607,10 +607,6 @@ public:
 	UPROPERTY(EditAnywhere, BlueprintReadWrite, AdvancedDisplay, Category="Behavior")
 	bool bSupportsKeyboardFocus;
 
-	/** The components contained in this user widget. */
-	UPROPERTY(Transient)
-	TArray<UWidget*> Components;
-
 	/** The widget tree contained inside this user widget initialized by the blueprint */
 	UPROPERTY(Transient)
 	class UWidgetTree* WidgetTree;
diff --git a/Engine/Source/Runtime/UMG/Public/Blueprint/WidgetBlueprintGeneratedClass.h b/Engine/Source/Runtime/UMG/Public/Blueprint/WidgetBlueprintGeneratedClass.h
--- a/Engine/Source/Runtime/UMG/Public/Blueprint/WidgetBlueprintGeneratedClass.h
+++ b/Engine/Source/Runtime/UMG/Public/Blueprint/WidgetBlueprintGeneratedClass.h
@@ -65,6 +65,9 @@ class UMG_API UWidgetBlueprintGeneratedClass : public UBlueprintGeneratedClass
 	UPROPERTY()
 	TArray<UWidgetAnimation*> Animations;
 
+	UPROPERTY()
+	TArray< FName > NamedSlots;
+
 	/** This is transient data calculated at link time. */
 	TArray<UStructProperty*> WidgetNodeProperties;
 
diff --git a/Engine/Source/Runtime/UMG/Public/Blueprint/WidgetTree.h b/Engine/Source/Runtime/UMG/Public/Blueprint/WidgetTree.h
--- a/Engine/Source/Runtime/UMG/Public/Blueprint/WidgetTree.h
+++ b/Engine/Source/Runtime/UMG/Public/Blueprint/WidgetTree.h
@@ -32,9 +32,55 @@ public:
 	/** Gathers descendant child widgets of a parent widget. */
 	void GetChildWidgets(UWidget* Parent, TArray<UWidget*>& Widgets) const;
 
+	/**  */
+	template <typename Predicate>
+	FORCEINLINE void ForEachWidget(Predicate Pred) const
+	{
+		if ( RootWidget )
+		{
+			Pred(RootWidget);
+			
+			ForWidgetAndChildren(RootWidget, Pred);
+		}
+	}
+
+	/**  */
+	template <typename Predicate>
+	FORCEINLINE void ForWidgetAndChildren(UWidget* Widget, Predicate Pred) const
+	{
+		if ( INamedSlotInterface* NamedSlotHost = Cast<INamedSlotInterface>(Widget) )
+		{
+			TArray<FName> SlotNames;
+			NamedSlotHost->GetSlotNames(SlotNames);
+
+			for ( FName SlotName : SlotNames )
+			{
+				if ( UWidget* SlotContent = NamedSlotHost->GetContentForSlot(SlotName) )
+				{
+					Pred(SlotContent);
+
+					ForWidgetAndChildren(SlotContent, Pred);
+				}
+			}
+		}
+
+		if ( UPanelWidget* PanelParent = Cast<UPanelWidget>(Widget) )
+		{
+			for ( int32 ChildIndex = 0; ChildIndex < PanelParent->GetChildrenCount(); ChildIndex++ )
+			{
+				if ( UWidget* ChildWidget = PanelParent->GetChildAt(ChildIndex) )
+				{
+					Pred(ChildWidget);
+
+					ForWidgetAndChildren(ChildWidget, Pred);
+				}
+			}
+		}
+	}
+
 	/** Constructs the widget, and adds it to the tree. */
 	template< class T >
-	T* ConstructWidget(TSubclassOf<UWidget> WidgetType)
+	FORCEINLINE T* ConstructWidget(TSubclassOf<UWidget> WidgetType)
 	{
 		if ( WidgetType->IsChildOf(UUserWidget::StaticClass()) )
 		{
