Build Fix - Fixing ambiguity in UUserWidget compile.

diff --git a/Engine/Source/Runtime/UMG/Private/UserWidget.cpp b/Engine/Source/Runtime/UMG/Private/UserWidget.cpp
--- a/Engine/Source/Runtime/UMG/Private/UserWidget.cpp
+++ b/Engine/Source/Runtime/UMG/Private/UserWidget.cpp
@@ -382,8 +382,6 @@ TSharedRef<SWidget> UUserWidget::RebuildWidget()
 		Initialize();
 	}
 
-	TSharedPtr<SWidget> UserRootWidget;
-
 	// Setup the player context on sub user widgets, if we have a valid context
 	if (PlayerContext.IsValid())
 	{
@@ -396,7 +394,7 @@ TSharedRef<SWidget> UUserWidget::RebuildWidget()
 	}
 
 	// Add the first component to the root of the widget surface.
-	UserRootWidget = WidgetTree->RootWidget ? WidgetTree->RootWidget->TakeWidget() : SNew(SSpacer);
+	TSharedRef<SWidget> UserRootWidget = WidgetTree->RootWidget ? WidgetTree->RootWidget->TakeWidget() : SNew(SSpacer);
 
 	if ( !IsDesignTime() )
 	{
@@ -404,7 +402,7 @@ TSharedRef<SWidget> UUserWidget::RebuildWidget()
 		Construct();
 	}
 
-	return UserRootWidget.ToSharedRef();
+	return UserRootWidget;
 }
 
 TSharedPtr<SWidget> UUserWidget::GetSlateWidgetFromName(const FName& Name) const
