Build Fix - Fixing ambiguity on Mac builds.

diff --git a/Engine/Source/Runtime/UMG/Private/UserWidget.cpp b/Engine/Source/Runtime/UMG/Private/UserWidget.cpp
--- a/Engine/Source/Runtime/UMG/Private/UserWidget.cpp
+++ b/Engine/Source/Runtime/UMG/Private/UserWidget.cpp
@@ -394,7 +394,7 @@ TSharedRef<SWidget> UUserWidget::RebuildWidget()
 	}
 
 	// Add the first component to the root of the widget surface.
-	TSharedRef<SWidget> UserRootWidget = WidgetTree->RootWidget ? WidgetTree->RootWidget->TakeWidget() : SNew(SSpacer);
+	TSharedRef<SWidget> UserRootWidget = WidgetTree->RootWidget ? WidgetTree->RootWidget->TakeWidget() : TSharedRef<SWidget>(SNew(SSpacer));
 
 	if ( !IsDesignTime() )
 	{
