UMG - The drag drop operation is now blueprintable.  UWidget is now a BlueprintType.

diff --git a/Engine/Source/Runtime/UMG/Public/Blueprint/DragDropOperation.h b/Engine/Source/Runtime/UMG/Public/Blueprint/DragDropOperation.h
--- a/Engine/Source/Runtime/UMG/Public/Blueprint/DragDropOperation.h
+++ b/Engine/Source/Runtime/UMG/Public/Blueprint/DragDropOperation.h
@@ -10,7 +10,7 @@ DECLARE_DYNAMIC_MULTICAST_DELEGATE_OneParam(FOnDragDropMulticast, UDragDropOpera
 /**
  * This class is the base drag drop operation for UMG, extend it to add additional data and add new functionality.
  */
-UCLASS(BlueprintType)
+UCLASS(BlueprintType, Blueprintable)
 class UMG_API UDragDropOperation : public UObject
 {
 	GENERATED_UCLASS_BODY()
diff --git a/Engine/Source/Runtime/UMG/Public/Components/Widget.h b/Engine/Source/Runtime/UMG/Public/Components/Widget.h
--- a/Engine/Source/Runtime/UMG/Public/Components/Widget.h
+++ b/Engine/Source/Runtime/UMG/Public/Components/Widget.h
@@ -23,7 +23,7 @@ class UPanelSlot;
 /**
  * This is the base class for all wrapped Slate controls that are exposed to UMG.
  */
-UCLASS(Abstract)
+UCLASS(Abstract, BlueprintType)
 class UMG_API UWidget : public UVisual
 {
 	GENERATED_UCLASS_BODY()
