UMG 3D first pass

diff --git a/Engine/Config/BaseEngine.ini b/Engine/Config/BaseEngine.ini
--- a/Engine/Config/BaseEngine.ini
+++ b/Engine/Config/BaseEngine.ini
@@ -2209,6 +2209,7 @@ EnableNiagara=false
 +Profiles=(Name="Trigger",CollisionEnabled=QueryOnly,ObjectTypeName="WorldDynamic",CustomResponses=((Channel="WorldStatic",Response=ECR_Overlap),(Channel="Pawn",Response=ECR_Overlap),(Channel="Visibility",Response=ECR_Ignore),(Channel="WorldDynamic",Response=ECR_Overlap),(Channel="Camera",Response=ECR_Overlap),(Channel="PhysicsBody",Response=ECR_Overlap),(Channel="Vehicle",Response=ECR_Overlap),(Channel="Destructible",Response=ECR_Overlap)),HelpMessage="WorldDynamic object that is used for trigger. All other channels will be set to default.",bCanModify=False)
 +Profiles=(Name="Ragdoll",CollisionEnabled=QueryAndPhysics,ObjectTypeName="PhysicsBody",CustomResponses=((Channel="Pawn",Response=ECR_Ignore)),HelpMessage="Simulating Skeletal Mesh Component. All other channels will be set to default.",bCanModify=False)
 +Profiles=(Name="Vehicle",CollisionEnabled=QueryAndPhysics,ObjectTypeName="Vehicle",CustomResponses=,HelpMessage="Vehicle object that blocks Vehicle, WorldStatic, and WorldDynamic. All other channels will be set to default.",bCanModify=False)
++Profiles=(Name="UI",CollisionEnabled=QueryOnly,ObjectTypeName="WorldDynamic",CustomResponses=((Channel="WorldStatic",Response=ECR_Overlap),(Channel="Pawn",Response=ECR_Overlap),(Channel="Visibility",Response=ECR_Block),(Channel="WorldDynamic",Response=ECR_Overlap),(Channel="Camera",Response=ECR_Overlap),(Channel="Weapon",Response=ECR_Overlap),(Channel="PhysicsBody",Response=ECR_Overlap),(Channel="Vehicle",Response=ECR_Overlap),(Channel="Destructible",Response=ECR_Overlap),(Channel="Projectile",Response=ECR_Overlap),),HelpMessage="WorldStatic object that overlaps all actors by default. All new custom channels will use its own default response. ",bCanModify=False)
 
 +OldProfiles=(Name="BlockingVolume",CollisionEnabled=QueryAndPhysics,ObjectTypeName=WorldStatic,CustomResponses=((Channel=Visibility, Response=ECR_Ignore)))
 ; ECC_Mover is gone, but it doens't matter. Only thing that matters is the name
diff --git a/Engine/Source/Editor/UMGEditor/Private/UMGEditorModule.cpp b/Engine/Source/Editor/UMGEditor/Private/UMGEditorModule.cpp
--- a/Engine/Source/Editor/UMGEditor/Private/UMGEditorModule.cpp
+++ b/Engine/Source/Editor/UMGEditor/Private/UMGEditorModule.cpp
@@ -14,6 +14,8 @@
 #include "Animation/MarginTrackEditor.h"
 #include "Animation/Sequencer2DTransformTrackEditor.h"
 #include "IUMGModule.h"
+#include "ComponentReregisterContext.h"
+#include "UMGComponent.h"
 #include "DesignerCommands.h"
 #include "WidgetBlueprint.h"
 
@@ -75,6 +77,8 @@ public:
 	{
 		if ( UWidgetBlueprint* WidgetBlueprint = Cast<UWidgetBlueprint>(Blueprint) )
 		{
+			TComponentReregisterContext<UUMGComponent> ComponentReregisterContext;
+
 			FWidgetBlueprintCompiler Compiler(WidgetBlueprint, Results, CompileOptions, ObjLoaded);
 			Compiler.Compile();
 			check(Compiler.NewClass);
diff --git a/Engine/Source/Editor/UMGEditor/UMGEditor.Build.cs b/Engine/Source/Editor/UMGEditor/UMGEditor.Build.cs
--- a/Engine/Source/Editor/UMGEditor/UMGEditor.Build.cs
+++ b/Engine/Source/Editor/UMGEditor/UMGEditor.Build.cs
@@ -49,7 +49,8 @@ public class UMGEditor : ModuleRules
 				"MovieSceneCore",
 				"Sequencer",
 				"DetailCustomizations",
-                "Settings"
+                "Settings",
+				"RenderCore",
 			}
 			);
 
diff --git a/Engine/Source/Runtime/Slate/Private/Framework/Application/SlateApplication.cpp b/Engine/Source/Runtime/Slate/Private/Framework/Application/SlateApplication.cpp
--- a/Engine/Source/Runtime/Slate/Private/Framework/Application/SlateApplication.cpp
+++ b/Engine/Source/Runtime/Slate/Private/Framework/Application/SlateApplication.cpp
@@ -10,6 +10,20 @@
 #include "NotificationManager.h"
 
 
+// @todo slate : Move TranslatePointerEvent into FEventRouter
+template <typename PointerEventType>
+static PointerEventType TranslatePointerEvent( const FArrangedWidget& InWidget, const PointerEventType& InEvent )
+{
+	if ( !InWidget.VirtualCursorPosition.IsValid() )
+	{
+		return InEvent;
+	}
+	else
+	{
+		return FPointerEvent::MakeTranslatedEvent<PointerEventType>( InEvent, *InWidget.VirtualCursorPosition );
+	}
+}
+
 class FEventRouter
 {
 
@@ -182,9 +196,10 @@ public:
 		for ( ; !Reply.IsEventHandled() && RoutingPolicy.ShouldKeepGoing(); RoutingPolicy.Next() )
 		{
 			const FArrangedWidget& ArrangedWidget = RoutingPolicy.GetWidget();
-			Reply = Lambda( ArrangedWidget, PointerEventCopy ).SetHandler( ArrangedWidget.Widget );
+			const FPointerEvent TranslatedEvent = TranslatePointerEvent( ArrangedWidget, PointerEventCopy );
+			Reply = Lambda( ArrangedWidget, TranslatedEvent ).SetHandler( ArrangedWidget.Widget );
 			const FWidgetPath* WidgetsUnderCursorPtr = &RoutingPath;
-			ThisApplication->ProcessReply( RoutingPath, Reply, WidgetsUnderCursorPtr, &PointerEventCopy );
+			ThisApplication->ProcessReply( RoutingPath, Reply, WidgetsUnderCursorPtr, &TranslatedEvent );
 		}
 
 		LogEvent(ThisApplication, PointerEventCopy, Reply);
@@ -374,10 +389,21 @@ void FSlateApplication::MouseCaptorHelper::InvalidateCaptureForAllPointers()
 	}
 }
 
-void FSlateApplication::MouseCaptorHelper::InvalidateCaptureForPointer(uint32 PointerIndex)
+FWidgetPath FSlateApplication::MouseCaptorHelper::ToWidgetPath( FWeakWidgetPath::EInterruptedPathHandling::Type InterruptedPathHandling, const FPointerEvent* PointerEvent )
 {
-	InformCurrentCaptorOfCaptureLoss(PointerIndex);
-	PointerIndexToMouseCaptorWeakPathMap.Remove(PointerIndex);
+	FWidgetPath WidgetPath;
+	const FWeakWidgetPath* MouseCaptorWeakPath = PointerIndexToMouseCaptorWeakPathMap.Find(PointerEvent->GetPointerIndex());
+	if ( MouseCaptorWeakPath->IsValid() )
+	{
+		if ( MouseCaptorWeakPath->ToWidgetPath( WidgetPath, InterruptedPathHandling, PointerEvent ) == FWeakWidgetPath::EPathResolutionResult::Truncated )
+		{
+			// If the path was truncated then it means this widget is no longer part of the active set,
+			// so we make sure to invalidate its capture
+			InvalidateCaptureForPointer(PointerEvent->GetPointerIndex());
+		}
+	}
+
+	return WidgetPath;
 }
 
 FWidgetPath FSlateApplication::MouseCaptorHelper::ToWidgetPath(uint32 PointerIndex, FWeakWidgetPath::EInterruptedPathHandling::Type InterruptedPathHandling)
@@ -397,6 +423,12 @@ FWidgetPath FSlateApplication::MouseCaptorHelper::ToWidgetPath(uint32 PointerInd
 	return WidgetPath;
 }
 
+void FSlateApplication::MouseCaptorHelper::InvalidateCaptureForPointer(uint32 PointerIndex)
+{
+	InformCurrentCaptorOfCaptureLoss(PointerIndex);
+	PointerIndexToMouseCaptorWeakPathMap.Remove(PointerIndex);
+}
+
 TArray<FWidgetPath> FSlateApplication::MouseCaptorHelper::ToWidgetPaths()
 {
 	TArray<FWidgetPath> WidgetPaths;
@@ -509,7 +541,6 @@ void FPopupSupport::SendNotifications( const FWidgetPath& WidgetsUnderCursor )
 	}
 }
 
-
 void FSlateApplication::Create()
 {
 	Create(MakeShareable(FPlatformMisc::CreateApplication()));
@@ -2078,14 +2109,14 @@ void FSlateApplication::ProcessReply( const FWidgetPath& CurrentEventPath, const
 		for (int32 WidgetIndex=0; WidgetIndex < WidgetsUnderMouse->Widgets.Num(); ++WidgetIndex)
 		{
 			const FArrangedWidget& SomeWidget = WidgetsUnderMouse->Widgets[WidgetIndex];
-			SomeWidget.Widget->OnMouseLeave( *InMouseEvent );
+			SomeWidget.Widget->OnMouseLeave( TranslatePointerEvent( SomeWidget, *InMouseEvent ) );
 		}
 
 		FDragDropEvent DragDropEvent( *InMouseEvent, ReplyDragDropContent );
 		for (int32 WidgetIndex=0; WidgetIndex < WidgetsUnderMouse->Widgets.Num(); ++WidgetIndex)
 		{
 			const FArrangedWidget& SomeWidget = WidgetsUnderMouse->Widgets[WidgetIndex];
-			SomeWidget.Widget->OnDragEnter( SomeWidget.Geometry, DragDropEvent );
+			SomeWidget.Widget->OnDragEnter( SomeWidget.Geometry, TranslatePointerEvent( SomeWidget, DragDropEvent ) );
 		}
 	}
 	
@@ -2253,10 +2284,23 @@ void FSlateApplication::QueryCursor()
 			FWidgetPath WidgetsToQueryForCursor;
 			const TSharedPtr<SWindow> ActiveModalWindow = GetActiveModalWindow();
 
+			const FVector2D CurrentCursorPosition = GetCursorPos();
+			const FVector2D LastCursorPosition = GetLastCursorPos();
+			const FPointerEvent CursorEvent(
+				CursorPointerIndex,
+				CurrentCursorPosition,
+				LastCursorPosition,
+				CurrentCursorPosition - LastCursorPosition,
+				PressedMouseButtons,
+				PlatformApplication->GetModifierKeys()
+				);
+
 			// Query widgets with mouse capture for the cursor
 			if (MouseCaptor.HasCaptureForPointerIndex(CursorPointerIndex))
 			{
-				FWidgetPath MouseCaptorPath = MouseCaptor.ToWidgetPath(CursorPointerIndex);
+				//@todo dhertzka - See if old ToWidgetPath should be replaced entirely
+				//FWidgetPath MouseCaptorPath = MouseCaptor.ToWidgetPath(CursorPointerIndex);
+				FWidgetPath MouseCaptorPath = MouseCaptor.ToWidgetPath( FWeakWidgetPath::EInterruptedPathHandling::Truncate, &CursorEvent);
 				if ( MouseCaptorPath.IsValid() )
 				{
 					TSharedRef< SWindow > CaptureWindow = MouseCaptorPath.GetWindow();
@@ -2278,6 +2322,7 @@ void FSlateApplication::QueryCursor()
 				// Switch worlds for widgets in the current path
 				FScopedSwitchWorldHack SwitchWorld( WidgetsToQueryForCursor );
 
+				//@TODO dhertzka: This was deleted in UMG 3D 2 => investigate
 				const FVector2D CurrentCursorPosition = GetCursorPos();
 				const FVector2D LastCursorPosition = GetLastCursorPos();
 				const FPointerEvent CursorEvent(
@@ -2293,7 +2338,7 @@ void FSlateApplication::QueryCursor()
 				for( int32 WidgetIndex = WidgetsToQueryForCursor.Widgets.Num() - 1; !CursorResult.IsEventHandled() && WidgetIndex >= 0; --WidgetIndex )
 				{
 					const FArrangedWidget& WidgetToQuery = WidgetsToQueryForCursor.Widgets[ WidgetIndex ];
-					CursorResult = WidgetToQuery.Widget->OnCursorQuery( WidgetToQuery.Geometry, CursorEvent );
+					CursorResult = WidgetToQuery.Widget->OnCursorQuery( WidgetToQuery.Geometry, TranslatePointerEvent( WidgetToQuery, CursorEvent ) );
 				}
 
 				if (CursorResult.IsEventHandled())
@@ -3503,13 +3548,15 @@ bool FSlateApplication::ProcessMouseButtonDownEvent( const TSharedPtr< FGenericW
 		FReply Reply = FReply::Unhandled();
 		if (MouseCaptor.HasCaptureForPointerIndex(MouseEvent.GetPointerIndex()))
 		{
-			FWidgetPath MouseCaptorPath = MouseCaptor.ToWidgetPath(MouseEvent.GetPointerIndex());
+			//FWidgetPath MouseCaptorPath = MouseCaptor.ToWidgetPath(MouseEvent.GetPointerIndex());
+			FWidgetPath MouseCaptorPath = MouseCaptor.ToWidgetPath( FWeakWidgetPath::EInterruptedPathHandling::Truncate, &MouseEvent );
 			FArrangedWidget& MouseCaptorWidget = MouseCaptorPath.Widgets.Last();
 
 			// Switch worlds widgets in the current path
 			FScopedSwitchWorldHack SwitchWorld( MouseCaptorPath );
 			bInGame = FApp::IsGame();
 
+			MouseEvent.SetEventPath( MouseCaptorPath );
 			
 			Reply = FEventRouter::RouteAlongPointerPath( this, FEventRouter::FToLeafmostPolicy(MouseCaptorPath), MouseEvent, []( const FArrangedWidget& MouseCaptorWidget, const FPointerEvent& Event )
 			{
@@ -3520,19 +3567,18 @@ bool FSlateApplication::ProcessMouseButtonDownEvent( const TSharedPtr< FGenericW
 			{
 				Reply = FEventRouter::RouteAlongPointerPath( this, FEventRouter::FToLeafmostPolicy(MouseCaptorPath), MouseEvent,
 					[this]( const FArrangedWidget& MouseCaptorWidget, const FPointerEvent& Event )
-				{
+					{
 						FReply TempReply = FReply::Unhandled();
-						if (Event.IsTouchEvent())
+						if ( Event.IsTouchEvent() )
 						{
 							TempReply = MouseCaptorWidget.Widget->OnTouchStarted( MouseCaptorWidget.Geometry, Event );
-				}
-						if (!Event.IsTouchEvent() || (!TempReply.IsEventHandled() && this->bTouchFallbackToMouse))
-				{
+						}
+						if ( !Event.IsTouchEvent() || ( !TempReply.IsEventHandled() && this->bTouchFallbackToMouse ) )
+						{
 							TempReply = MouseCaptorWidget.Widget->OnMouseButtonDown( MouseCaptorWidget.Geometry, Event );
-				}
+						}
 						return TempReply;
 					} );
-
 			}
 			LOG_EVENT( EEventLog::MouseButtonDown, Reply );
 		}
@@ -3718,7 +3764,8 @@ bool FSlateApplication::ProcessMouseButtonUpEvent( FPointerEvent& MouseEvent )
 
 	if (MouseCaptor.HasCaptureForPointerIndex(MouseEvent.GetPointerIndex()))
 	{
-		FWidgetPath MouseCaptorPath = MouseCaptor.ToWidgetPath(MouseEvent.GetPointerIndex());
+		//FWidgetPath MouseCaptorPath = MouseCaptor.ToWidgetPath(MouseEvent.GetPointerIndex());
+		FWidgetPath MouseCaptorPath = MouseCaptor.ToWidgetPath( FWeakWidgetPath::EInterruptedPathHandling::Truncate, &MouseEvent );
 		if ( ensureMsg(MouseCaptorPath.Widgets.Num() > 0, TEXT("A window had a widget with mouse capture. That entire window has been dismissed before the mouse up could be processed.")) )
 		{
 #if PLATFORM_MAC
@@ -3749,6 +3796,7 @@ bool FSlateApplication::ProcessMouseButtonUpEvent( FPointerEvent& MouseEvent )
 				MouseCaptorPath.TopLevelWindow->BringToFront(true);
 			}
 #endif
+			//ProcessReply(MouseCaptorPath, Reply, &MouseCaptorPath, &TranslatedMouseEvent);
 			LOG_EVENT( EEventLog::MouseButtonUp, Reply );
 		}
 	}
@@ -3964,8 +4012,9 @@ bool FSlateApplication::ProcessMouseMoveEvent( FPointerEvent& MouseEvent, bool b
 				// Switch worlds widgets in the current path
 				FScopedSwitchWorldHack SwitchWorld( DragDetectPath );
 
-				FReply Reply = DetectDragForMe.Widget->OnDragDetected(DetectDragForMe.Geometry, MouseEvent).SetHandler(DetectDragForMe.Widget);
-				ProcessReply( DragDetectPath, Reply, &DragDetectPath, &MouseEvent );
+				FPointerEvent TranslatedMouseEvent = TranslatePointerEvent( DetectDragForMe, MouseEvent );
+				FReply Reply = DetectDragForMe.Widget->OnDragDetected(DetectDragForMe.Geometry, TranslatedMouseEvent ).SetHandler(DetectDragForMe.Widget);
+				ProcessReply( DragDetectPath, Reply, &DragDetectPath, &TranslatedMouseEvent );
 				LOG_EVENT( EEventLog::DragDetected, Reply );
 			}
 			else
@@ -4010,12 +4059,14 @@ bool FSlateApplication::ProcessMouseMoveEvent( FPointerEvent& MouseEvent, bool b
 				const TSharedPtr<SWidget>& SomeWidgetPreviouslyUnderCursor = LastWidgetsUnderCursor.Widgets[WidgetIndex].Pin();
 				if( SomeWidgetPreviouslyUnderCursor.IsValid() )
 				{
-					if ( !WidgetsUnderCursor.ContainsWidget(SomeWidgetPreviouslyUnderCursor.ToSharedRef()) )
+					//@todo dhertzka - return a TOptional<FArrangedWidget>
+					FArrangedWidget FoundWidget = WidgetsUnderCursor.FindArrangedWidget( SomeWidgetPreviouslyUnderCursor.ToSharedRef() );
+					if ( FoundWidget.Widget == SNullWidget::NullWidget )
 					{
 						// Widget is no longer under cursor, so send a MouseLeave.
 						if ( IsDragDropping() )
 						{
-							SomeWidgetPreviouslyUnderCursor->OnDragLeave( DragDropEvent );
+							SomeWidgetPreviouslyUnderCursor->OnDragLeave( TranslatePointerEvent( FoundWidget, DragDropEvent ) );
 							LOG_EVENT( EEventLog::DragLeave, SomeWidgetPreviouslyUnderCursor );
 
 							// Reset the cursor override
@@ -4023,7 +4074,7 @@ bool FSlateApplication::ProcessMouseMoveEvent( FPointerEvent& MouseEvent, bool b
 						}
 						else
 						{
-							SomeWidgetPreviouslyUnderCursor->OnMouseLeave( MouseEvent );
+							SomeWidgetPreviouslyUnderCursor->OnMouseLeave( TranslatePointerEvent( FoundWidget, MouseEvent ) );
 							LOG_EVENT( EEventLog::MouseLeave, SomeWidgetPreviouslyUnderCursor );
 						}			
 					}
@@ -4036,7 +4087,8 @@ bool FSlateApplication::ProcessMouseMoveEvent( FPointerEvent& MouseEvent, bool b
 	FWidgetPath MouseCaptorPath;
 	if (MouseCaptor.HasCaptureForPointerIndex(MouseEvent.GetPointerIndex()))
 	{
-		MouseCaptorPath = MouseCaptor.ToWidgetPath(MouseEvent.GetPointerIndex(), FWeakWidgetPath::EInterruptedPathHandling::ReturnInvalid);
+		//MouseCaptorPath = MouseCaptor.ToWidgetPath(MouseEvent.GetPointerIndex(), FWeakWidgetPath::EInterruptedPathHandling::ReturnInvalid);
+		MouseCaptorPath = MouseCaptor.ToWidgetPath(FWeakWidgetPath::EInterruptedPathHandling::ReturnInvalid, &MouseEvent );
 	}
 
 	if (MouseCaptorPath.IsValid())
@@ -4079,13 +4131,13 @@ bool FSlateApplication::ProcessMouseMoveEvent( FPointerEvent& MouseEvent, bool b
 				if ( IsDragDropping() )
 				{
 					// Doing a drag and drop; send a DragDropEvent
-					SomeWidgetUnderCursor.Widget->OnDragEnter( SomeWidgetUnderCursor.Geometry, DragDropEvent );
+					SomeWidgetUnderCursor.Widget->OnDragEnter( SomeWidgetUnderCursor.Geometry, TranslatePointerEvent( SomeWidgetUnderCursor, DragDropEvent ) );
 					LOG_EVENT( EEventLog::DragEnter, SomeWidgetUnderCursor.Widget );
 				}
 				else
 				{
 					// Not drag dropping; send regular mouse event
-					SomeWidgetUnderCursor.Widget->OnMouseEnter( SomeWidgetUnderCursor.Geometry, MouseEvent );
+					SomeWidgetUnderCursor.Widget->OnMouseEnter( SomeWidgetUnderCursor.Geometry, TranslatePointerEvent( SomeWidgetUnderCursor, MouseEvent ) );
 					LOG_EVENT( EEventLog::MouseEnter, SomeWidgetUnderCursor.Widget );
 				}
 			}
diff --git a/Engine/Source/Runtime/Slate/Private/Widgets/SViewport.cpp b/Engine/Source/Runtime/Slate/Private/Widgets/SViewport.cpp
--- a/Engine/Source/Runtime/Slate/Private/Widgets/SViewport.cpp
+++ b/Engine/Source/Runtime/Slate/Private/Widgets/SViewport.cpp
@@ -1,6 +1,7 @@
 // Copyright 1998-2014 Epic Games, Inc. All Rights Reserved.
 
 #include "SlatePrivatePCH.h"
+#include "HittestGrid.h"
 
 DECLARE_CYCLE_STAT( TEXT("OnPaint SViewport"), STAT_SlateOnPaint_SViewport, STATGROUP_Slate );
 
@@ -42,6 +43,9 @@ int32 SViewport::OnPaint( const FPaintArgs& Args, const FGeometry& AllottedGeome
 	bool bShowDisabledEffect = ShowDisabledEffect.Get();
 	ESlateDrawEffect::Type DrawEffects = bShowDisabledEffect && !bEnabled ? ESlateDrawEffect::DisabledEffect : ESlateDrawEffect::None;
 
+	int32 LastHitTestIndex = Args.GetLastHitTestIndex();
+
+
 	// Viewport texture alpha channels are often in an indeterminate state, even after the resolve,
 	// so we'll tell the shader to not use the alpha channel when blending
 	if( bIgnoreTextureAlpha )
@@ -57,6 +61,7 @@ int32 SViewport::OnPaint( const FPaintArgs& Args, const FGeometry& AllottedGeome
 		ViewportInterfacePin->OnDrawViewport( AllottedGeometry, MyClippingRect, OutDrawElements, LayerId, InWidgetStyle, bParentEnabled );
 	}
 
+
 	// Only draw a quad if not rendering directly to the backbuffer
 	if( !ShouldRenderDirectly() )
 	{
@@ -105,6 +110,13 @@ int32 SViewport::OnPaint( const FPaintArgs& Args, const FGeometry& AllottedGeome
 		);
 	}
 
+
+	if(CustomHitTestPath.IsValid())
+	{
+		Args.InsertCustomHitTestPath(CustomHitTestPath.ToSharedRef(), LastHitTestIndex);
+	}
+
+
 	return Layer;
 }
 
@@ -216,6 +228,23 @@ void SViewport::OnKeyboardFocusLost( const FKeyboardFocusEvent& InKeyboardFocusE
 	}
 }
 
+void SViewport::SetContent( TSharedPtr<SWidget> InContent )
+{
+	ChildSlot
+	[
+		InContent.IsValid() ? InContent.ToSharedRef() : (TSharedRef<SWidget>)SNullWidget::NullWidget
+	];
+}
+
+void SViewport::SetCustomHitTestPath( TSharedPtr<ICustomHitTestPath> InCustomHitTestPath )
+{
+	CustomHitTestPath = InCustomHitTestPath;
+}
+
+TSharedPtr<ICustomHitTestPath> SViewport::GetCustomHitTestPath()
+{
+	return CustomHitTestPath;
+}
 
 void SViewport::OnWindowClosed( const TSharedRef<SWindow>& WindowBeingClosed )
 {
@@ -281,3 +310,22 @@ void SViewport::OnFinishedPointerInput()
 		PinnedInterface->OnFinishedPointerInput();
 	}
 }
+
+void SViewport::OnArrangeChildren( const FGeometry& AllottedGeometry, FArrangedChildren& ArrangedChildren ) const
+{
+	SCompoundWidget::OnArrangeChildren(AllottedGeometry, ArrangedChildren);
+	if( ArrangedChildren.Allows3DWidgets() && CustomHitTestPath.IsValid() )
+	{
+		CustomHitTestPath->ArrangeChildren( ArrangedChildren );
+	}
+}
+
+TSharedPtr<FVirtualCursorPosition> SViewport::TranslateMouseCoordinateFor3DChild(const TSharedRef<SWidget>& ChildWidget, const FGeometry& MyGeometry, const FVector2D& ScreenSpaceMouseCoordinate, const FVector2D& LastScreenSpaceMouseCoordinate) const
+{
+	if( CustomHitTestPath.IsValid() )
+	{
+		return CustomHitTestPath->TranslateMouseCoordinateFor3DChild( ChildWidget, MyGeometry, ScreenSpaceMouseCoordinate, LastScreenSpaceMouseCoordinate );
+	}
+
+	return nullptr;
+}
\ No newline at end of file
diff --git a/Engine/Source/Runtime/Slate/Public/Framework/Application/SlateApplication.h b/Engine/Source/Runtime/Slate/Public/Framework/Application/SlateApplication.h
--- a/Engine/Source/Runtime/Slate/Public/Framework/Application/SlateApplication.h
+++ b/Engine/Source/Runtime/Slate/Public/Framework/Application/SlateApplication.h
@@ -1182,6 +1182,8 @@ private:
 		 */
 		FWidgetPath ToWidgetPath(uint32 PointerIndex, FWeakWidgetPath::EInterruptedPathHandling::Type InterruptedPathHandling = FWeakWidgetPath::EInterruptedPathHandling::Truncate );
 
+		FWidgetPath ToWidgetPath( FWeakWidgetPath::EInterruptedPathHandling::Type InterruptedPathHandling = FWeakWidgetPath::EInterruptedPathHandling::Truncate, const FPointerEvent* PointerEvent = nullptr );
+
 		/**
 		 * Retrieves an array of the resolved widget paths for any active captures.
 		 */
@@ -1189,6 +1191,7 @@ private:
 
 		/** Retrieves the weak path for a current mouse captor with a specific pointer index */
 		FWeakWidgetPath ToWeakPath(uint32 PointerIndex) const;
+		
 
 		/* Walks the weak path and retrieves the widget that is set as the current mouse captor with a specific pointer index */
 		TSharedPtr< SWidget > ToSharedWidget(uint32 PointerIndex) const;
diff --git a/Engine/Source/Runtime/Slate/Public/Widgets/SViewport.h b/Engine/Source/Runtime/Slate/Public/Widgets/SViewport.h
--- a/Engine/Source/Runtime/Slate/Public/Widgets/SViewport.h
+++ b/Engine/Source/Runtime/Slate/Public/Widgets/SViewport.h
@@ -90,13 +90,11 @@ public:
 	 *
 	 * @param InContent	The new content (can be null)
 	 */
-	void SetContent( TSharedPtr<SWidget> InContent )
-	{
-		ChildSlot
-		[
-			InContent.IsValid() ? InContent.ToSharedRef() : (TSharedRef<SWidget>)SNullWidget::NullWidget
-		];
-	}
+	void SetContent( TSharedPtr<SWidget> InContent );
+
+	void SetCustomHitTestPath( TSharedPtr<ICustomHitTestPath> CustomHitTestPath );
+
+	TSharedPtr<ICustomHitTestPath> GetCustomHitTestPath();
 
 	const TSharedPtr<SWidget> GetContent() const { return ChildSlot.GetWidget(); }
 
@@ -156,6 +154,8 @@ public:
 	virtual FReply OnControllerAnalogValueChanged( const FGeometry& MyGeometry, const FControllerEvent& ControllerEvent ) override;
 	virtual FReply OnMotionDetected( const FGeometry& MyGeometry, const FMotionEvent& MotionEvent ) override;
 	virtual void OnFinishedPointerInput() override;
+	virtual void OnArrangeChildren( const FGeometry& AllottedGeometry, FArrangedChildren& ArrangedChildren ) const override;
+	virtual TSharedPtr<struct FVirtualCursorPosition> TranslateMouseCoordinateFor3DChild(const TSharedRef<SWidget>& ChildWidget, const FGeometry& MyGeometry, const FVector2D& ScreenSpaceMouseCoordinate, const FVector2D& LastScreenSpaceMouseCoordinate) const;
 
 private:
 
@@ -177,6 +177,14 @@ private:
 	/** Size of the viewport. */
 	TAttribute<FVector2D> ViewportSize;
 
+	/** 
+	 * Widget to transfer keyboard focus to when this window becomes active, if any.
+	 * This is used to restore focus to a widget after a popup has been dismissed.
+	 */
+	TWeakPtr< SWidget > WidgetToFocusOnActivate;
+
+	TSharedPtr<ICustomHitTestPath> CustomHitTestPath;
+
 	/** Whether or not this viewport renders directly to the window back-buffer. */
 	bool bRenderDirectlyToWindow;
 
@@ -188,10 +196,4 @@ private:
 
 	/** Whether or not to allow texture alpha to be used in blending calculations. */
 	bool bIgnoreTextureAlpha;
-
-	/**
-	 * Widget to transfer keyboard focus to when this window becomes active, if any.
-	 * This is used to restore focus to a widget after a pop-up has been dismissed.
-	 */
-	TWeakPtr<SWidget> WidgetToFocusOnActivate;
 };
diff --git a/Engine/Source/Runtime/SlateCore/Private/Input/HittestGrid.cpp b/Engine/Source/Runtime/SlateCore/Private/Input/HittestGrid.cpp
--- a/Engine/Source/Runtime/SlateCore/Private/Input/HittestGrid.cpp
+++ b/Engine/Source/Runtime/SlateCore/Private/Input/HittestGrid.cpp
@@ -24,6 +24,8 @@ struct FHittestGrid::FCachedWidget
 	}
 
 	TWeakPtr<SWidget> WidgetPtr;
+	/** Allow widgets that implement this interface to insert widgets into the bubble path */
+	TWeakPtr<ICustomHitTestPath> CustomPath;
 	FGeometry CachedGeometry;
 	// @todo umg : ideally this clipping rect is optional and we only have them on a small number of widgets.
 	FSlateRect ClippingRect;
@@ -79,10 +81,24 @@ TArray<FArrangedWidget> FHittestGrid::GetBubblePath( FVector2D DesktopSpaceCoord
 				HitWidgetIndex = IndexesInCell[i];
 			}
 		}
+		
+		TArray<FArrangedWidget> BubblePath =[&]()
+		{
+			if( HitWidgetIndex != INDEX_NONE )
+			{
+				const FCachedWidget& PhysicallyHitWidget = (*WidgetsCachedThisFrame)[HitWidgetIndex];
+				if(PhysicallyHitWidget.CustomPath.IsValid())
+				{
+					return PhysicallyHitWidget.CustomPath.Pin()->GetBubblePath(PhysicallyHitWidget.CachedGeometry, DesktopSpaceCoordinate, bIgnoreEnabledStatus);
+				}
+			}
+
+			return TArray<FArrangedWidget>();
+
+		}();
 
 		if (HitWidgetIndex != INDEX_NONE)
 		{
-			TArray<FArrangedWidget> BubblePath;
 			int32 CurWidgetIndex=HitWidgetIndex;
 			bool bPathUninterrupted = false;
 			do
@@ -90,7 +106,8 @@ TArray<FArrangedWidget> FHittestGrid::GetBubblePath( FVector2D DesktopSpaceCoord
 				check( CurWidgetIndex < WidgetsCachedThisFrame->Num() );
 				const FCachedWidget& CurCachedWidget = (*WidgetsCachedThisFrame)[CurWidgetIndex];
 				const TSharedPtr<SWidget> CachedWidgetPtr = CurCachedWidget.WidgetPtr.Pin();
-				
+
+		
 				bPathUninterrupted = CachedWidgetPtr.IsValid();
 				if (bPathUninterrupted)
 				{
@@ -200,6 +217,14 @@ int32 FHittestGrid::InsertWidget( const int32 ParentHittestIndex, const EVisibil
 
 }
 
+void FHittestGrid::InsertCustomHitTestPath( TSharedRef<ICustomHitTestPath> CustomHitTestPath, int32 WidgetIndex )
+{
+	if( WidgetsCachedThisFrame->IsValidIndex( WidgetIndex ) )
+	{
+		(*WidgetsCachedThisFrame)[WidgetIndex].CustomPath = CustomHitTestPath;
+	}
+}
+
 FHittestGrid::FCell& FHittestGrid::CellAt( const int32 X, const int32 Y )
 {
 	check( (Y*NumCells.X + X) < Cells.Num() );
diff --git a/Engine/Source/Runtime/SlateCore/Private/Layout/WidgetPath.cpp b/Engine/Source/Runtime/SlateCore/Private/Layout/WidgetPath.cpp
--- a/Engine/Source/Runtime/SlateCore/Private/Layout/WidgetPath.cpp
+++ b/Engine/Source/Runtime/SlateCore/Private/Layout/WidgetPath.cpp
@@ -111,14 +111,14 @@ FWeakWidgetPath::FWeakWidgetPath( const FWidgetPath& InWidgetPath )
 }
 
 /** Make a non-weak WidgetPath out of this WeakWidgetPath. Do this by computing all the relevant geometries and converting the weak pointers to TSharedPtr. */	
-FWidgetPath FWeakWidgetPath::ToWidgetPath(EInterruptedPathHandling::Type InterruptedPathHandling) const
+FWidgetPath FWeakWidgetPath::ToWidgetPath(EInterruptedPathHandling::Type InterruptedPathHandling, const FPointerEvent* PointerEvent ) const
 {
 	FWidgetPath WidgetPath;
-	ToWidgetPath( WidgetPath, InterruptedPathHandling );
+	ToWidgetPath( WidgetPath, InterruptedPathHandling, PointerEvent );
 	return WidgetPath;
 }
 
-FWeakWidgetPath::EPathResolutionResult::Result FWeakWidgetPath::ToWidgetPath( FWidgetPath& WidgetPath, EInterruptedPathHandling::Type InterruptedPathHandling ) const
+FWeakWidgetPath::EPathResolutionResult::Result FWeakWidgetPath::ToWidgetPath( FWidgetPath& WidgetPath, EInterruptedPathHandling::Type InterruptedPathHandling, const FPointerEvent* PointerEvent ) const
 {
 	FArrangedChildren PathWithGeometries(EVisibility::Visible);
 	TArray< TSharedPtr<SWidget> > WidgetPtrs;
@@ -141,8 +141,9 @@ FWeakWidgetPath::EPathResolutionResult::Result FWeakWidgetPath::ToWidgetPath( FW
 		FGeometry ParentGeometry = TopLevelWindowPtr->GetWindowGeometryInScreen();
 		PathWithGeometries.AddWidget( FArrangedWidget( TopLevelWindowPtr.ToSharedRef(), ParentGeometry ) );
 		
-		FArrangedChildren ArrangedChildren(EVisibility::Visible);
+		FArrangedChildren ArrangedChildren(EVisibility::Visible, true);
 		
+		TSharedPtr<FVirtualCursorPosition> VirtualCursorPos;
 		// For every widget in the vertical slice...
 		for( int32 WidgetIndex = 0; bPathUninterrupted && WidgetIndex < WidgetPtrs.Num()-1; ++WidgetIndex )
 		{
@@ -157,9 +158,18 @@ FWeakWidgetPath::EPathResolutionResult::Result FWeakWidgetPath::ToWidgetPath( FW
 				
 				// Find the next widget in the path among the arranged children.
 				for( int32 SearchIndex = 0; !bFoundChild && SearchIndex < ArrangedChildren.Num(); ++SearchIndex )
-				{						
-					if ( ArrangedChildren[SearchIndex].Widget == WidgetPtrs[WidgetIndex+1] )
+				{					
+					FArrangedWidget& ArrangedWidget = ArrangedChildren[SearchIndex];
+
+					if ( ArrangedWidget.Widget == WidgetPtrs[WidgetIndex+1] )
 					{
+						if( PointerEvent && !VirtualCursorPos.IsValid() )
+						{
+							VirtualCursorPos = CurWidget->TranslateMouseCoordinateFor3DChild( ArrangedWidget.Widget, ParentGeometry, PointerEvent->GetScreenSpacePosition(), PointerEvent->GetLastScreenSpacePosition() );
+						}
+
+						ArrangedWidget.VirtualCursorPosition = VirtualCursorPos;
+
 						bFoundChild = true;
 						// Remember the widget and the associated geometry.
 						PathWithGeometries.AddWidget( ArrangedChildren[SearchIndex] );
diff --git a/Engine/Source/Runtime/SlateCore/Private/Types/PaintArgs.cpp b/Engine/Source/Runtime/SlateCore/Private/Types/PaintArgs.cpp
--- a/Engine/Source/Runtime/SlateCore/Private/Types/PaintArgs.cpp
+++ b/Engine/Source/Runtime/SlateCore/Private/Types/PaintArgs.cpp
@@ -34,4 +34,10 @@ FPaintArgs FPaintArgs::RecordHittestGeometry(const SWidget* Widget, const FGeome
 	UpdatedArgs.LastHittestIndex = RecordedHittestIndex;
 	UpdatedArgs.LastRecordedVisibility = RecordedVisibility;
 	return UpdatedArgs;
+}
+
+FPaintArgs FPaintArgs::InsertCustomHitTestPath( TSharedRef<ICustomHitTestPath> CustomHitTestPath, int32 LastHittestIndex ) const
+{
+	const_cast<FHittestGrid&>(Grid).InsertCustomHitTestPath( CustomHitTestPath, LastHittestIndex );
+	return *this;
 }
\ No newline at end of file
diff --git a/Engine/Source/Runtime/SlateCore/Public/Input/Events.h b/Engine/Source/Runtime/SlateCore/Public/Input/Events.h
--- a/Engine/Source/Runtime/SlateCore/Public/Input/Events.h
+++ b/Engine/Source/Runtime/SlateCore/Public/Input/Events.h
@@ -81,6 +81,14 @@ private:
 	EKeyboardFocusCause::Type Cause;
 };
 
+/**
+ * Represents the current and last cursor position in a "virtual window" for events that are routed to widgets transformed in a 3D scene.
+ */
+struct FVirtualCursorPosition
+{
+	FVector2D CurrentCursorPosition;
+	FVector2D LastCursorPosition;
+};
 
 /**
  * Base class for all mouse and keyboard events.
@@ -619,6 +627,16 @@ public:
 
 	SLATECORE_API virtual FText ToText() const override;
 
+	template<typename PointerEventType>
+	static PointerEventType MakeTranslatedEvent( const PointerEventType& InPointerEvent, const FVirtualCursorPosition& VirtualPosition )
+	{
+		PointerEventType NewEvent = InPointerEvent;
+		NewEvent.ScreenSpacePosition = VirtualPosition.CurrentCursorPosition;
+		NewEvent.LastScreenSpacePosition = VirtualPosition.LastCursorPosition;
+		//NewEvent.CursorDelta = VirtualPosition.GetDelta();
+		return NewEvent;
+	}
+
 private:
 
 	FVector2D ScreenSpacePosition;
diff --git a/Engine/Source/Runtime/SlateCore/Public/Input/HittestGrid.h b/Engine/Source/Runtime/SlateCore/Public/Input/HittestGrid.h
--- a/Engine/Source/Runtime/SlateCore/Public/Input/HittestGrid.h
+++ b/Engine/Source/Runtime/SlateCore/Public/Input/HittestGrid.h
@@ -4,6 +4,17 @@
 
 class FArrangedWidget;
 
+class ICustomHitTestPath
+{
+public:
+	virtual ~ICustomHitTestPath(){}
+
+	virtual TArray<FArrangedWidget> GetBubblePath( const FGeometry& InGeometry, FVector2D DesktopSpaceCoordinate, bool bIgnoreEnabledStatus ) const = 0;
+
+	virtual void ArrangeChildren( FArrangedChildren& ArrangedChildren ) const = 0;
+
+	virtual TSharedPtr<struct FVirtualCursorPosition> TranslateMouseCoordinateFor3DChild( const TSharedRef<SWidget>& ChildWidget, const FGeometry& ViewportGeometry, const FVector2D& ScreenSpaceMouseCoordinate, const FVector2D& LastScreenSpaceMouseCoordinate ) const = 0;
+};
 
 class SLATECORE_API FHittestGrid
 {
@@ -29,6 +40,7 @@ public:
 	/** Add Widget into the hittest data structure so that we can later make queries about it. */
 	int32 InsertWidget( const int32 ParentHittestIndex, const EVisibility& Visibility, const FArrangedWidget& Widget, const FVector2D InWindowOffset, const FSlateRect& InClippingRect );
 
+	void InsertCustomHitTestPath( TSharedRef<ICustomHitTestPath> CustomHitTestPath, int32 WidgetIndex );
 private:
 
 	/**
diff --git a/Engine/Source/Runtime/SlateCore/Public/Layout/ArrangedChildren.h b/Engine/Source/Runtime/SlateCore/Public/Layout/ArrangedChildren.h
--- a/Engine/Source/Runtime/SlateCore/Public/Layout/ArrangedChildren.h
+++ b/Engine/Source/Runtime/SlateCore/Public/Layout/ArrangedChildren.h
@@ -23,8 +23,9 @@ class SLATECORE_API FArrangedChildren
 	 *  FArrangedChildren ArrangedChildren( EVisibility::Visible ); // Only visible children will be included
 	 *  FArrangedChildren ArrangedChildren( EVisibility::Collapsed | EVisibility::Hidden ); // Only hidden and collapsed children will be included.
 	 */
-	FArrangedChildren( EVisibility InVisibilityFilter )
+	FArrangedChildren( EVisibility InVisibilityFilter, bool bInAllow3DWidgets = false )
 	: VisibilityFilter( InVisibilityFilter )
+	, bAllow3DWidgets( bInAllow3DWidgets )
 	{
 	}
 
@@ -71,16 +72,19 @@ class SLATECORE_API FArrangedChildren
 		return 0 != (InVisibility.Value & VisibilityFilter.Value);
 	}
 
+	bool Allows3DWidgets() const { return bAllow3DWidgets; }
+
 	// We duplicate parts of TArray's interface here!
 	// Inheriting from "public TArray<FArrangedWidget>"
 	// saves us this boilerplate, but causes instantiation of
 	// many template combinations that we do not want.
 
-	private:
+private:
 	/** Internal representation of the array widgets */
 	TArray<FArrangedWidget> Array;
 
-	public:
+	bool bAllow3DWidgets;
+public:
 
 	TArray<FArrangedWidget>& GetInternalArray()
 	{
diff --git a/Engine/Source/Runtime/SlateCore/Public/Layout/ArrangedWidget.h b/Engine/Source/Runtime/SlateCore/Public/Layout/ArrangedWidget.h
--- a/Engine/Source/Runtime/SlateCore/Public/Layout/ArrangedWidget.h
+++ b/Engine/Source/Runtime/SlateCore/Public/Layout/ArrangedWidget.h
@@ -4,6 +4,8 @@
 
 #include "Geometry.h"
 
+struct FVirtualCursorPosition;
+
 /**
  * A pair: Widget and its Geometry. Widgets populate an list of WidgetGeometries
  * when they arrange their children. See SWidget::ArrangeChildren.
@@ -46,4 +48,7 @@ public:
 
 	/** The widget that is being arranged. */
 	TSharedRef<SWidget> Widget;
+
+	/** Virtual cursor position if this widget is transformed in a 3D scene */
+	TSharedPtr<FVirtualCursorPosition> VirtualCursorPosition;
 };
diff --git a/Engine/Source/Runtime/SlateCore/Public/Layout/WidgetPath.h b/Engine/Source/Runtime/SlateCore/Public/Layout/WidgetPath.h
--- a/Engine/Source/Runtime/SlateCore/Public/Layout/WidgetPath.h
+++ b/Engine/Source/Runtime/SlateCore/Public/Layout/WidgetPath.h
@@ -4,7 +4,7 @@
 
 
 class SWindow;
-
+struct FPointerEvent;
 
 /** Matches widgets against InWidget */
 struct FWidgetMatcher
@@ -249,7 +249,7 @@ public:
 	 *
 	 * @param InterruptedPathHandling  Should interrupted paths result in a truncated path or an invalid path
 	 */
-	FWidgetPath ToWidgetPath( EInterruptedPathHandling::Type InterruptedPathHandling = EInterruptedPathHandling::Truncate ) const;
+	FWidgetPath ToWidgetPath( EInterruptedPathHandling::Type InterruptedPathHandling = EInterruptedPathHandling::Truncate, const FPointerEvent* PointerEvent = nullptr ) const;
 
 	struct EPathResolutionResult
 	{
@@ -266,7 +266,7 @@ public:
 	 * @param InterruptedPathHandling	Should interrupted paths result in a truncated path or an invalid path.
 	 * @return Whether the path is truncated or live - a live path refers to a widget that is currently active and visible, a widget with a truncated path is not.
 	 */
-	EPathResolutionResult::Result ToWidgetPath( FWidgetPath& WidgetPath, EInterruptedPathHandling::Type InterruptedPathHandling = EInterruptedPathHandling::Truncate ) const;
+	EPathResolutionResult::Result ToWidgetPath( FWidgetPath& WidgetPath, EInterruptedPathHandling::Type InterruptedPathHandling = EInterruptedPathHandling::Truncate, const FPointerEvent* PointerEvent = nullptr ) const;
 
 	bool ContainsWidget( const TSharedRef< const SWidget >& SomeWidget ) const;
 
diff --git a/Engine/Source/Runtime/SlateCore/Public/Types/PaintArgs.h b/Engine/Source/Runtime/SlateCore/Public/Types/PaintArgs.h
--- a/Engine/Source/Runtime/SlateCore/Public/Types/PaintArgs.h
+++ b/Engine/Source/Runtime/SlateCore/Public/Types/PaintArgs.h
@@ -8,6 +8,7 @@ class SWidget;
 struct FGeometry;
 class FHittestGrid;
 class FSlateRect;
+class ICustomHitTestPath;
 
 /**
  * SWidget::OnPaint and SWidget::Paint use FPaintArgs as their
@@ -20,8 +21,8 @@ public:
 	FPaintArgs( const TSharedRef<SWidget>& Parent, FHittestGrid& InHittestGrid, FVector2D InWindowOffset );
 	FPaintArgs WithNewParent( const SWidget* Parent ) const;
 	FPaintArgs RecordHittestGeometry(const SWidget* Widget, const FGeometry& WidgetGeometry, const FSlateRect& InClippingRect) const;
-
-
+	FPaintArgs InsertCustomHitTestPath( TSharedRef<ICustomHitTestPath> CustomHitTestPath, int32 HitTestIndex ) const;
+	int32 GetLastHitTestIndex() const { return LastHittestIndex; }
 private:
 	const TSharedRef<SWidget>& ParentPtr;
 	FHittestGrid& Grid;
diff --git a/Engine/Source/Runtime/SlateCore/Public/Widgets/SWidget.h b/Engine/Source/Runtime/SlateCore/Public/Widgets/SWidget.h
--- a/Engine/Source/Runtime/SlateCore/Public/Widgets/SWidget.h
+++ b/Engine/Source/Runtime/SlateCore/Public/Widgets/SWidget.h
@@ -403,6 +403,9 @@ public:
 	 */
 	virtual FReply OnMotionDetected( const FGeometry& MyGeometry, const FMotionEvent& InMotionEvent );
 
+	virtual TSharedPtr<struct FVirtualCursorPosition> TranslateMouseCoordinateFor3DChild( const TSharedRef<SWidget>& ChildWidget, const FGeometry& MyGeometry, const FVector2D& ScreenSpaceMouseCoordinate, const FVector2D& LastScreenSpaceMouseCoordinate ) const { return nullptr; }
+
+
 	/**
 	 * All the pointer (mouse, touch, stylus, etc.) events from this frame have been routed.
 	 * This is a widget's chance to act on any accumulated data.
diff --git a/Engine/Source/Runtime/SlateRHIRenderer/Private/Slate3DRenderer.cpp b/Engine/Source/Runtime/SlateRHIRenderer/Private/Slate3DRenderer.cpp
new file mode 100644
--- /dev/null
+++ b/Engine/Source/Runtime/SlateRHIRenderer/Private/Slate3DRenderer.cpp
@@ -0,0 +1,119 @@
+// Copyright 1998-2014 Epic Games, Inc. All Rights Reserved.
+
+#include "SlateRHIRendererPrivatePCH.h"
+#include "Slate3DRenderer.h"
+#include "ElementBatcher.h"
+
+FSlate3DRenderer::FSlate3DRenderer( TSharedPtr<FSlateRHIResourceManager> InResourceManager, TSharedPtr<FSlateFontCache> InFontCache )
+	: ResourceManager( InResourceManager.ToSharedRef() )
+	, FontCache( InFontCache.ToSharedRef() )
+{
+
+	RenderTargetPolicy = MakeShareable( new FSlateRHIRenderingPolicy( FontCache, ResourceManager ) );
+	RenderTargetPolicy->SetUseGammaCorrection( false );
+
+	ElementBatcher = MakeShareable(new FSlateElementBatcher(RenderTargetPolicy.ToSharedRef()));
+}
+
+FSlate3DRenderer::~FSlate3DRenderer()
+{
+	if( RenderTargetPolicy.IsValid() )
+	{
+		RenderTargetPolicy->ReleaseResources();
+	}
+
+	FlushRenderingCommands();
+}
+
+FSlateDrawBuffer& FSlate3DRenderer::GetDrawBuffer()
+{
+	FreeBufferIndex = (FreeBufferIndex + 1) % 2;
+
+	FSlateDrawBuffer* Buffer = &DrawBuffers[FreeBufferIndex];
+
+	while (!Buffer->Lock())
+	{
+		FlushRenderingCommands();
+
+		UE_LOG(LogSlate, Log, TEXT("Slate: Had to block on waiting for a draw buffer"));
+		FreeBufferIndex = (FreeBufferIndex + 1) % NumDrawBuffers;
+
+		Buffer = &DrawBuffers[FreeBufferIndex];
+	}
+
+
+	Buffer->ClearBuffer();
+
+	return *Buffer;
+}
+
+void FSlate3DRenderer::DrawWindow_GameThread(FSlateDrawBuffer& DrawBuffer)
+{
+	check( IsInGameThread() );
+
+	TArray<FSlateWindowElementList>& WindowElementLists = DrawBuffer.GetWindowElementLists();
+
+	// Only one window element list for now.
+	check( WindowElementLists.Num() == 1);
+
+	FSlateWindowElementList& ElementList = WindowElementLists[0];
+
+	TSharedPtr<SWindow> Window = ElementList.GetWindow();
+
+	if (Window.IsValid())
+	{
+		const FVector2D WindowSize = Window->GetSizeInScreen();
+		if (WindowSize.X > 0 && WindowSize.Y > 0)
+		{
+			// Add all elements for this window to the element batcher
+			ElementBatcher->AddElements(ElementList.GetDrawElements());
+
+			// Update the font cache with new text after elements are batched
+			FontCache->UpdateCache();
+
+			bool temp = false;
+			// Populate the element list with batched vertices and indices
+			ElementBatcher->FillBatchBuffers(ElementList, temp);
+
+
+			// All elements for this window have been batched and rendering data updated
+			ElementBatcher->ResetBatches();
+		}
+	}
+
+	FontCache->ConditionalFlushCache();
+
+	ElementBatcher->ResetStats();
+
+}
+
+void FSlate3DRenderer::DrawWindowToTarget_RenderThread( FRHICommandListImmediate& RHICmdList, UTextureRenderTarget2D* RenderTarget, FSlateDrawBuffer& InDrawBuffer )
+{
+	const FSlateWindowElementList& WindowElementList = InDrawBuffer.GetWindowElementLists()[0];
+
+	RenderTargetPolicy->BeginDrawingWindows();
+
+	RenderTargetPolicy->UpdateBuffers( WindowElementList );
+	FTextureRenderTarget2DResource* RenderTargetResource = static_cast<FTextureRenderTarget2DResource*>( RenderTarget->GetRenderTargetResource() );
+	
+	SetRenderTarget( RHICmdList, RenderTargetResource->GetTextureRHI(), FTextureRHIParamRef() );
+	RHICmdList.Clear( true, FLinearColor(0,0,0,0), false, 0.0f, false, 0x00, FIntRect() );
+
+	FMatrix ProjectionMatrix = FSlateRHIRenderer::CreateProjectionMatrix( RenderTarget->SizeX, RenderTarget->SizeY );
+	if (WindowElementList.GetRenderBatches().Num() > 0)
+	{
+		FSlateBackBuffer BackBufferTarget(RenderTargetResource->GetTextureRHI(), FIntPoint(RenderTarget->SizeX, RenderTarget->SizeY ) );
+
+		RenderTargetPolicy->DrawElements
+		(
+			RHICmdList,
+			BackBufferTarget,
+			ProjectionMatrix,
+			WindowElementList.GetRenderBatches()
+			);
+	}
+
+	InDrawBuffer.Unlock();
+
+	RenderTargetPolicy->EndDrawingWindows();
+}
\ No newline at end of file
diff --git a/Engine/Source/Runtime/SlateRHIRenderer/Private/Slate3DRenderer.h b/Engine/Source/Runtime/SlateRHIRenderer/Private/Slate3DRenderer.h
new file mode 100644
--- /dev/null
+++ b/Engine/Source/Runtime/SlateRHIRenderer/Private/Slate3DRenderer.h
@@ -0,0 +1,35 @@
+// Copyright 1998-2014 Epic Games, Inc. All Rights Reserved.
+
+#pragma once
+
+#include "ISlate3DRenderer.h"
+
+
+class FSlate3DRenderer : public ISlate3DRenderer
+{
+public:
+	FSlate3DRenderer( TSharedPtr<FSlateRHIResourceManager> InResourceManager, TSharedPtr<FSlateFontCache> InFontCache );
+	~FSlate3DRenderer();
+
+	virtual FSlateDrawBuffer& GetDrawBuffer() override;
+	virtual void DrawWindow_GameThread(FSlateDrawBuffer& DrawBuffer) override;
+	virtual void DrawWindowToTarget_RenderThread( FRHICommandListImmediate& RHICmdList, UTextureRenderTarget2D* RenderTarget, FSlateDrawBuffer& InDrawBuffer ) override;
+private:
+
+	/** Double buffered draw buffers so that the rendering thread can be rendering windows while the game thread is setting up for next frame */
+	FSlateDrawBuffer DrawBuffers[2];
+
+	/** Texture manager for accessing textures on the game thread */
+	TSharedRef<FSlateRHIResourceManager> ResourceManager;
+
+	TSharedRef<FSlateFontCache> FontCache;
+
+	TSharedPtr<class FSlateRHIRenderingPolicy> RenderTargetPolicy;
+
+	/** Element batcher which renders draw elements */
+	TSharedPtr<FSlateElementBatcher> ElementBatcher;
+
+	/** The draw buffer which is currently free for use by the game thread */
+	uint8 FreeBufferIndex;
+
+};
\ No newline at end of file
diff --git a/Engine/Source/Runtime/SlateRHIRenderer/Private/SlateRHIRendererModule.cpp b/Engine/Source/Runtime/SlateRHIRenderer/Private/SlateRHIRendererModule.cpp
--- a/Engine/Source/Runtime/SlateRHIRenderer/Private/SlateRHIRendererModule.cpp
+++ b/Engine/Source/Runtime/SlateRHIRenderer/Private/SlateRHIRendererModule.cpp
@@ -1,6 +1,7 @@
 // Copyright 1998-2014 Epic Games, Inc. All Rights Reserved.
 
 #include "SlateRHIRendererPrivatePCH.h"
+#include "Slate3DRenderer.h"
 
 class FSlateRHIFontAtlasFactory : public ISlateFontAtlasFactory
 {
@@ -46,6 +47,12 @@ public:
 		return MakeShareable( new FSlateRHIRenderer( ResourceManager, FontCache, FontMeasure ) );
 	}
 
+	virtual TSharedRef<ISlate3DRenderer> CreateSlate3DRenderer() override
+	{
+		ConditionalCreateResources();
+
+		return MakeShareable( new FSlate3DRenderer( ResourceManager, FontCache ) );
+	}
 
 	virtual void StartupModule( ) override { }
 	virtual void ShutdownModule( ) override { }
diff --git a/Engine/Source/Runtime/SlateRHIRenderer/Private/SlateRHIRenderingPolicy.cpp b/Engine/Source/Runtime/SlateRHIRenderer/Private/SlateRHIRenderingPolicy.cpp
--- a/Engine/Source/Runtime/SlateRHIRenderer/Private/SlateRHIRenderingPolicy.cpp
+++ b/Engine/Source/Runtime/SlateRHIRenderer/Private/SlateRHIRenderingPolicy.cpp
@@ -94,6 +94,7 @@ FSlateRHIRenderingPolicy::FSlateRHIRenderingPolicy( TSharedPtr<FSlateFontCache>
 	, FontCache( InFontCache )
 	, CurrentBufferIndex(0)
 	, bShouldShrinkResources(false)
+	, bGammaCorrect(true)
 {
 	InitResources();
 };
@@ -232,10 +233,10 @@ static FSceneView& CreateSceneView( FSceneViewFamilyContext& ViewFamilyContext,
 	ViewUniformShaderParameters.RealTime = View->Family->CurrentRealTime;
 	ViewUniformShaderParameters.Random = FMath::Rand();
 	ViewUniformShaderParameters.FrameNumber = View->FrameNumber;
+
 	ViewUniformShaderParameters.DirectionalLightShadowTexture = GWhiteTexture->TextureRHI;
 	ViewUniformShaderParameters.DirectionalLightShadowSampler = TStaticSamplerState<SF_Point, AM_Clamp, AM_Clamp, AM_Clamp>::GetRHI();
 
-
 	ViewUniformShaderParameters.ScreenToWorld = FMatrix(
 		FPlane(1, 0, 0, 0),
 		FPlane(0, 1, 0, 0),
@@ -270,7 +271,7 @@ void FSlateRHIRenderingPolicy::DrawElements(FRHICommandListImmediate& RHICmdList
 
 	static const FEngineShowFlags DefaultShowFlags(ESFIM_Game);
 
-	const float DisplayGamma = GEngine ? GEngine->GetDisplayGamma() : 2.2f;
+	const float DisplayGamma = bGammaCorrect ? GEngine ? GEngine->GetDisplayGamma() : 2.2f : 1.0f;
 
 	FSceneViewFamilyContext SceneViewContext
 	(
diff --git a/Engine/Source/Runtime/SlateRHIRenderer/Private/SlateRHIRenderingPolicy.h b/Engine/Source/Runtime/SlateRHIRenderer/Private/SlateRHIRenderingPolicy.h
--- a/Engine/Source/Runtime/SlateRHIRenderer/Private/SlateRHIRenderingPolicy.h
+++ b/Engine/Source/Runtime/SlateRHIRenderer/Private/SlateRHIRenderingPolicy.h
@@ -175,7 +175,7 @@ class FSlateRHIRenderingPolicy : public FSlateRenderingPolicy
 public:
 	FSlateRHIRenderingPolicy( TSharedPtr<FSlateFontCache> InFontCache, TSharedRef<FSlateRHIResourceManager> InResourceManager );
 	~FSlateRHIRenderingPolicy();
-	
+
 	virtual void UpdateBuffers( const FSlateWindowElementList& WindowElementList ) override;
 	virtual void DrawElements(FRHICommandListImmediate& RHICmdList, class FSlateBackBuffer& BackBuffer, const FMatrix& ViewProjectionMatrix, const TArray<FSlateRenderBatch>& RenderBatches);
 
@@ -187,6 +187,8 @@ public:
 
 	void BeginDrawingWindows();
 	void EndDrawingWindows();
+
+	void SetUseGammaCorrection( bool bInUseGammaCorrection ) { bGammaCorrect = bInUseGammaCorrection; }
 private:
 	/**
 	 * Returns the pixel shader that should be used for the specified ShaderType and DrawEffects
@@ -212,6 +214,8 @@ private:
 	
 	/** If we should shrink resources that are no longer used (do not set this from the game thread)*/
 	bool bShouldShrinkResources;
+
+	bool bGammaCorrect;
 };
 
 
diff --git a/Engine/Source/Runtime/SlateRHIRenderer/Public/Interfaces/ISlate3DRenderer.h b/Engine/Source/Runtime/SlateRHIRenderer/Public/Interfaces/ISlate3DRenderer.h
new file mode 100644
--- /dev/null
+++ b/Engine/Source/Runtime/SlateRHIRenderer/Public/Interfaces/ISlate3DRenderer.h
@@ -0,0 +1,28 @@
+// Copyright 1998-2014 Epic Games, Inc. All Rights Reserved.
+
+#pragma once
+
+class ISlate3DRenderer 
+{
+public:
+	virtual ~ISlate3DRenderer() {};
+
+	/** @return The free buffer for drawing */
+	virtual FSlateDrawBuffer& GetDrawBuffer() = 0;
+
+	/** 
+	* Batches the draw elements in the buffer to prepare it for rendering.
+	* Call in the game thread before sending to the render thread.
+	* @param DrawBuffer The draw buffer to prepare
+	*/
+	virtual void DrawWindow_GameThread(FSlateDrawBuffer& DrawBuffer) = 0;
+
+	/** 
+	* Renders the batched draw elements of the draw buffer to the given render target.
+	* Call after preparing the draw buffer and render target on the game thread.
+	* 
+	* @param RenderTarget The render target to render the contents of the draw buffer to
+	* @param InDrawBuffer The draw buffer containing the batched elements to render
+	*/
+	virtual void DrawWindowToTarget_RenderThread( FRHICommandListImmediate& RHICmdList, UTextureRenderTarget2D* RenderTarget, FSlateDrawBuffer& InDrawBuffer ) = 0;
+};
diff --git a/Engine/Source/Runtime/SlateRHIRenderer/Public/Interfaces/ISlateRHIRendererModule.h b/Engine/Source/Runtime/SlateRHIRenderer/Public/Interfaces/ISlateRHIRendererModule.h
--- a/Engine/Source/Runtime/SlateRHIRenderer/Public/Interfaces/ISlateRHIRendererModule.h
+++ b/Engine/Source/Runtime/SlateRHIRenderer/Public/Interfaces/ISlateRHIRendererModule.h
@@ -5,6 +5,7 @@
 #include "ModuleInterface.h"
 #include "SlateCore.h"
 
+class ISlate3DRenderer;
 
 /**
  * Interface for the Slate RHI Renderer module.
@@ -20,4 +21,6 @@ public:
 	 * @return A new renderer.
 	 */
 	virtual TSharedRef<FSlateRenderer> CreateSlateRHIRenderer( ) = 0;
+
+	virtual TSharedRef<ISlate3DRenderer> CreateSlate3DRenderer() = 0;
 };
diff --git a/Engine/Source/Runtime/UMG/Private/Components/UMGComponent.cpp b/Engine/Source/Runtime/UMG/Private/Components/UMGComponent.cpp
new file mode 100644
--- /dev/null
+++ b/Engine/Source/Runtime/UMG/Private/Components/UMGComponent.cpp
@@ -0,0 +1,629 @@
+// Copyright 1998-2014 Epic Games, Inc. All Rights Reserved.
+
+#include "UMGPrivatePCH.h"
+#include "UMGComponent.h"
+#include "HittestGrid.h"
+#include "ISlateRHIRendererModule.h"
+#include "ISlate3DRenderer.h"
+#include "DynamicMeshBuilder.h"
+
+
+class SVirtualWindow : public SWindow
+{
+	SLATE_BEGIN_ARGS( SVirtualWindow )
+		: _Size( FVector2D(100,100) )
+	{}
+
+		SLATE_ARGUMENT( FVector2D, Size )
+	SLATE_END_ARGS()
+
+public:
+	void Construct( const FArguments& InArgs )
+	{
+		bIsPopupWindow = true;
+		SetCachedSize( InArgs._Size );
+		SetNativeWindow( MakeShareable( new FGenericWindow() ) );
+		SetContent( SNullWidget::NullWidget );
+	}
+};
+
+
+/** Represents a billboard sprite to the scene manager. */
+class FUMG3DSceneProxy : public FPrimitiveSceneProxy
+{
+public:
+	/** Initialization constructor. */
+	FUMG3DSceneProxy(const UUMGComponent* InComponent, ISlate3DRenderer& InRenderer )
+		: FPrimitiveSceneProxy(InComponent)
+		, Renderer( InRenderer )
+		, RenderTarget( InComponent->GetRenderTarget() )
+		, MaterialInstance( InComponent->GetMaterialInstance() )
+		, BodySetup( const_cast<UUMGComponent*>( InComponent )->GetBodySetup() )
+	{
+		bWillEverBeLit = false;
+	}
+
+	~FUMG3DSceneProxy()
+	{
+	}
+
+	// FPrimitiveSceneProxy interface.
+	virtual void GetDynamicMeshElements(const TArray<const FSceneView*>& Views, const FSceneViewFamily& ViewFamily, uint32 VisibilityMap, FMeshElementCollector& Collector) const override
+	{
+		const FMatrix& LocalToWorld = GetLocalToWorld();
+	
+		FDynamicMeshBuilder MeshBuilder;
+
+		if( RenderTarget )
+		{
+			FTextureResource* TextureResource = RenderTarget->Resource;
+			if(TextureResource)
+			{
+				float U = 0.0f;
+				float V = 0.0f;
+				float UL = RenderTarget->SizeX;
+				float VL = RenderTarget->SizeY;
+
+				int32 VertexIndices[4];
+
+				VertexIndices[0] = MeshBuilder.AddVertex(FVector(U, V, 0), FVector2D(0, 0), FVector(1, 0, 0), FVector(0, 1, 0), FVector(0, 0, 1), FColor::White);
+				VertexIndices[1] = MeshBuilder.AddVertex(FVector(U, VL, 0), FVector2D(0, 1), FVector(1, 0, 0), FVector(0, 1, 0), FVector(0, 0, 1), FColor::White);
+				VertexIndices[2] = MeshBuilder.AddVertex(FVector(UL, VL, 0), FVector2D(1, 1), FVector(1, 0, 0), FVector(0, 1, 0), FVector(0, 0, 1), FColor::White);
+				VertexIndices[3] = MeshBuilder.AddVertex(FVector(UL, V, 0), FVector2D(1, 0), FVector(1, 0, 0), FVector(0, 1, 0), FVector(0, 0, 1), FColor::White);
+
+				MeshBuilder.AddTriangle(VertexIndices[0], VertexIndices[1], VertexIndices[2]);
+				MeshBuilder.AddTriangle(VertexIndices[0], VertexIndices[2], VertexIndices[3]);
+			}
+
+
+			for(int32 ViewIndex = 0; ViewIndex < Views.Num(); ViewIndex++)
+			{
+				if(VisibilityMap & (1 << ViewIndex))
+				{
+					FPrimitiveDrawInterface* PDI = Collector.GetPDI(ViewIndex);
+
+					MeshBuilder.Draw(PDI, LocalToWorld, MaterialInstance->GetRenderProxy(IsSelected()), SDPG_World);
+				}
+			}
+		}
+
+	}
+	/** 
+	 * Draw the scene proxy as a dynamic element
+	 *
+	 * @param	PDI - draw interface to render to
+	 * @param	View - current view
+	 */
+	virtual void DrawDynamicElements(FPrimitiveDrawInterface* PDI, const FSceneView* View) override
+	{
+		if( RenderTarget )
+		{
+			FTextureResource* TextureResource = RenderTarget->Resource;
+			if(TextureResource)
+			{
+				float U = 0.0f;
+				float V = 0.0f;
+				float UL = RenderTarget->SizeX;
+				float VL = RenderTarget->SizeY;
+
+				FDynamicMeshBuilder MeshBuilder;
+
+				int32 VertexIndices[4];
+
+				VertexIndices[0] = MeshBuilder.AddVertex(FVector(U, V, 0), FVector2D(0, 0), FVector(1, 0, 0), FVector(0, 1, 0), FVector(0, 0, 1), FColor::White);
+				VertexIndices[1] = MeshBuilder.AddVertex(FVector(U, VL, 0), FVector2D(0, 1), FVector(1, 0, 0), FVector(0, 1, 0), FVector(0, 0, 1), FColor::White);
+				VertexIndices[2] = MeshBuilder.AddVertex(FVector(UL, VL, 0), FVector2D(1, 1), FVector(1, 0, 0), FVector(0, 1, 0), FVector(0, 0, 1), FColor::White);
+				VertexIndices[3] = MeshBuilder.AddVertex(FVector(UL, V, 0), FVector2D(1, 0), FVector(1, 0, 0), FVector(0, 1, 0), FVector(0, 0, 1), FColor::White);
+
+				MeshBuilder.AddTriangle(VertexIndices[0], VertexIndices[1], VertexIndices[2]);
+				MeshBuilder.AddTriangle(VertexIndices[0], VertexIndices[2], VertexIndices[3]);
+
+				MeshBuilder.Draw(PDI, GetLocalToWorld(), MaterialInstance->GetRenderProxy(IsSelected()), SDPG_World);
+
+			}
+		}
+
+		//FColor CollisionColor = FColor(157,149,223,255);
+		//FTransform GeomTransform(GetLocalToWorld());
+		//BodySetup->AggGeom.DrawAggGeom(PDI, GeomTransform, GetSelectionColor(CollisionColor, false, IsHovered()), nullptr, false, false, true );
+
+#if !(UE_BUILD_SHIPPING || UE_BUILD_TEST)
+		RenderBounds(PDI, View->Family->EngineShowFlags, GetBounds(), IsSelected());
+#endif
+	}
+
+	virtual FPrimitiveViewRelevance GetViewRelevance(const FSceneView* View) override
+	{
+		bool bVisible = true;// View->Family->EngineShowFlags.BillboardSprites;
+
+		FPrimitiveViewRelevance Result;
+		Result.bDrawRelevance = IsShown(View) && bVisible;
+		Result.bOpaqueRelevance = false;
+		Result.bNormalTranslucencyRelevance = true;
+		Result.bDynamicRelevance = true;
+		Result.bShadowRelevance = IsShadowCast(View);;
+		Result.bEditorPrimitiveRelevance = false;
+		return Result;
+	}
+
+	virtual void GetLightRelevance(const FLightSceneProxy* LightSceneProxy, bool& bDynamic, bool& bRelevant, bool& bLightMapped, bool& bShadowMapped) const override
+	{
+		bDynamic = false;
+		bRelevant = false;
+		bLightMapped = false;
+		bShadowMapped = false;
+	}
+
+	virtual void OnTransformChanged() override
+	{
+		Origin = GetLocalToWorld().GetOrigin();
+	}
+
+	virtual uint32 GetMemoryFootprint(void) const override { return(sizeof(*this) + GetAllocatedSize()); }
+
+	uint32 GetAllocatedSize(void) const { return( FPrimitiveSceneProxy::GetAllocatedSize() ); }
+
+private:
+	FVector Origin;
+	ISlate3DRenderer& Renderer;
+	UTextureRenderTarget2D* RenderTarget;
+	UMaterialInstanceDynamic* MaterialInstance;
+	UBodySetup* BodySetup;
+};
+
+/**
+* The hit tester used by all UMG Component objects.
+*/
+class FUMG3DHitTester : public ICustomHitTestPath
+{
+public:
+	FUMG3DHitTester( UWorld* InWorld )
+		: World( InWorld )
+	{}
+
+	// ICustomHitTestPath implementation
+	virtual TArray<FArrangedWidget> GetBubblePath(const FGeometry& InGeometry, FVector2D DesktopSpaceCoordinate, bool bIgnoreEnabledStatus) const override
+	{
+		if( World.IsValid() && ensure( World->IsGameWorld() ) )
+		{
+			ULocalPlayer* const TargetPlayer = GEngine->GetLocalPlayerFromControllerId(World.Get(), 0);
+			if (TargetPlayer && TargetPlayer->PlayerController)
+			{
+				FHitResult HitResult;
+				if (TargetPlayer->PlayerController->GetHitResultAtScreenPosition(InGeometry.AbsoluteToLocal(DesktopSpaceCoordinate), ECC_Visibility, true, HitResult))
+				{
+					UUMGComponent* UMGComponent = Cast<UUMGComponent>(HitResult.Component.Get());
+					if (UMGComponent)
+					{
+						return UMGComponent->GetHitWidgetPath(HitResult, bIgnoreEnabledStatus);
+					}
+				}
+			}
+		}
+		return TArray<FArrangedWidget>();
+	}
+
+	virtual void ArrangeChildren( FArrangedChildren& ArrangedChildren ) const override
+	{
+		for( TWeakObjectPtr<UUMGComponent> Component : RegisteredComponents )
+		{
+			UUMGComponent* UMGComponent = Component.Get();
+			// Check if visible;
+			if( UMGComponent && UMGComponent->GetSlateWidget().IsValid() )
+			{
+				FGeometry WidgetGeom;
+
+				ArrangedChildren.AddWidget( FArrangedWidget( UMGComponent->GetSlateWidget().ToSharedRef(), WidgetGeom.MakeChild( UMGComponent->GetDrawSize(), FSlateLayoutTransform() ) ) );
+			}
+		}
+	}
+
+	virtual TSharedPtr<struct FVirtualCursorPosition> TranslateMouseCoordinateFor3DChild( const TSharedRef<SWidget>& ChildWidget, const FGeometry& ViewportGeometry, const FVector2D& ScreenSpaceMouseCoordinate, const FVector2D& LastScreenSpaceMouseCoordinate ) const override
+	{
+		if(World.IsValid() && ensure(World->IsGameWorld()))
+		{
+			ULocalPlayer* const TargetPlayer = GEngine->GetLocalPlayerFromControllerId(World.Get(), 0);
+			if(TargetPlayer && TargetPlayer->PlayerController)
+			{
+				// Check for a hit against any umg components in the world
+				for(TWeakObjectPtr<UUMGComponent> Component : RegisteredComponents)
+				{
+					UUMGComponent* UMGComponent = Component.Get();
+					// Check if visible;
+					if(UMGComponent && UMGComponent->GetSlateWidget() == ChildWidget)
+					{
+						FVector2D LocalMouseCoordinate = ViewportGeometry.AbsoluteToLocal(ScreenSpaceMouseCoordinate);
+						FVector2D LocalLastMouseCoordinate = ViewportGeometry.AbsoluteToLocal( LastScreenSpaceMouseCoordinate );
+
+
+						FHitResult HitResult;
+						if(TargetPlayer->PlayerController->GetHitResultAtScreenPosition(LocalMouseCoordinate, ECC_Visibility, true, HitResult))
+						{
+							if(UMGComponent == HitResult.Component.Get())
+							{
+								TSharedPtr<FVirtualCursorPosition> VirtualCursorPos = MakeShareable(new FVirtualCursorPosition);
+
+								FVector2D LocalHitLocation = FVector2D(UMGComponent->ComponentToWorld.InverseTransformPosition(HitResult.Location));
+
+								VirtualCursorPos->CurrentCursorPosition = LocalHitLocation;
+								VirtualCursorPos->LastCursorPosition = LocalHitLocation;
+
+								return VirtualCursorPos;
+							}
+						}
+					}
+				}
+			}
+		}
+
+		return nullptr;
+	}
+	// End ICustomHitTestPath
+
+	void RegisterUMGComponent( UUMGComponent* InComponent )
+	{
+		RegisteredComponents.AddUnique( InComponent );
+	}
+
+	void UnregisterUMGComponent( UUMGComponent* InComponent )
+	{
+		RegisteredComponents.RemoveSingleSwap( InComponent );
+	}
+
+	uint32 GetNumRegisteredComponents() const { return RegisteredComponents.Num(); }
+	
+	UWorld* GetWorld() const { return World.Get(); }
+private:
+	TArray< TWeakObjectPtr<UUMGComponent> > RegisteredComponents;
+	TWeakObjectPtr<UWorld> World;
+};
+
+UUMGComponent::UUMGComponent(const class FPostConstructInitializeProperties& PCIP)
+	: Super(PCIP)
+{
+	PrimaryComponentTick.bCanEverTick = true;
+	bTickInEditor = true;
+	DrawSize = FIntPoint(500,500);
+
+	BodyInstance.SetCollisionProfileName(FName(TEXT("UI")));
+
+	static ConstructorHelpers::FObjectFinder<UMaterial> Material( TEXT("/Engine/EngineMaterials/UMG3DPassThrough") );
+
+	BaseMaterial = Material.Object;
+
+	LastLocalHitLocation = FVector2D::ZeroVector;
+	//bGenerateOverlapEvents = false;
+	bUseEditorCompositing = false;
+}
+
+FPrimitiveSceneProxy* UUMGComponent::CreateSceneProxy()
+{
+	return new FUMG3DSceneProxy(this,*Renderer);
+}
+
+FBoxSphereBounds UUMGComponent::CalcBounds(const FTransform & LocalToWorld) const
+{
+	return FBoxSphereBounds( FVector( DrawSize.X/2.0f, DrawSize.Y/2.0f, .5f), FVector(DrawSize.X/2,DrawSize.Y/2,1.0f), DrawSize.Size()/2 ).TransformBy( LocalToWorld );
+}
+
+UBodySetup* UUMGComponent::GetBodySetup() 
+{
+	UpdateBodySetup();
+	return BodySetup;
+}
+
+FCollisionShape UUMGComponent::GetCollisionShape(float Inflation) const
+{
+	FVector	Extent = ( FVector( DrawSize.Y, DrawSize.Y, 1.0f ) * ComponentToWorld.GetScale3D() ) + Inflation;
+	if( Inflation < 0.0f )
+	{
+		// Don't shrink below zero size.
+		Extent = Extent.ComponentMax(FVector::ZeroVector);
+	}
+
+	return FCollisionShape::MakeBox(Extent);
+}
+
+void UUMGComponent::OnRegister()
+{
+	Super::OnRegister();
+
+	if( GetWorld()->IsGameWorld() )
+	{
+		TSharedPtr<SViewport> GameViewportWidget = GEngine->GetGameViewportWidget();
+
+		if( GameViewportWidget.IsValid() )
+		{
+			TSharedPtr<ICustomHitTestPath> CustomHitTestPath = GameViewportWidget->GetCustomHitTestPath();
+			if( !CustomHitTestPath.IsValid() )
+			{
+				CustomHitTestPath = MakeShareable(new FUMG3DHitTester(GetWorld()));
+				GameViewportWidget->SetCustomHitTestPath( CustomHitTestPath );
+			}
+
+			TSharedPtr<FUMG3DHitTester> UMGHitTester = StaticCastSharedPtr<FUMG3DHitTester>(CustomHitTestPath);
+			if( UMGHitTester->GetWorld() == GetWorld() )
+			{
+				UMGHitTester->RegisterUMGComponent( this );
+			}
+
+		}
+
+	}
+
+	if( !MaterialInstance )
+	{
+		MaterialInstance = UMaterialInstanceDynamic::Create(BaseMaterial, this);
+	}
+
+	UpdateWidget();
+
+	if( !Renderer.IsValid() )
+	{
+		Renderer = FModuleManager::Get().GetModuleChecked<ISlateRHIRendererModule>("SlateRHIRenderer").CreateSlate3DRenderer();
+	}
+}
+
+void UUMGComponent::OnUnregister()
+{
+	if( GetWorld()->IsGameWorld() )
+	{
+		TSharedPtr<SViewport> GameViewportWidget = GEngine->GetGameViewportWidget();
+		if( GameViewportWidget.IsValid() )
+		{
+			TSharedPtr<ICustomHitTestPath> CustomHitTestPath = GameViewportWidget->GetCustomHitTestPath();
+			if( CustomHitTestPath.IsValid() )
+			{
+				TSharedPtr<FUMG3DHitTester> UMGHitTestPath = StaticCastSharedPtr<FUMG3DHitTester>(CustomHitTestPath);
+
+				UMGHitTestPath->UnregisterUMGComponent( this );
+
+				if( UMGHitTestPath->GetNumRegisteredComponents() == 0 )
+				{
+					GameViewportWidget->SetCustomHitTestPath( nullptr );
+				}
+			}
+		}
+	}
+
+	if( Widget )
+	{
+		Widget->MarkPendingKill();
+		Widget = NULL;
+	}
+
+	SlateWidget.Reset();
+	Super::OnUnregister();
+}
+
+void UUMGComponent::DestroyComponent()
+{
+	Super::DestroyComponent();
+	Renderer.Reset();
+}
+
+void UUMGComponent::TickComponent(float DeltaTime, enum ELevelTick TickType, FActorComponentTickFunction *ThisTickFunction)
+{
+	const float RenderTimeThreshold = .5f;
+	if( IsVisible() && GetWorld()->TimeSince(LastRenderTime) <= RenderTimeThreshold ) // Don't bother ticking if it hasn't been rendered recently
+	{
+		SlateWidget->SlatePrepass();
+
+		FGeometry WindowGeometry = FGeometry::MakeRoot( DrawSize, FSlateLayoutTransform() );
+
+		SlateWidget->TickWidgetsRecursively(WindowGeometry, FApp::GetCurrentTime(), DeltaTime);
+
+		// Ticking can cause geometry changes.  Recompute
+		SlateWidget->SlatePrepass();
+
+		// Get the free buffer & add our virtual window
+		FSlateDrawBuffer& DrawBuffer = Renderer->GetDrawBuffer();
+		FSlateWindowElementList& WindowElementList = DrawBuffer.AddWindowElementList(SlateWidget.ToSharedRef());
+
+		// Prepare the test grid 
+		HitTestGrid->BeginFrame(WindowGeometry.GetClippingRect());
+
+		// Paint the window
+		int32 MaxLayerId = 0;
+		{
+			MaxLayerId = SlateWidget->Paint(
+				FPaintArgs(SlateWidget.ToSharedRef(), *HitTestGrid, FVector2D::ZeroVector),
+				WindowGeometry, WindowGeometry.GetClippingRect(),
+				WindowElementList,
+				0,
+				FWidgetStyle(),
+				SlateWidget->IsEnabled());
+		}
+
+		Renderer->DrawWindow_GameThread(DrawBuffer);
+
+		UpdateRenderTarget();
+
+		// Enqueue a command to unlock the draw buffer after all windows have been drawn
+		ENQUEUE_UNIQUE_RENDER_COMMAND_THREEPARAMETER(UUMG3DComponentRenderToTexture,
+			FSlateDrawBuffer&, InDrawBuffer, DrawBuffer,
+			UTextureRenderTarget2D*, InRenderTarget, RenderTarget,
+			ISlate3DRenderer*, InRenderer, Renderer.Get(),
+			{
+				InRenderer->DrawWindowToTarget_RenderThread( RHICmdList, InRenderTarget, InDrawBuffer );	
+			});
+
+	}
+
+}
+
+class FUMGComponentInstanceData : public FComponentInstanceDataBase
+{
+public:
+	FUMGComponentInstanceData( const UUMGComponent* SourceComponent )
+		: FComponentInstanceDataBase(SourceComponent)
+		, WidgetClass ( SourceComponent->GetWidgetClass() )
+		, RenderTarget( SourceComponent->GetRenderTarget() )
+	{}
+
+	virtual bool MatchesComponent( const UActorComponent* Component ) const override
+	{
+		return (CastChecked<UUMGComponent>(Component)->GetWidgetClass() == WidgetClass && FComponentInstanceDataBase::MatchesComponent(Component));
+	}
+
+public:
+	TSubclassOf<UUserWidget> WidgetClass;
+	UTextureRenderTarget2D* RenderTarget;
+};
+
+FName UUMGComponent::GetComponentInstanceDataType() const
+{
+	static const FName InstanceDataName(TEXT("UMGInstanceData"));
+	return InstanceDataName;
+}
+
+TSharedPtr<FComponentInstanceDataBase> UUMGComponent::GetComponentInstanceData() const
+{
+	return MakeShareable( new FUMGComponentInstanceData( this ) );
+}
+
+void UUMGComponent::ApplyComponentInstanceData(TSharedPtr<FComponentInstanceDataBase> ComponentInstanceData)
+{
+	check(ComponentInstanceData.IsValid());
+
+	// Note: ApplyComponentInstanceData is called while the component is registered so the rendering thread is already using this component
+	// That means all component state that is modified here must be mirrored on the scene proxy, which will be recreated to receive the changes later due to MarkRenderStateDirty.
+	TSharedPtr<FUMGComponentInstanceData> UMGInstanceData = StaticCastSharedPtr<FUMGComponentInstanceData>(ComponentInstanceData);
+
+	RenderTarget = UMGInstanceData->RenderTarget;
+	MaterialInstance->SetTextureParameterValue("SlateUI", RenderTarget);
+
+	MarkRenderStateDirty();
+}
+
+#if WITH_EDITORONLY_DATA
+void UUMGComponent::PostEditChangeProperty(FPropertyChangedEvent& PropertyChangedEvent)
+{
+	UProperty* Property = PropertyChangedEvent.MemberProperty;
+
+	if( Property && PropertyChangedEvent.ChangeType != EPropertyChangeType::Interactive )
+	{
+		static FName DrawSizeName("DrawSize");
+
+		static FName WidgetClass("WidgetClass");
+
+		if( Property->GetFName() == DrawSizeName || Property->GetFName() == WidgetClass )
+		{
+			UpdateWidget();
+
+			if( Property->GetFName() == DrawSizeName )
+			{
+				UpdateBodySetup(true);
+			}
+
+			MarkRenderStateDirty();
+		}
+	}
+
+	Super::PostEditChangeProperty(PropertyChangedEvent);
+}
+#endif
+
+void UUMGComponent::UpdateWidget()
+{
+	if(WidgetClass && !Widget && GetWorld() )
+	{
+		Widget = CreateWidget<UUserWidget>( GetWorld(), WidgetClass );
+	}
+	else if( !WidgetClass && Widget )
+	{
+		Widget = nullptr;
+	}
+
+	if(!SlateWidget.IsValid())
+	{
+		SlateWidget = SNew(SVirtualWindow).Size(DrawSize);
+
+	}
+	
+	if( !HitTestGrid.IsValid() )
+	{
+		HitTestGrid = MakeShareable( new FHittestGrid );
+	}
+
+	{
+		SlateWidget->Resize( DrawSize );
+
+		if( Widget )
+		{
+			SlateWidget->SetContent( Widget->TakeWidget() );
+		}
+		else
+		{
+			SlateWidget->SetContent( SNullWidget::NullWidget );
+		}
+
+	}
+
+	if( Widget && !GetWorld()->IsGameWorld() )
+	{
+		// Prevent native ticking of editor component previews
+		Widget->SetIsDesignTime( true );
+	}
+}
+
+void UUMGComponent::UpdateRenderTarget()
+{
+	if(!RenderTarget && DrawSize != FIntPoint::ZeroValue)
+	{
+		RenderTarget = ConstructObject<UTextureRenderTarget2D>(UTextureRenderTarget2D::StaticClass(), this);
+		RenderTarget->InitCustomFormat(DrawSize.X, DrawSize.Y, PF_B8G8R8A8, false);
+		MaterialInstance->SetTextureParameterValue("SlateUI", RenderTarget);
+		MarkRenderStateDirty();
+	}
+	else if(DrawSize != FIntPoint::ZeroValue && (RenderTarget->SizeX != DrawSize.X || RenderTarget->SizeY != DrawSize.Y))
+	{
+		RenderTarget->InitCustomFormat(DrawSize.X, DrawSize.Y, PF_B8G8R8A8, false);
+		MarkRenderStateDirty();
+	}
+}
+
+void UUMGComponent::UpdateBodySetup( bool bDrawSizeChanged )
+{
+	if( !BodySetup || bDrawSizeChanged )
+	{
+		BodySetup = ConstructObject<UBodySetup>( UBodySetup::StaticClass(), this );
+		BodySetup->CollisionTraceFlag = CTF_UseSimpleAsComplex;
+		BodySetup->AggGeom.BoxElems.Add(FKBoxElem());
+
+		FKBoxElem* BoxElem = BodySetup->AggGeom.BoxElems.GetTypedData();
+
+		BoxElem->SetTransform(FTransform::Identity);
+		BoxElem->Center = FVector(DrawSize.X/2.0f, DrawSize.Y/2.0f, .5f);
+		BoxElem->X = DrawSize.X;
+		BoxElem->Y = DrawSize.Y;
+		BoxElem->Z = 1.0f;
+	}	
+}
+
+TArray<FArrangedWidget> UUMGComponent::GetHitWidgetPath( const FHitResult& HitResult, bool bIgnoreEnabledStatus  )
+{
+	FVector2D LocalHitLocation = FVector2D( ComponentToWorld.InverseTransformPosition( HitResult.Location ) );
+
+	TSharedRef<FVirtualCursorPosition> VirtualMouseCoordinate = MakeShareable( new FVirtualCursorPosition );
+
+	VirtualMouseCoordinate->CurrentCursorPosition = LocalHitLocation;
+	VirtualMouseCoordinate->LastCursorPosition = LastLocalHitLocation;
+
+	// Cache the location of the hit
+	LastLocalHitLocation = LocalHitLocation;
+
+	TArray<FArrangedWidget> ArrangedWidgets = HitTestGrid->GetBubblePath( LocalHitLocation, bIgnoreEnabledStatus );
+
+	for( FArrangedWidget& Widget : ArrangedWidgets )
+	{
+		Widget.VirtualCursorPosition = VirtualMouseCoordinate;
+	}
+
+	return ArrangedWidgets;
+}
+
+TSharedPtr<SWidget> UUMGComponent::GetSlateWidget() const
+{
+	return SlateWidget;
+}
\ No newline at end of file
diff --git a/Engine/Source/Runtime/UMG/Public/Components/UMGComponent.h b/Engine/Source/Runtime/UMG/Public/Components/UMGComponent.h
new file mode 100644
--- /dev/null
+++ b/Engine/Source/Runtime/UMG/Public/Components/UMGComponent.h
@@ -0,0 +1,114 @@
+// Copyright 1998-2014 Epic Games, Inc. All Rights Reserved.
+
+#pragma once
+
+#include "UMGComponent.generated.h"
+
+/**
+* UMG Component is a primitive component used to render User Widgets in the game world.
+* Functionally, it maps the user widget content to an SWindow that is in turn mapped to a quad (FKBoxElem) in the world via render target.
+*/
+UCLASS(hidecategories=(Object,Activation,"Components|Activation",Sockets,Base,Rendering,Lighting,LOD,Mesh), editinlinenew, meta=(BlueprintSpawnableComponent),MinimalAPI)
+class UUMGComponent : public UPrimitiveComponent
+{
+	GENERATED_UCLASS_BODY()
+
+public:
+	/* UPrimitiveComponent Interface */
+	virtual FPrimitiveSceneProxy* CreateSceneProxy() override;
+	virtual FBoxSphereBounds CalcBounds(const FTransform & LocalToWorld) const override;
+	virtual UBodySetup* GetBodySetup() override;
+	virtual FCollisionShape GetCollisionShape(float Inflation) const override;
+	virtual void OnRegister() override;
+	virtual void OnUnregister() override;
+	virtual void DestroyComponent() override;
+	virtual void TickComponent(float DeltaTime, enum ELevelTick TickType, FActorComponentTickFunction *ThisTickFunction) override;
+
+	virtual TSharedPtr<FComponentInstanceDataBase> GetComponentInstanceData() const override;
+	virtual FName GetComponentInstanceDataType() const override;
+	virtual void ApplyComponentInstanceData(TSharedPtr<class FComponentInstanceDataBase> ComponentInstanceData ) override;
+
+#if WITH_EDITORONLY_DATA
+	virtual void PostEditChangeProperty(FPropertyChangedEvent& PropertyChangedEvent);
+#endif
+
+	/** Ensures the user widget is initialized and updates its size and content. */
+	void UpdateWidget();
+
+	/** Ensure the render target is initialized and updates it if needed. */
+	void UpdateRenderTarget();
+
+	/** 
+	* Ensures the body setup is initialized and updates it if needed.
+	* @param bDrawSizeChanged Whether the draw size of this component has changed since the last update call.
+	*/
+	void UpdateBodySetup( bool bDrawSizeChanged = false );
+
+	/** @return The class of the user widget displayed by this component */
+	TSubclassOf<UUserWidget> GetWidgetClass() const { return WidgetClass; }
+
+	/** @return The user widget object displayed by this component */
+	UUserWidget* GetUMGWidgetObject() const { return Widget; }
+
+	/** @return  */
+	TArray<FArrangedWidget> GetHitWidgetPath( const FHitResult& HitResult, bool bIgnoreEnabledStatus );
+
+	/** @return The render target to which the user widget is rendered */
+	UTextureRenderTarget2D* GetRenderTarget() const { return RenderTarget; };
+
+	/** @return The dynamic material instance used to render the user widget */
+	UMaterialInstanceDynamic* GetMaterialInstance() const { return MaterialInstance; }
+
+	/** @return The window containing the user widget content */
+	TSharedPtr<SWidget> GetSlateWidget() const;
+
+	/** @return The draw size of the quad in the world */
+	const FIntPoint& GetDrawSize() const { return DrawSize; }
+
+private:
+	/** The size of the displayed quad */
+	UPROPERTY(EditAnywhere, Category=UI)
+	FIntPoint DrawSize;
+
+	/** The class of User Widget to create and display an instance of */
+	UPROPERTY(EditAnywhere, Category=UI)
+	TSubclassOf<UUserWidget> WidgetClass;
+
+	/** The User Widget object displayed and managed by this component */
+	UPROPERTY(transient, duplicatetransient)
+	UUserWidget* Widget;
+
+	/** The body setup of the displayed quad */
+	UPROPERTY(transient, duplicatetransient)
+	class UBodySetup* BodySetup;
+
+	/**  */
+	UPROPERTY()
+	UMaterial* BaseMaterial;
+
+	/** The target to which the user widget is rendered */
+	UPROPERTY(transient, duplicatetransient)
+	UTextureRenderTarget2D* RenderTarget;
+
+	/** The dynamic instance of the material that the render target is attached to */
+	UPROPERTY(transient, duplicatetransient)
+	UMaterialInstanceDynamic* MaterialInstance;
+
+	/**  */
+	TSharedPtr<class FHittestGrid> HitTestGrid;
+
+	/** The slate 3D renderer used to render the user slate widget */
+	TSharedPtr<class ISlate3DRenderer> Renderer;
+	
+	/** The slate window that contains the user widget content */
+	TSharedPtr<class SVirtualWindow> SlateWidget;
+
+	/** The relative location of the last hit on this component */
+	FVector2D LastLocalHitLocation;
+
+	/** The hit tester to use for this component */
+	static TSharedPtr<class FUMG3DHitTester> UMGHitTester;
+};
+
+ 
+
diff --git a/Engine/Source/Runtime/UMG/UMG.Build.cs b/Engine/Source/Runtime/UMG/UMG.Build.cs
--- a/Engine/Source/Runtime/UMG/UMG.Build.cs
+++ b/Engine/Source/Runtime/UMG/UMG.Build.cs
@@ -20,7 +20,10 @@ public class UMG : ModuleRules
                 "InputCore",
 				"Slate",
 				"SlateCore",
+				"RenderCore",
 				"MovieSceneCore",
+				"RHI",
+				"SlateRHIRenderer",
 			}
 		);
 	}
