UMG - Fixing a few warnings on the 3D component.

diff --git a/Engine/Source/Runtime/UMG/Private/Components/UMGComponent.cpp b/Engine/Source/Runtime/UMG/Private/Components/UMGComponent.cpp
--- a/Engine/Source/Runtime/UMG/Private/Components/UMGComponent.cpp
+++ b/Engine/Source/Runtime/UMG/Private/Components/UMGComponent.cpp
@@ -278,7 +278,7 @@ private:
 	TWeakObjectPtr<UWorld> World;
 };
 
-UUMGComponent::UUMGComponent(const class FPostConstructInitializeProperties& PCIP)
+UUMGComponent::UUMGComponent(const class FObjectInitializer& PCIP)
 	: Super(PCIP)
 {
 	PrimaryComponentTick.bCanEverTick = true;
@@ -591,7 +591,7 @@ void UUMGComponent::UpdateBodySetup( bool bDrawSizeChanged )
 		BodySetup->CollisionTraceFlag = CTF_UseSimpleAsComplex;
 		BodySetup->AggGeom.BoxElems.Add(FKBoxElem());
 
-		FKBoxElem* BoxElem = BodySetup->AggGeom.BoxElems.GetTypedData();
+		FKBoxElem* BoxElem = BodySetup->AggGeom.BoxElems.GetData();
 
 		BoxElem->SetTransform(FTransform::Identity);
 		BoxElem->Center = FVector(DrawSize.X/2.0f, DrawSize.Y/2.0f, .5f);
