Band-aid fix compilation of server targets.

Proper resolution pending (see JIRA UE-4505).

diff --git a/Engine/Source/Runtime/UMG/Private/Components/UMGComponent.cpp b/Engine/Source/Runtime/UMG/Private/Components/UMGComponent.cpp
--- a/Engine/Source/Runtime/UMG/Private/Components/UMGComponent.cpp
+++ b/Engine/Source/Runtime/UMG/Private/Components/UMGComponent.cpp
@@ -3,8 +3,10 @@
 #include "UMGPrivatePCH.h"
 #include "UMGComponent.h"
 #include "HittestGrid.h"
-#include "ISlateRHIRendererModule.h"
-#include "ISlate3DRenderer.h"
+#if !UE_SERVER
+	#include "ISlateRHIRendererModule.h"
+	#include "ISlate3DRenderer.h"
+#endif // !UE_SERVER
 #include "DynamicMeshBuilder.h"
 
 
@@ -358,10 +360,12 @@ void UUMGComponent::OnRegister()
 
 	UpdateWidget();
 
+#if !UE_SERVER
 	if( !Renderer.IsValid() )
 	{
 		Renderer = FModuleManager::Get().GetModuleChecked<ISlateRHIRendererModule>("SlateRHIRenderer").CreateSlate3DRenderer();
 	}
+#endif // !UE_SERVER
 }
 
 void UUMGComponent::OnUnregister()
@@ -404,6 +408,7 @@ void UUMGComponent::DestroyComponent()
 
 void UUMGComponent::TickComponent(float DeltaTime, enum ELevelTick TickType, FActorComponentTickFunction *ThisTickFunction)
 {
+#if !UE_SERVER
 	const float RenderTimeThreshold = .5f;
 	if( IsVisible() && GetWorld()->TimeSince(LastRenderTime) <= RenderTimeThreshold ) // Don't bother ticking if it hasn't been rendered recently
 	{
@@ -449,7 +454,7 @@ void UUMGComponent::TickComponent(float DeltaTime, enum ELevelTick TickType, FAc
 			});
 
 	}
-
+#endif // !UE_SERVER
 }
 
 class FUMGComponentInstanceData : public FComponentInstanceDataBase
diff --git a/Engine/Source/Runtime/UMG/UMG.Build.cs b/Engine/Source/Runtime/UMG/UMG.Build.cs
--- a/Engine/Source/Runtime/UMG/UMG.Build.cs
+++ b/Engine/Source/Runtime/UMG/UMG.Build.cs
@@ -23,8 +23,16 @@ public class UMG : ModuleRules
 				"RenderCore",
 				"MovieSceneCore",
 				"RHI",
-				"SlateRHIRenderer",
 			}
 		);
+
+		if (Target.Type != TargetRules.TargetType.Server)
+		{
+			PrivateDependencyModuleNames.AddRange(
+				new string[] {
+					"SlateRHIRenderer",
+				}
+			);
+		};
 	}
 }
