UMG - Restricting access to the slots array in PanelWidget since it's not legal to manipulate post init.  Changing a few functions that don't make sense to be exposed at runtime editor only.

diff --git a/Engine/Source/Editor/UMGEditor/Private/Extensions/HorizontalSlotExtension.cpp b/Engine/Source/Editor/UMGEditor/Private/Extensions/HorizontalSlotExtension.cpp
--- a/Engine/Source/Editor/UMGEditor/Private/Extensions/HorizontalSlotExtension.cpp
+++ b/Engine/Source/Editor/UMGEditor/Private/Extensions/HorizontalSlotExtension.cpp
@@ -76,12 +76,10 @@ FReply FHorizontalSlotExtension::HandleShift(int32 ShiftAmount)
 
 void FHorizontalSlotExtension::ShiftHorizontal(UWidget* Widget, int32 ShiftAmount)
 {
-	UHorizontalBoxSlot* Slot = CastChecked<UHorizontalBoxSlot>(Widget->Slot);
-	UHorizontalBox* Parent = CastChecked<UHorizontalBox>(Slot->Parent);
+	UHorizontalBox* Parent = CastChecked<UHorizontalBox>(Widget->GetParent());
 
 	int32 CurrentIndex = Parent->GetChildIndex(Widget);
-	Parent->Slots.RemoveAt(CurrentIndex);
-	Parent->Slots.Insert(Slot, FMath::Clamp(CurrentIndex + ShiftAmount, 0, Parent->GetChildrenCount()));
+	Parent->ShiftChild(CurrentIndex + ShiftAmount, Widget);
 }
 
 #undef LOCTEXT_NAMESPACE
diff --git a/Engine/Source/Editor/UMGEditor/Private/Extensions/VerticalSlotExtension.cpp b/Engine/Source/Editor/UMGEditor/Private/Extensions/VerticalSlotExtension.cpp
--- a/Engine/Source/Editor/UMGEditor/Private/Extensions/VerticalSlotExtension.cpp
+++ b/Engine/Source/Editor/UMGEditor/Private/Extensions/VerticalSlotExtension.cpp
@@ -76,12 +76,10 @@ FReply FVerticalSlotExtension::HandleShiftVertical(int32 ShiftAmount)
 
 void FVerticalSlotExtension::ShiftVertical(UWidget* Widget, int32 ShiftAmount)
 {
-	UVerticalBoxSlot* Slot = CastChecked<UVerticalBoxSlot>(Widget->Slot);
-	UVerticalBox* Parent = CastChecked<UVerticalBox>(Slot->Parent);
+	UVerticalBox* Parent = CastChecked<UVerticalBox>(Widget->GetParent());
 
 	int32 CurrentIndex = Parent->GetChildIndex(Widget);
-	Parent->Slots.RemoveAt(CurrentIndex);
-	Parent->Slots.Insert(Slot, FMath::Clamp(CurrentIndex + ShiftAmount, 0, Parent->GetChildrenCount()));
+	Parent->ShiftChild(CurrentIndex + ShiftAmount, Widget);
 }
 
 #undef LOCTEXT_NAMESPACE
diff --git a/Engine/Source/Runtime/UMG/Private/Components/PanelWidget.cpp b/Engine/Source/Runtime/UMG/Private/Components/PanelWidget.cpp
--- a/Engine/Source/Runtime/UMG/Private/Components/PanelWidget.cpp
+++ b/Engine/Source/Runtime/UMG/Private/Components/PanelWidget.cpp
@@ -135,26 +135,24 @@ bool UPanelWidget::ReplaceChildAt(int32 Index, UWidget* Content)
 	return true;
 }
 
-void UPanelWidget::InsertChildAt(int32 Index, UWidget* Content)
-{
-	UPanelSlot* Slot = ConstructObject<UPanelSlot>(GetSlotClass(), this);
-	Slot->SetFlags(RF_Transactional);
-	Slot->Content = Content;
-	Slot->Parent = this;
-
-	if ( Content )
-	{
-		Content->Slot = Slot;
-	}
-
-	// Only allow inserting within the valid range of slots (and one more than the size).
-	Index = FMath::Clamp(Index, 0, FMath::Max(Slots.Num(), 1));
+#if WITH_EDITOR
 
-	Slots.Insert(Slot, Index);
+UPanelSlot* UPanelWidget::InsertChildAt(int32 Index, UWidget* Content)
+{
+	UPanelSlot* NewSlot = AddChild(Content);
+	ShiftChild(Index, Content);
+	return NewSlot;
+}
 
-	OnSlotAdded(Slot);
+void UPanelWidget::ShiftChild(int32 Index, UWidget* Child)
+{
+	int32 CurrentIndex = GetChildIndex(Child);
+	Slots.RemoveAt(CurrentIndex);
+	Slots.Insert(Child->Slot, FMath::Clamp(Index, 0, Slots.Num()));
 }
 
+#endif
+
 bool UPanelWidget::RemoveChild(UWidget* Content)
 {
 	int32 ChildIndex = GetChildIndex(Content);
@@ -210,6 +208,11 @@ void UPanelWidget::PostLoad()
 	}
 }
 
+const TArray<UPanelSlot*>& UPanelWidget::GetSlots() const
+{
+	return Slots;
+}
+
 /////////////////////////////////////////////////////
 
 #undef LOCTEXT_NAMESPACE
diff --git a/Engine/Source/Runtime/UMG/Public/Components/PanelWidget.h b/Engine/Source/Runtime/UMG/Public/Components/PanelWidget.h
--- a/Engine/Source/Runtime/UMG/Public/Components/PanelWidget.h
+++ b/Engine/Source/Runtime/UMG/Public/Components/PanelWidget.h
@@ -11,9 +11,9 @@ class UMG_API UPanelWidget : public UWidget
 {
 	GENERATED_UCLASS_BODY()
 
-public:
+protected:
 
-	/** The items placed in the panel */
+	/** The slots in the widget holding the child widgets of this panel. */
 	UPROPERTY()
 	TArray<UPanelSlot*> Slots;
 
@@ -55,8 +55,20 @@ public:
 	 */
 	bool ReplaceChildAt(int32 Index, UWidget* Content);
 
-	/** Inserts a widget at a specific index */
-	void InsertChildAt(int32 Index, UWidget* Content);
+#if WITH_EDITOR
+
+	/**
+	 * Inserts a widget at a specific index.  This does not update the live slate version, it requires
+	 * a rebuild of the whole UI to see a change.
+	 */
+	UPanelSlot* InsertChildAt(int32 Index, UWidget* Content);
+
+	/**
+	 * Moves the child widget from its current index to the new index provided.
+	 */
+	void ShiftChild(int32 Index, UWidget* Child);
+
+#endif
 
 	/**
 	 * Removes a specific widget from the container.
@@ -73,6 +85,9 @@ public:
 	UFUNCTION(BlueprintCallable, Category="Widget|Panel")
 	void ClearChildren();
 
+	/** The slots in the widget holding the child widgets of this panel. */
+	const TArray<UPanelSlot*>& GetSlots() const;
+
 	/** @returns true if the panel supports more than one child. */
 	bool CanHaveMultipleChildren() const
 	{
