UMG - Simplifying the runtime binding hookup logic upon widget construction.

diff --git a/Engine/Source/Runtime/UMG/Private/WidgetBlueprintGeneratedClass.cpp b/Engine/Source/Runtime/UMG/Private/WidgetBlueprintGeneratedClass.cpp
--- a/Engine/Source/Runtime/UMG/Private/WidgetBlueprintGeneratedClass.cpp
+++ b/Engine/Source/Runtime/UMG/Private/WidgetBlueprintGeneratedClass.cpp
@@ -52,7 +52,6 @@ void UWidgetBlueprintGeneratedClass::Link(FArchive& Ar, bool bRelinkExistingProp
 void UWidgetBlueprintGeneratedClass::InitializeWidget(UUserWidget* UserWidget) const
 {
 	UWidgetTree* ClonedTree = DuplicateObject<UWidgetTree>( WidgetTree, UserWidget );
-	//check(ClonedTree);
 
 	if ( ClonedTree )
 	{
@@ -79,11 +78,12 @@ void UWidgetBlueprintGeneratedClass::InitializeWidget(UUserWidget* UserWidget) c
 
 		ClonedTree->ForEachWidget([&] (UWidget* Widget) {
 			// Not fatal if NULL, but shouldn't happen
-			if ( !ensure(Widget != NULL) )
+			if ( !ensure(Widget != nullptr) )
 			{
 				return;
 			}
 
+			// TODO UMG Make this an FName
 			FString VariableName = Widget->GetName();
 
 			Widget->bCreatedByConstructionScript = true; // Indicate it comes from a blueprint so it gets cleared when we rerun construction scripts
@@ -100,45 +100,21 @@ void UWidgetBlueprintGeneratedClass::InitializeWidget(UUserWidget* UserWidget) c
 			// Perform binding
 			for ( const FDelegateRuntimeBinding& Binding : Bindings )
 			{
-				//TODO UMG Terrible performance, improve with Maps.
+				//TODO UMG Make this faster.
 				if ( Binding.ObjectName == VariableName )
 				{
-					FString DelegateName = Binding.PropertyName.ToString() + "Delegate";
+					UDelegateProperty* DelegateProperty = FindField<UDelegateProperty>(Widget->GetClass(), FName(*( Binding.PropertyName.ToString() + TEXT("Delegate") )));
+					if ( !DelegateProperty )
+					{
+						DelegateProperty = FindField<UDelegateProperty>(Widget->GetClass(), Binding.PropertyName);
+					}
 
-					for ( TFieldIterator<UProperty> It(Widget->GetClass()); It; ++It )
+					if ( DelegateProperty )
 					{
-						if ( UDelegateProperty* DelegateProp = Cast<UDelegateProperty>(*It) )
+						FScriptDelegate* ScriptDelegate = DelegateProperty->GetPropertyValuePtr_InContainer(Widget);
+						if ( ScriptDelegate )
 						{
-							if ( DelegateProp->GetName() == DelegateName || DelegateProp->GetFName() == Binding.PropertyName )
-							{
-								FScriptDelegate* ScriptDelegate = DelegateProp->GetPropertyValuePtr_InContainer(Widget);
-
-								if ( Binding.Kind == EBindingKind::Function )
-								{
-									ScriptDelegate->BindUFunction(UserWidget, Binding.FunctionName);
-								}
-								else if ( Binding.Kind == EBindingKind::Property )
-								{
-									//FString FunctionNameStr = FString(TEXT("__Get")) + Binding.FunctionName.ToString();// +TEXT("_0");
-									//FName FunctionName = FName(*FunctionNameStr, 1);
-									UFunction* Fun = UserWidget->FindFunction(Binding.FunctionName);
-
-									TArray<FName> names;
-									for ( TFieldIterator<UFunction> It(WidgetBlueprintClass); It; ++It )
-									{
-										FName FunName = It->GetFName();
-										names.Add(FunName);
-									}
-
-									ScriptDelegate->BindUFunction(UserWidget, Binding.FunctionName);
-								}
-								else
-								{
-									check(false);
-								}
-								
-								break;
-							}
+							ScriptDelegate->BindUFunction(UserWidget, Binding.FunctionName);
 						}
 					}
 				}
@@ -152,7 +128,7 @@ void UWidgetBlueprintGeneratedClass::InitializeWidget(UUserWidget* UserWidget) c
 		// Bind any delegates on widgets
 		BindDynamicDelegates(UserWidget);
 
-		//TODO Add OnWidgetInitialized
+		//TODO UMG Add OnWidgetInitialized?
 	}
 }
 
