[Github PR 525] byte3 - UMG: Added Blueprint node "Slot as Overlay/Horizontal/Vertical Slot"

diff --git a/Engine/Source/Runtime/UMG/Private/WidgetLayoutLibrary.cpp b/Engine/Source/Runtime/UMG/Private/WidgetLayoutLibrary.cpp
--- a/Engine/Source/Runtime/UMG/Private/WidgetLayoutLibrary.cpp
+++ b/Engine/Source/Runtime/UMG/Private/WidgetLayoutLibrary.cpp
@@ -58,9 +58,24 @@ UGridSlot* UWidgetLayoutLibrary::SlotAsGridSlot(UWidget* ChildWidget)
 	return Cast<UGridSlot>(ChildWidget->Slot);
 }
 
+UHorizontalBoxSlot* UWidgetLayoutLibrary::SlotAsHorizontalBoxSlot(UWidget* ChildWidget)
+{
+	return Cast<UHorizontalBoxSlot>(ChildWidget->Slot);
+}
+
+UOverlaySlot* UWidgetLayoutLibrary::SlotAsOverlaySlot(UWidget* ChildWidget)
+{
+	return Cast<UOverlaySlot>(ChildWidget->Slot);
+}
+
 UUniformGridSlot* UWidgetLayoutLibrary::SlotAsUniformGridSlot(UWidget* ChildWidget)
 {
 	return Cast<UUniformGridSlot>(ChildWidget->Slot);
 }
 
+UVerticalBoxSlot* UWidgetLayoutLibrary::SlotAsVerticalBoxSlot(UWidget* ChildWidget)
+{
+	return Cast<UVerticalBoxSlot>(ChildWidget->Slot);
+}
+
 #undef LOCTEXT_NAMESPACE
\ No newline at end of file
diff --git a/Engine/Source/Runtime/UMG/Public/Blueprint/WidgetLayoutLibrary.h b/Engine/Source/Runtime/UMG/Public/Blueprint/WidgetLayoutLibrary.h
--- a/Engine/Source/Runtime/UMG/Public/Blueprint/WidgetLayoutLibrary.h
+++ b/Engine/Source/Runtime/UMG/Public/Blueprint/WidgetLayoutLibrary.h
@@ -31,10 +31,31 @@ class UWidgetLayoutLibrary : public UBlueprintFunctionLibrary
 	UFUNCTION(BlueprintPure, Category="Slot")
 	static UGridSlot* SlotAsGridSlot(UWidget* Widget);
 
+	/**
+	 * Gets the slot object on the child widget as a Horizontal Box Slot, allowing you to manipulate its information.
+	 * @param Widget The child widget of a Horizontal Box.
+	 */
+	UFUNCTION(BlueprintPure, Category = "Slot")
+	static UHorizontalBoxSlot* SlotAsHorizontalBoxSlot(UWidget* Widget);
+
+	/**
+	 * Gets the slot object on the child widget as a Overlay Slot, allowing you to manipulate layout information.
+	 * @param Widget The child widget of a overlay panel.
+	 */
+	UFUNCTION(BlueprintPure, Category = "Slot")
+	static UOverlaySlot* SlotAsOverlaySlot(UWidget* Widget);
+
 	/**
 	 * Gets the slot object on the child widget as a Uniform Grid Slot, allowing you to manipulate layout information.
 	 * @param Widget The child widget of a uniform grid panel.
 	 */
 	UFUNCTION(BlueprintPure, Category="Slot")
 	static UUniformGridSlot* SlotAsUniformGridSlot(UWidget* Widget);
+
+	/**
+	 * Gets the slot object on the child widget as a Vertical Box Slot, allowing you to manipulate its information.
+	 * @param Widget The child widget of a Vertical Box.
+	 */
+	UFUNCTION(BlueprintPure, Category = "Slot")
+	static UVerticalBoxSlot* SlotAsVerticalBoxSlot(UWidget* Widget);
 };
