- UUMGComponent renamed to UWidgetComponent
- Widget components support opaque/translucent, one/two-sided, and specifying a background color
- The background color is used as the render target clear color for the front of the UI and the solid color of the back

diff --git a/Engine/Source/Editor/UMGEditor/Private/UMGEditorModule.cpp b/Engine/Source/Editor/UMGEditor/Private/UMGEditorModule.cpp
--- a/Engine/Source/Editor/UMGEditor/Private/UMGEditorModule.cpp
+++ b/Engine/Source/Editor/UMGEditor/Private/UMGEditorModule.cpp
@@ -15,7 +15,7 @@
 #include "Animation/Sequencer2DTransformTrackEditor.h"
 #include "IUMGModule.h"
 #include "ComponentReregisterContext.h"
-#include "UMGComponent.h"
+#include "WidgetComponent.h"
 #include "DesignerCommands.h"
 #include "WidgetBlueprint.h"
 
@@ -77,7 +77,7 @@ public:
 	{
 		if ( UWidgetBlueprint* WidgetBlueprint = Cast<UWidgetBlueprint>(Blueprint) )
 		{
-			TComponentReregisterContext<UUMGComponent> ComponentReregisterContext;
+			TComponentReregisterContext<UWidgetComponent> ComponentReregisterContext;
 
 			FWidgetBlueprintCompiler Compiler(WidgetBlueprint, Results, CompileOptions, ObjLoaded);
 			Compiler.Compile();
diff --git a/Engine/Source/Runtime/SlateRHIRenderer/Private/Slate3DRenderer.cpp b/Engine/Source/Runtime/SlateRHIRenderer/Private/Slate3DRenderer.cpp
--- a/Engine/Source/Runtime/SlateRHIRenderer/Private/Slate3DRenderer.cpp
+++ b/Engine/Source/Runtime/SlateRHIRenderer/Private/Slate3DRenderer.cpp
@@ -97,7 +97,7 @@ void FSlate3DRenderer::DrawWindowToTarget_RenderThread( FRHICommandListImmediate
 	FTextureRenderTarget2DResource* RenderTargetResource = static_cast<FTextureRenderTarget2DResource*>( RenderTarget->GetRenderTargetResource() );
 	
 	SetRenderTarget( RHICmdList, RenderTargetResource->GetTextureRHI(), FTextureRHIParamRef() );
-	RHICmdList.Clear( true, FLinearColor(0,0,0,0), false, 0.0f, false, 0x00, FIntRect() );
+	RHICmdList.Clear( true, RenderTarget->ClearColor, false, 0.0f, false, 0x00, FIntRect() );
 
 	FMatrix ProjectionMatrix = FSlateRHIRenderer::CreateProjectionMatrix( RenderTarget->SizeX, RenderTarget->SizeY );
 	if (WindowElementList.GetRenderBatches().Num() > 0)
diff --git a/Engine/Source/Runtime/SlateRHIRenderer/Private/Slate3DRenderer.h b/Engine/Source/Runtime/SlateRHIRenderer/Private/Slate3DRenderer.h
--- a/Engine/Source/Runtime/SlateRHIRenderer/Private/Slate3DRenderer.h
+++ b/Engine/Source/Runtime/SlateRHIRenderer/Private/Slate3DRenderer.h
@@ -4,7 +4,6 @@
 
 #include "ISlate3DRenderer.h"
 
-
 class FSlate3DRenderer : public ISlate3DRenderer
 {
 public:
@@ -22,14 +21,16 @@ private:
 	/** Texture manager for accessing textures on the game thread */
 	TSharedRef<FSlateRHIResourceManager> ResourceManager;
 
+	/** The font cache for rendering text */
 	TSharedRef<FSlateFontCache> FontCache;
 
+	/** The rendering policy to use for drawing to the render target */
 	TSharedPtr<class FSlateRHIRenderingPolicy> RenderTargetPolicy;
 
-	/** Element batcher which renders draw elements */
+	/** Element batcher that renders draw elements */
 	TSharedPtr<FSlateElementBatcher> ElementBatcher;
 
-	/** The draw buffer which is currently free for use by the game thread */
+	/** The draw buffer that is currently free for use by the game thread */
 	uint8 FreeBufferIndex;
 
 };
\ No newline at end of file
diff --git a/Engine/Source/Runtime/SlateRHIRenderer/Public/Interfaces/ISlate3DRenderer.h b/Engine/Source/Runtime/SlateRHIRenderer/Public/Interfaces/ISlate3DRenderer.h
--- a/Engine/Source/Runtime/SlateRHIRenderer/Public/Interfaces/ISlate3DRenderer.h
+++ b/Engine/Source/Runtime/SlateRHIRenderer/Public/Interfaces/ISlate3DRenderer.h
@@ -10,19 +10,20 @@ public:
 	/** @return The free buffer for drawing */
 	virtual FSlateDrawBuffer& GetDrawBuffer() = 0;
 
-	/** 
-	* Batches the draw elements in the buffer to prepare it for rendering.
-	* Call in the game thread before sending to the render thread.
-	* @param DrawBuffer The draw buffer to prepare
-	*/
+ 	/** 
+	 * Batches the draw elements in the buffer to prepare it for rendering.
+	 * Call in the game thread before sending to the render thread.
+	 *
+	 * @param DrawBuffer The draw buffer to prepare
+	 */
 	virtual void DrawWindow_GameThread(FSlateDrawBuffer& DrawBuffer) = 0;
 
 	/** 
-	* Renders the batched draw elements of the draw buffer to the given render target.
-	* Call after preparing the draw buffer and render target on the game thread.
-	* 
-	* @param RenderTarget The render target to render the contents of the draw buffer to
-	* @param InDrawBuffer The draw buffer containing the batched elements to render
-	*/
+	 * Renders the batched draw elements of the draw buffer to the given render target.
+	 * Call after preparing the draw buffer and render target on the game thread.
+	 * 
+	 * @param RenderTarget The render target to render the contents of the draw buffer to
+	 * @param InDrawBuffer The draw buffer containing the batched elements to render
+	 */
 	virtual void DrawWindowToTarget_RenderThread( FRHICommandListImmediate& RHICmdList, UTextureRenderTarget2D* RenderTarget, FSlateDrawBuffer& InDrawBuffer ) = 0;
 };
diff --git a/Engine/Source/Runtime/UMG/Private/Components/UMGComponent.cpp b/Engine/Source/Runtime/UMG/Private/Components/UMGComponent.cpp
deleted file mode 100644
diff --git a/Engine/Source/Runtime/UMG/Private/Components/WidgetComponent.cpp b/Engine/Source/Runtime/UMG/Private/Components/WidgetComponent.cpp
new file mode 100644
--- /dev/null
+++ b/Engine/Source/Runtime/UMG/Private/Components/WidgetComponent.cpp
@@ -0,0 +1,708 @@
+// Copyright 1998-2014 Epic Games, Inc. All Rights Reserved.
+
+#include "UMGPrivatePCH.h"
+#include "WidgetComponent.h"
+#include "HittestGrid.h"
+#if !UE_SERVER
+	#include "ISlateRHIRendererModule.h"
+	#include "ISlate3DRenderer.h"
+#endif // !UE_SERVER
+#include "DynamicMeshBuilder.h"
+
+
+class SVirtualWindow : public SWindow
+{
+	SLATE_BEGIN_ARGS( SVirtualWindow )
+		: _Size( FVector2D(100,100) )
+	{}
+
+		SLATE_ARGUMENT( FVector2D, Size )
+	SLATE_END_ARGS()
+
+public:
+	void Construct( const FArguments& InArgs )
+	{
+		bIsPopupWindow = true;
+		SetCachedSize( InArgs._Size );
+		SetNativeWindow( MakeShareable( new FGenericWindow() ) );
+		SetContent( SNullWidget::NullWidget );
+	}
+};
+
+
+/** Represents a billboard sprite to the scene manager. */
+class FWidget3DSceneProxy : public FPrimitiveSceneProxy
+{
+public:
+	/** Initialization constructor. */
+	FWidget3DSceneProxy( UWidgetComponent* InComponent, ISlate3DRenderer& InRenderer )
+		: FPrimitiveSceneProxy( InComponent )
+		, Renderer( InRenderer )
+		, RenderTarget( InComponent->GetRenderTarget() )
+		, MaterialInstance( InComponent->GetMaterialInstance() )
+		, BodySetup( InComponent->GetBodySetup() )
+		, bIsOpaque( InComponent->IsOpaque() )
+	{
+		bWillEverBeLit = false;
+	}
+
+	~FWidget3DSceneProxy()
+	{
+	}
+
+	// FPrimitiveSceneProxy interface.
+	virtual void GetDynamicMeshElements(const TArray<const FSceneView*>& Views, const FSceneViewFamily& ViewFamily, uint32 VisibilityMap, FMeshElementCollector& Collector) const override
+	{
+		const FMatrix& LocalToWorld = GetLocalToWorld();
+	
+		FDynamicMeshBuilder MeshBuilder;
+
+		if( RenderTarget )
+		{
+			FTextureResource* TextureResource = RenderTarget->Resource;
+			if(TextureResource)
+			{
+				float U = 0.0f;
+				float V = 0.0f;
+				float UL = RenderTarget->SizeX;
+				float VL = RenderTarget->SizeY;
+
+				int32 VertexIndices[4];
+
+				VertexIndices[0] = MeshBuilder.AddVertex(FVector(U, V, 0), FVector2D(0, 0), FVector(1, 0, 0), FVector(0, 1, 0), FVector(0, 0, 1), FColor::White);
+				VertexIndices[1] = MeshBuilder.AddVertex(FVector(U, VL, 0), FVector2D(0, 1), FVector(1, 0, 0), FVector(0, 1, 0), FVector(0, 0, 1), FColor::White);
+				VertexIndices[2] = MeshBuilder.AddVertex(FVector(UL, VL, 0), FVector2D(1, 1), FVector(1, 0, 0), FVector(0, 1, 0), FVector(0, 0, 1), FColor::White);
+				VertexIndices[3] = MeshBuilder.AddVertex(FVector(UL, V, 0), FVector2D(1, 0), FVector(1, 0, 0), FVector(0, 1, 0), FVector(0, 0, 1), FColor::White);
+
+				MeshBuilder.AddTriangle(VertexIndices[0], VertexIndices[1], VertexIndices[2]);
+				MeshBuilder.AddTriangle(VertexIndices[0], VertexIndices[2], VertexIndices[3]);
+			}
+
+
+			for(int32 ViewIndex = 0; ViewIndex < Views.Num(); ViewIndex++)
+			{
+				if(VisibilityMap & (1 << ViewIndex))
+				{
+					FPrimitiveDrawInterface* PDI = Collector.GetPDI(ViewIndex);
+
+					MeshBuilder.Draw(PDI, LocalToWorld, MaterialInstance->GetRenderProxy(IsSelected()), SDPG_World);
+				}
+			}
+		}
+
+	}
+	/** 
+	 * Draw the scene proxy as a dynamic element
+	 *
+	 * @param	PDI - draw interface to render to
+	 * @param	View - current view
+	 */
+	virtual void DrawDynamicElements(FPrimitiveDrawInterface* PDI, const FSceneView* View) override
+	{
+		if( RenderTarget )
+		{
+			FTextureResource* TextureResource = RenderTarget->Resource;
+			if(TextureResource)
+			{
+				float U = 0.0f;
+				float V = 0.0f;
+				float UL = RenderTarget->SizeX;
+				float VL = RenderTarget->SizeY;
+
+				FDynamicMeshBuilder MeshBuilder;
+
+				int32 VertexIndices[4];
+
+				VertexIndices[0] = MeshBuilder.AddVertex(FVector(U, V, 0), FVector2D(0, 0), FVector(1, 0, 0), FVector(0, 1, 0), FVector(0, 0, 1), FColor::White);
+				VertexIndices[1] = MeshBuilder.AddVertex(FVector(U, VL, 0), FVector2D(0, 1), FVector(1, 0, 0), FVector(0, 1, 0), FVector(0, 0, 1), FColor::White);
+				VertexIndices[2] = MeshBuilder.AddVertex(FVector(UL, VL, 0), FVector2D(1, 1), FVector(1, 0, 0), FVector(0, 1, 0), FVector(0, 0, 1), FColor::White);
+				VertexIndices[3] = MeshBuilder.AddVertex(FVector(UL, V, 0), FVector2D(1, 0), FVector(1, 0, 0), FVector(0, 1, 0), FVector(0, 0, 1), FColor::White);
+
+				MeshBuilder.AddTriangle(VertexIndices[0], VertexIndices[1], VertexIndices[2]);
+				MeshBuilder.AddTriangle(VertexIndices[0], VertexIndices[2], VertexIndices[3]);
+
+				MeshBuilder.Draw(PDI, GetLocalToWorld(), MaterialInstance->GetRenderProxy(IsSelected()), SDPG_World);
+
+			}
+		}
+
+		//FColor CollisionColor = FColor(157,149,223,255);
+		//FTransform GeomTransform(GetLocalToWorld());
+		//BodySetup->AggGeom.DrawAggGeom(PDI, GeomTransform, GetSelectionColor(CollisionColor, false, IsHovered()), nullptr, false, false, true );
+
+#if !(UE_BUILD_SHIPPING || UE_BUILD_TEST)
+		RenderBounds(PDI, View->Family->EngineShowFlags, GetBounds(), IsSelected());
+#endif
+	}
+
+	virtual FPrimitiveViewRelevance GetViewRelevance(const FSceneView* View) override
+	{
+		bool bVisible = true;// View->Family->EngineShowFlags.BillboardSprites;
+
+		FPrimitiveViewRelevance Result;
+		Result.bDrawRelevance = IsShown(View) && bVisible;
+		Result.bOpaqueRelevance = bIsOpaque;
+		Result.bNormalTranslucencyRelevance = !bIsOpaque;
+		Result.bSeparateTranslucencyRelevance = !bIsOpaque;
+		Result.bDynamicRelevance = true;
+		Result.bShadowRelevance = IsShadowCast(View);;
+		Result.bEditorPrimitiveRelevance = false;
+		return Result;
+	}
+
+	virtual void GetLightRelevance(const FLightSceneProxy* LightSceneProxy, bool& bDynamic, bool& bRelevant, bool& bLightMapped, bool& bShadowMapped) const override
+	{
+		bDynamic = false;
+		bRelevant = false;
+		bLightMapped = false;
+		bShadowMapped = false;
+	}
+
+	virtual void OnTransformChanged() override
+	{
+		Origin = GetLocalToWorld().GetOrigin();
+	}
+
+	virtual uint32 GetMemoryFootprint(void) const override { return(sizeof(*this) + GetAllocatedSize()); }
+
+	uint32 GetAllocatedSize(void) const { return( FPrimitiveSceneProxy::GetAllocatedSize() ); }
+
+private:
+	FVector Origin;
+	ISlate3DRenderer& Renderer;
+	UTextureRenderTarget2D* RenderTarget;
+	UMaterialInstanceDynamic* MaterialInstance;
+	UBodySetup* BodySetup;
+	bool bIsOpaque;
+};
+
+/**
+* The hit tester used by all Widget Component objects.
+*/
+class FWidget3DHitTester : public ICustomHitTestPath
+{
+public:
+	FWidget3DHitTester( UWorld* InWorld )
+		: World( InWorld )
+	{}
+
+	// ICustomHitTestPath implementation
+	virtual TArray<FArrangedWidget> GetBubblePath(const FGeometry& InGeometry, FVector2D DesktopSpaceCoordinate, bool bIgnoreEnabledStatus) const override
+	{
+		if( World.IsValid() && ensure( World->IsGameWorld() ) )
+		{
+			ULocalPlayer* const TargetPlayer = GEngine->GetLocalPlayerFromControllerId(World.Get(), 0);
+			auto* PlayerController = TargetPlayer->PlayerController;
+
+			if (TargetPlayer && PlayerController)
+			{
+				FHitResult HitResult;
+				if (PlayerController->GetHitResultAtScreenPosition(InGeometry.AbsoluteToLocal(DesktopSpaceCoordinate), ECC_Visibility, true, HitResult))
+				{
+					UWidgetComponent* WidgetComponent = Cast<UWidgetComponent>(HitResult.Component.Get());
+					if ( WidgetComponent )
+					{
+						return WidgetComponent->GetHitWidgetPath( HitResult, bIgnoreEnabledStatus );
+					}
+				}
+			}
+		}
+		return TArray<FArrangedWidget>();
+	}
+
+	virtual void ArrangeChildren( FArrangedChildren& ArrangedChildren ) const override
+	{
+		for( TWeakObjectPtr<UWidgetComponent> Component : RegisteredComponents )
+		{
+			UWidgetComponent* WidgetComponent = Component.Get();
+			// Check if visible;
+			if ( WidgetComponent && WidgetComponent->GetSlateWidget().IsValid() )
+			{
+				FGeometry WidgetGeom;
+
+				ArrangedChildren.AddWidget( FArrangedWidget( WidgetComponent->GetSlateWidget().ToSharedRef(), WidgetGeom.MakeChild( WidgetComponent->GetDrawSize(), FSlateLayoutTransform() ) ) );
+			}
+		}
+	}
+
+	virtual TSharedPtr<struct FVirtualCursorPosition> TranslateMouseCoordinateFor3DChild( const TSharedRef<SWidget>& ChildWidget, const FGeometry& ViewportGeometry, const FVector2D& ScreenSpaceMouseCoordinate, const FVector2D& LastScreenSpaceMouseCoordinate ) const override
+	{
+		if(World.IsValid() && ensure(World->IsGameWorld()))
+		{
+			ULocalPlayer* const TargetPlayer = GEngine->GetLocalPlayerFromControllerId(World.Get(), 0);
+			if(TargetPlayer && TargetPlayer->PlayerController)
+			{
+				// Check for a hit against any widget components in the world
+				for(TWeakObjectPtr<UWidgetComponent> Component : RegisteredComponents)
+				{
+					UWidgetComponent* WidgetComponent = Component.Get();
+					// Check if visible;
+					if ( WidgetComponent && WidgetComponent->GetSlateWidget() == ChildWidget )
+					{
+						FVector2D LocalMouseCoordinate = ViewportGeometry.AbsoluteToLocal(ScreenSpaceMouseCoordinate);
+						FVector2D LocalLastMouseCoordinate = ViewportGeometry.AbsoluteToLocal( LastScreenSpaceMouseCoordinate );
+
+
+						FHitResult HitResult;
+						if(TargetPlayer->PlayerController->GetHitResultAtScreenPosition(LocalMouseCoordinate, ECC_Visibility, true, HitResult))
+						{
+							if ( WidgetComponent == HitResult.Component.Get() )
+							{
+								TSharedPtr<FVirtualCursorPosition> VirtualCursorPos = MakeShareable(new FVirtualCursorPosition);
+
+								FVector2D LocalHitLocation = FVector2D( WidgetComponent->ComponentToWorld.InverseTransformPosition( HitResult.Location ) );
+
+								VirtualCursorPos->CurrentCursorPosition = LocalHitLocation;
+								VirtualCursorPos->LastCursorPosition = LocalHitLocation;
+
+								return VirtualCursorPos;
+							}
+						}
+					}
+				}
+			}
+		}
+
+		return nullptr;
+	}
+	// End ICustomHitTestPath
+
+	void RegisterWidgetComponent( UWidgetComponent* InComponent )
+	{
+		RegisteredComponents.AddUnique( InComponent );
+	}
+
+	void UnregisterWidgetComponent( UWidgetComponent* InComponent )
+	{
+		RegisteredComponents.RemoveSingleSwap( InComponent );
+	}
+
+	uint32 GetNumRegisteredComponents() const { return RegisteredComponents.Num(); }
+	
+	UWorld* GetWorld() const { return World.Get(); }
+private:
+	TArray< TWeakObjectPtr<UWidgetComponent> > RegisteredComponents;
+	TWeakObjectPtr<UWorld> World;
+};
+
+UWidgetComponent::UWidgetComponent( const class FObjectInitializer& PCIP )
+	: Super( PCIP )
+	, DrawSize( FIntPoint( 500, 500 ) )
+	, BackgroundColor( FLinearColor::Transparent )
+	, bIsOpaque( false )
+	, bIsTwoSided( true )
+{
+	PrimaryComponentTick.bCanEverTick = true;
+	bTickInEditor = true;
+	
+	SetRelativeRotation( FRotator( 0.f, 0.f, 90.f ) );
+	
+	BodyInstance.SetCollisionProfileName(FName(TEXT("UI")));
+
+	// Translucent material instances
+	static ConstructorHelpers::FObjectFinder<UMaterialInterface> TranslucentMaterial_Finder( TEXT("/Engine/EngineMaterials/Widget3DPassThrough_Translucent") );
+	TranslucentMaterial = TranslucentMaterial_Finder.Object;
+
+	static ConstructorHelpers::FObjectFinder<UMaterialInterface> TranslucentMaterial_OneSided_Finder( TEXT( "/Engine/EngineMaterials/Widget3DPassThrough_Translucent_OneSided" ) );
+	TranslucentMaterial_OneSided = TranslucentMaterial_OneSided_Finder.Object;
+
+	// Opaque material instances
+	static ConstructorHelpers::FObjectFinder<UMaterialInterface> OpaqueMaterial_Finder( TEXT( "/Engine/EngineMaterials/Widget3DPassThrough_Opaque" ) );
+	OpaqueMaterial = OpaqueMaterial_Finder.Object;
+
+	static ConstructorHelpers::FObjectFinder<UMaterialInterface> OpaqueMaterial_OneSided_Finder( TEXT( "/Engine/EngineMaterials/Widget3DPassThrough_Opaque_OneSided" ) );
+	OpaqueMaterial_OneSided = OpaqueMaterial_OneSided_Finder.Object;
+
+	LastLocalHitLocation = FVector2D::ZeroVector;
+	//bGenerateOverlapEvents = false;
+	bUseEditorCompositing = false;
+}
+
+FPrimitiveSceneProxy* UWidgetComponent::CreateSceneProxy()
+{
+	return new FWidget3DSceneProxy(this,*Renderer);
+}
+
+FBoxSphereBounds UWidgetComponent::CalcBounds(const FTransform & LocalToWorld) const
+{
+	return FBoxSphereBounds( FVector( DrawSize.X/2.0f, DrawSize.Y/2.0f, .5f), FVector(DrawSize.X/2,DrawSize.Y/2,1.0f), DrawSize.Size()/2 ).TransformBy( LocalToWorld );
+}
+
+UBodySetup* UWidgetComponent::GetBodySetup() 
+{
+	UpdateBodySetup();
+	return BodySetup;
+}
+
+FCollisionShape UWidgetComponent::GetCollisionShape(float Inflation) const
+{
+	FVector	Extent = ( FVector( DrawSize.Y, DrawSize.Y, 1.0f ) * ComponentToWorld.GetScale3D() ) + Inflation;
+	if( Inflation < 0.0f )
+	{
+		// Don't shrink below zero size.
+		Extent = Extent.ComponentMax(FVector::ZeroVector);
+	}
+
+	return FCollisionShape::MakeBox(Extent);
+}
+
+void UWidgetComponent::OnRegister()
+{
+	Super::OnRegister();
+
+	if( GetWorld()->IsGameWorld() )
+	{
+		TSharedPtr<SViewport> GameViewportWidget = GEngine->GetGameViewportWidget();
+
+		if( GameViewportWidget.IsValid() )
+		{
+			TSharedPtr<ICustomHitTestPath> CustomHitTestPath = GameViewportWidget->GetCustomHitTestPath();
+			if( !CustomHitTestPath.IsValid() )
+			{
+				CustomHitTestPath = MakeShareable(new FWidget3DHitTester(GetWorld()));
+				GameViewportWidget->SetCustomHitTestPath( CustomHitTestPath );
+			}
+
+			TSharedPtr<FWidget3DHitTester> WidgetHitTester = StaticCastSharedPtr<FWidget3DHitTester>(CustomHitTestPath);
+			if( WidgetHitTester->GetWorld() == GetWorld() )
+			{
+				WidgetHitTester->RegisterWidgetComponent( this );
+			}
+		}
+	}
+
+	if( !MaterialInstance )
+	{
+		UMaterialInterface* Parent = nullptr;
+		if (bIsOpaque)
+		{
+			Parent = bIsTwoSided ? OpaqueMaterial : OpaqueMaterial_OneSided;
+		}
+		else
+		{
+			Parent = bIsTwoSided ? TranslucentMaterial : TranslucentMaterial_OneSided;
+		}
+
+		MaterialInstance = UMaterialInstanceDynamic::Create(Parent, this);
+	}
+
+	UpdateWidget();
+
+#if !UE_SERVER
+	if( !Renderer.IsValid() )
+	{
+		Renderer = FModuleManager::Get().GetModuleChecked<ISlateRHIRendererModule>("SlateRHIRenderer").CreateSlate3DRenderer();
+	}
+#endif // !UE_SERVER
+}
+
+void UWidgetComponent::OnUnregister()
+{
+	if( GetWorld()->IsGameWorld() )
+	{
+		TSharedPtr<SViewport> GameViewportWidget = GEngine->GetGameViewportWidget();
+		if( GameViewportWidget.IsValid() )
+		{
+			TSharedPtr<ICustomHitTestPath> CustomHitTestPath = GameViewportWidget->GetCustomHitTestPath();
+			if( CustomHitTestPath.IsValid() )
+			{
+				TSharedPtr<FWidget3DHitTester> WidgetHitTestPath = StaticCastSharedPtr<FWidget3DHitTester>(CustomHitTestPath);
+
+				WidgetHitTestPath->UnregisterWidgetComponent( this );
+
+				if ( WidgetHitTestPath->GetNumRegisteredComponents() == 0 )
+				{
+					GameViewportWidget->SetCustomHitTestPath( nullptr );
+				}
+			}
+		}
+	}
+
+	if( Widget )
+	{
+		Widget->MarkPendingKill();
+		Widget = NULL;
+	}
+
+	SlateWidget.Reset();
+	Super::OnUnregister();
+}
+
+void UWidgetComponent::DestroyComponent()
+{
+	Super::DestroyComponent();
+	Renderer.Reset();
+}
+
+void UWidgetComponent::TickComponent(float DeltaTime, enum ELevelTick TickType, FActorComponentTickFunction *ThisTickFunction)
+{
+#if !UE_SERVER
+	const float RenderTimeThreshold = .5f;
+	if( IsVisible() && GetWorld()->TimeSince(LastRenderTime) <= RenderTimeThreshold ) // Don't bother ticking if it hasn't been rendered recently
+	{
+		SlateWidget->SlatePrepass();
+
+		FGeometry WindowGeometry = FGeometry::MakeRoot( DrawSize, FSlateLayoutTransform() );
+
+		SlateWidget->TickWidgetsRecursively(WindowGeometry, FApp::GetCurrentTime(), DeltaTime);
+
+		// Ticking can cause geometry changes.  Recompute
+		SlateWidget->SlatePrepass();
+
+		// Get the free buffer & add our virtual window
+		FSlateDrawBuffer& DrawBuffer = Renderer->GetDrawBuffer();
+		FSlateWindowElementList& WindowElementList = DrawBuffer.AddWindowElementList(SlateWidget.ToSharedRef());
+
+		// Prepare the test grid 
+		HitTestGrid->BeginFrame(WindowGeometry.GetClippingRect());
+
+		// Paint the window
+		int32 MaxLayerId = 0;
+		{
+			MaxLayerId = SlateWidget->Paint(
+				FPaintArgs(SlateWidget.ToSharedRef(), *HitTestGrid, FVector2D::ZeroVector),
+				WindowGeometry, WindowGeometry.GetClippingRect(),
+				WindowElementList,
+				0,
+				FWidgetStyle(),
+				SlateWidget->IsEnabled());
+		}
+
+		Renderer->DrawWindow_GameThread(DrawBuffer);
+
+		UpdateRenderTarget();
+
+		// Enqueue a command to unlock the draw buffer after all windows have been drawn
+		ENQUEUE_UNIQUE_RENDER_COMMAND_THREEPARAMETER(UWidgetComponentRenderToTexture,
+			FSlateDrawBuffer&, InDrawBuffer, DrawBuffer,
+			UTextureRenderTarget2D*, InRenderTarget, RenderTarget,
+			ISlate3DRenderer*, InRenderer, Renderer.Get(),
+			{
+				InRenderer->DrawWindowToTarget_RenderThread( RHICmdList, InRenderTarget, InDrawBuffer );	
+			});
+
+	}
+#endif // !UE_SERVER
+}
+
+class FWidgetComponentInstanceData : public FComponentInstanceDataBase
+{
+public:
+	FWidgetComponentInstanceData( const UWidgetComponent* SourceComponent )
+		: FComponentInstanceDataBase(SourceComponent)
+		, WidgetClass ( SourceComponent->GetWidgetClass() )
+		, RenderTarget( SourceComponent->GetRenderTarget() )
+	{}
+
+	virtual bool MatchesComponent( const UActorComponent* Component ) const override
+	{
+		return (CastChecked<UWidgetComponent>(Component)->GetWidgetClass() == WidgetClass && FComponentInstanceDataBase::MatchesComponent(Component));
+	}
+
+public:
+	TSubclassOf<UUserWidget> WidgetClass;
+	UTextureRenderTarget2D* RenderTarget;
+};
+
+FName UWidgetComponent::GetComponentInstanceDataType() const
+{
+	static const FName InstanceDataName(TEXT("WidgetInstanceData"));
+	return InstanceDataName;
+}
+
+TSharedPtr<FComponentInstanceDataBase> UWidgetComponent::GetComponentInstanceData() const
+{
+	return MakeShareable( new FWidgetComponentInstanceData( this ) );
+}
+
+void UWidgetComponent::ApplyComponentInstanceData(TSharedPtr<FComponentInstanceDataBase> ComponentInstanceData)
+{
+	check(ComponentInstanceData.IsValid());
+
+	// Note: ApplyComponentInstanceData is called while the component is registered so the rendering thread is already using this component
+	// That means all component state that is modified here must be mirrored on the scene proxy, which will be recreated to receive the changes later due to MarkRenderStateDirty.
+	TSharedPtr<FWidgetComponentInstanceData> WidgetInstanceData = StaticCastSharedPtr<FWidgetComponentInstanceData>(ComponentInstanceData);
+
+	RenderTarget = WidgetInstanceData->RenderTarget;
+	MaterialInstance->SetTextureParameterValue("SlateUI", RenderTarget);
+
+	MarkRenderStateDirty();
+}
+
+#if WITH_EDITORONLY_DATA
+void UWidgetComponent::PostEditChangeProperty(FPropertyChangedEvent& PropertyChangedEvent)
+{
+	UProperty* Property = PropertyChangedEvent.MemberProperty;
+
+	if( Property && PropertyChangedEvent.ChangeType != EPropertyChangeType::Interactive )
+	{
+		static FName DrawSizeName("DrawSize");
+		static FName WidgetClassName("WidgetClass");
+		static FName IsOpaqueName("bIsOpaque");
+		static FName IsTwoSidedName("bIsTwoSided");
+		static FName BackgroundColorName("BackgroundColor");
+
+		auto PropertyName = Property->GetFName();
+
+		if( PropertyName == DrawSizeName || PropertyName == WidgetClassName )
+		{
+			UpdateWidget();
+
+			if( PropertyName == DrawSizeName )
+			{
+				UpdateBodySetup(true);
+			}
+
+			MarkRenderStateDirty();
+		}
+		else if ( PropertyName == IsOpaqueName  || PropertyName == BackgroundColorName || PropertyName == IsTwoSidedName)
+		{
+			MarkRenderStateDirty();
+		}
+	}
+
+	Super::PostEditChangeProperty(PropertyChangedEvent);
+}
+#endif
+
+void UWidgetComponent::UpdateWidget()
+{
+	if(WidgetClass && !Widget && GetWorld() )
+	{
+		Widget = CreateWidget<UUserWidget>( GetWorld(), WidgetClass );
+	}
+	else if( !WidgetClass && Widget )
+	{
+		Widget = nullptr;
+	}
+
+	if(!SlateWidget.IsValid())
+	{
+		SlateWidget = SNew(SVirtualWindow).Size(DrawSize);
+
+	}
+	
+	if( !HitTestGrid.IsValid() )
+	{
+		HitTestGrid = MakeShareable( new FHittestGrid );
+	}
+
+	{
+		SlateWidget->Resize( DrawSize );
+
+		if( Widget )
+		{
+			SlateWidget->SetContent( Widget->TakeWidget() );
+		}
+		else
+		{
+			SlateWidget->SetContent( SNullWidget::NullWidget );
+		}
+
+	}
+
+	if( Widget && !GetWorld()->IsGameWorld() )
+	{
+		// Prevent native ticking of editor component previews
+		Widget->SetIsDesignTime( true );
+	}
+}
+
+void UWidgetComponent::UpdateRenderTarget()
+{
+	bool bClearColorChanged = false;
+
+	if(!RenderTarget && DrawSize != FIntPoint::ZeroValue)
+	{
+		RenderTarget = ConstructObject<UTextureRenderTarget2D>(UTextureRenderTarget2D::StaticClass(), this);
+		RenderTarget->InitCustomFormat(DrawSize.X, DrawSize.Y, PF_B8G8R8A8, false);
+		RenderTarget->ClearColor = BackgroundColor;
+		if (bIsOpaque)
+		{
+			RenderTarget->ClearColor.A = 1.0f;
+		}
+		bClearColorChanged = true;
+		
+		MaterialInstance->SetTextureParameterValue( "SlateUI", RenderTarget );
+		MarkRenderStateDirty();
+	}
+	else if ( DrawSize != FIntPoint::ZeroValue )
+	{
+		// Update the format
+		if ( RenderTarget->SizeX != DrawSize.X || RenderTarget->SizeY != DrawSize.Y )
+		{
+			RenderTarget->InitCustomFormat( DrawSize.X, DrawSize.Y, PF_B8G8R8A8, false );
+		}
+
+		// Update the clear color
+		if (!bIsOpaque && RenderTarget->ClearColor != BackgroundColor)
+		{
+			RenderTarget->ClearColor = BackgroundColor;
+			bClearColorChanged = true;
+		}
+		else if ( bIsOpaque )
+		{
+			// If opaque, make sure the alpha channel is set to 1.0
+			FLinearColor ClearColor( BackgroundColor.R, BackgroundColor.G, BackgroundColor.B, 1.0f );
+			if ( RenderTarget->ClearColor != ClearColor )
+			{
+				RenderTarget->ClearColor = ClearColor;
+				bClearColorChanged = true;
+			}
+		}
+		
+		MarkRenderStateDirty();
+	}
+
+	// If the clear color of the render target changed, update the BackColor of the material to match
+	if ( bClearColorChanged )
+	{
+		MaterialInstance->SetVectorParameterValue( "BackColor", RenderTarget->ClearColor );
+	}
+}
+
+void UWidgetComponent::UpdateBodySetup( bool bDrawSizeChanged )
+{
+	if( !BodySetup || bDrawSizeChanged )
+	{
+		BodySetup = ConstructObject<UBodySetup>( UBodySetup::StaticClass(), this );
+		BodySetup->CollisionTraceFlag = CTF_UseSimpleAsComplex;
+		BodySetup->bDoubleSidedGeometry = false;
+		BodySetup->AggGeom.BoxElems.Add(FKBoxElem());
+
+		FKBoxElem* BoxElem = BodySetup->AggGeom.BoxElems.GetData();
+
+		BoxElem->SetTransform(FTransform::Identity);
+		BoxElem->Center = FVector(DrawSize.X/2.0f, DrawSize.Y/2.0f, .5f);
+		BoxElem->X = DrawSize.X;
+		BoxElem->Y = DrawSize.Y;
+		BoxElem->Z = 1.0f;
+	}	
+}
+
+TArray<FArrangedWidget> UWidgetComponent::GetHitWidgetPath( const FHitResult& HitResult, bool bIgnoreEnabledStatus  )
+{
+	FVector2D LocalHitLocation = FVector2D( ComponentToWorld.InverseTransformPosition( HitResult.Location ) );
+
+	TSharedRef<FVirtualCursorPosition> VirtualMouseCoordinate = MakeShareable( new FVirtualCursorPosition );
+
+	VirtualMouseCoordinate->CurrentCursorPosition = LocalHitLocation;
+	VirtualMouseCoordinate->LastCursorPosition = LastLocalHitLocation;
+
+	// Cache the location of the hit
+	LastLocalHitLocation = LocalHitLocation;
+
+	TArray<FArrangedWidget> ArrangedWidgets = HitTestGrid->GetBubblePath( LocalHitLocation, bIgnoreEnabledStatus );
+
+	for( FArrangedWidget& Widget : ArrangedWidgets )
+	{
+		Widget.VirtualCursorPosition = VirtualMouseCoordinate;
+	}
+
+	return ArrangedWidgets;
+}
+
+TSharedPtr<SWidget> UWidgetComponent::GetSlateWidget() const
+{
+	return SlateWidget;
+}
\ No newline at end of file
diff --git a/Engine/Source/Runtime/UMG/Public/Components/UMGComponent.h b/Engine/Source/Runtime/UMG/Public/Components/UMGComponent.h
deleted file mode 100644
diff --git a/Engine/Source/Runtime/UMG/Public/Components/WidgetComponent.h b/Engine/Source/Runtime/UMG/Public/Components/WidgetComponent.h
new file mode 100644
--- /dev/null
+++ b/Engine/Source/Runtime/UMG/Public/Components/WidgetComponent.h
@@ -0,0 +1,144 @@
+// Copyright 1998-2014 Epic Games, Inc. All Rights Reserved.
+
+#pragma once
+
+#include "WidgetComponent.generated.h"
+
+/**
+* Beware! This feature is experimental and may be substantially changed or removed in future releases.
+* A 3D instance of a Widget Blueprint that can be interacted with in the world.
+*/
+UCLASS(ClassGroup=Experimental, hidecategories=(Object,Activation,"Components|Activation",Sockets,Base,Lighting,LOD,Mesh), editinlinenew, meta=(BlueprintSpawnableComponent,  DevelopmentStatus=Experimental),MinimalAPI)
+class UWidgetComponent : public UPrimitiveComponent
+{
+	GENERATED_UCLASS_BODY()
+
+public:
+	/* UPrimitiveComponent Interface */
+	virtual FPrimitiveSceneProxy* CreateSceneProxy() override;
+	virtual FBoxSphereBounds CalcBounds(const FTransform & LocalToWorld) const override;
+	virtual UBodySetup* GetBodySetup() override;
+	virtual FCollisionShape GetCollisionShape(float Inflation) const override;
+	virtual void OnRegister() override;
+	virtual void OnUnregister() override;
+	virtual void DestroyComponent() override;
+	virtual void TickComponent(float DeltaTime, enum ELevelTick TickType, FActorComponentTickFunction *ThisTickFunction) override;
+
+	virtual TSharedPtr<FComponentInstanceDataBase> GetComponentInstanceData() const override;
+	virtual FName GetComponentInstanceDataType() const override;
+	virtual void ApplyComponentInstanceData(TSharedPtr<class FComponentInstanceDataBase> ComponentInstanceData ) override;
+
+#if WITH_EDITORONLY_DATA
+	virtual void PostEditChangeProperty(FPropertyChangedEvent& PropertyChangedEvent);
+#endif
+
+	/** Ensures the user widget is initialized and updates its size and content. */
+	void UpdateWidget();
+
+	/** Ensure the render target is initialized and updates it if needed. */
+	void UpdateRenderTarget();
+
+	/** 
+	* Ensures the body setup is initialized and updates it if needed.
+	* @param bDrawSizeChanged Whether the draw size of this component has changed since the last update call.
+	*/
+	void UpdateBodySetup( bool bDrawSizeChanged = false );
+
+	/** @return The class of the user widget displayed by this component */
+	TSubclassOf<UUserWidget> GetWidgetClass() const { return WidgetClass; }
+
+	/** @return The user widget object displayed by this component */
+	UUserWidget* GetUserWidgetObject() const { return Widget; }
+
+	/** @return  */
+	TArray<FArrangedWidget> GetHitWidgetPath( const FHitResult& HitResult, bool bIgnoreEnabledStatus );
+
+	/** @return The render target to which the user widget is rendered */
+	UTextureRenderTarget2D* GetRenderTarget() const { return RenderTarget; };
+
+	/** @return The dynamic material instance used to render the user widget */
+	UMaterialInstanceDynamic* GetMaterialInstance() const { return MaterialInstance; }
+
+	/** @return The window containing the user widget content */
+	TSharedPtr<SWidget> GetSlateWidget() const;
+
+	/** @return The draw size of the quad in the world */
+	const FIntPoint& GetDrawSize() const { return DrawSize; }
+
+	/** @return True if the component is opaque */
+	bool IsOpaque() const { return bIsOpaque; }
+	
+private:
+	/** The size of the displayed quad */
+	UPROPERTY(EditAnywhere, Category=UI)
+	FIntPoint DrawSize;
+
+	/** The class of User Widget to create and display an instance of */
+	UPROPERTY(EditAnywhere, Category=UI)
+	TSubclassOf<UUserWidget> WidgetClass;
+
+	/** The background color of the component */
+	UPROPERTY( EditAnywhere, Category = UI )
+	FLinearColor BackgroundColor;
+
+	/** 
+	 * Should the component be rendered opaque? 
+	 * This improves aliasing of the UI in the world.
+	 */
+	UPROPERTY(EditAnywhere, Category=UI)
+	bool bIsOpaque;
+
+	/** Is the component visible from behind? */
+	UPROPERTY( EditAnywhere, Category = UI )
+	bool bIsTwoSided;
+
+	/** The User Widget object displayed and managed by this component */
+	UPROPERTY(transient, duplicatetransient)
+	UUserWidget* Widget;
+
+	/** The body setup of the displayed quad */
+	UPROPERTY(transient, duplicatetransient)
+	class UBodySetup* BodySetup;
+
+	/** The material instance for translucent widget components */
+	UPROPERTY()
+	UMaterialInterface* TranslucentMaterial;
+
+	/** The material instance for translucent, one-sided widget components */
+	UPROPERTY()
+	UMaterialInterface* TranslucentMaterial_OneSided;
+
+	/** The material instance for opaque widget components */
+	UPROPERTY()
+	UMaterialInterface* OpaqueMaterial;
+
+	/** The material instance for opaque, one-sided widget components */
+	UPROPERTY()
+	UMaterialInterface* OpaqueMaterial_OneSided;
+
+	/** The target to which the user widget is rendered */
+	UPROPERTY(transient, duplicatetransient)
+	UTextureRenderTarget2D* RenderTarget;
+
+	/** The dynamic instance of the material that the render target is attached to */
+	UPROPERTY(transient, duplicatetransient)
+	UMaterialInstanceDynamic* MaterialInstance;
+
+	/**  */
+	TSharedPtr<class FHittestGrid> HitTestGrid;
+
+	/** The slate 3D renderer used to render the user slate widget */
+	TSharedPtr<class ISlate3DRenderer> Renderer;
+	
+	/** The slate window that contains the user widget content */
+	TSharedPtr<class SVirtualWindow> SlateWidget;
+
+	/** The relative location of the last hit on this component */
+	FVector2D LastLocalHitLocation;
+
+	/** The hit tester to use for this component */
+	static TSharedPtr<class FWidget3DHitTester> WidgetHitTester;
+};
+
+ 
+
