-Disable two sided by default on widget components since it is more expensive and not as useful.
-Tweaked some property categories
-Added some additional comments.

diff --git a/Engine/Source/Runtime/UMG/Private/Components/WidgetComponent.cpp b/Engine/Source/Runtime/UMG/Private/Components/WidgetComponent.cpp
--- a/Engine/Source/Runtime/UMG/Private/Components/WidgetComponent.cpp
+++ b/Engine/Source/Runtime/UMG/Private/Components/WidgetComponent.cpp
@@ -192,7 +192,7 @@ public:
 		if( World.IsValid() && ensure( World->IsGameWorld() ) )
 		{
 			ULocalPlayer* const TargetPlayer = GEngine->GetLocalPlayerFromControllerId(World.Get(), 0);
-			auto* PlayerController = TargetPlayer->PlayerController;
+			APlayerController* PlayerController = TargetPlayer->PlayerController;
 
 			if (TargetPlayer && PlayerController)
 			{
@@ -290,7 +290,7 @@ UWidgetComponent::UWidgetComponent( const class FObjectInitializer& PCIP )
 	, DrawSize( FIntPoint( 500, 500 ) )
 	, BackgroundColor( FLinearColor::Transparent )
 	, bIsOpaque( false )
-	, bIsTwoSided( true )
+	, bIsTwoSided( false )
 {
 	PrimaryComponentTick.bCanEverTick = true;
 	bTickInEditor = true;
diff --git a/Engine/Source/Runtime/UMG/Public/Components/WidgetComponent.h b/Engine/Source/Runtime/UMG/Public/Components/WidgetComponent.h
--- a/Engine/Source/Runtime/UMG/Public/Components/WidgetComponent.h
+++ b/Engine/Source/Runtime/UMG/Public/Components/WidgetComponent.h
@@ -50,7 +50,7 @@ public:
 	/** @return The user widget object displayed by this component */
 	UUserWidget* GetUserWidgetObject() const { return Widget; }
 
-	/** @return  */
+	/** @return The path of hit widgets that were found when the component was hit */
 	TArray<FArrangedWidget> GetHitWidgetPath( const FHitResult& HitResult, bool bIgnoreEnabledStatus );
 
 	/** @return The render target to which the user widget is rendered */
@@ -78,18 +78,18 @@ private:
 	TSubclassOf<UUserWidget> WidgetClass;
 
 	/** The background color of the component */
-	UPROPERTY( EditAnywhere, Category = UI )
+	UPROPERTY(EditAnywhere, Category=Rendering)
 	FLinearColor BackgroundColor;
 
 	/** 
 	 * Should the component be rendered opaque? 
 	 * This improves aliasing of the UI in the world.
 	 */
-	UPROPERTY(EditAnywhere, Category=UI)
+	UPROPERTY(EditAnywhere, Category=Rendering)
 	bool bIsOpaque;
 
 	/** Is the component visible from behind? */
-	UPROPERTY( EditAnywhere, Category = UI )
+	UPROPERTY(EditAnywhere, Category=Rendering)
 	bool bIsTwoSided;
 
 	/** The User Widget object displayed and managed by this component */
@@ -124,7 +124,7 @@ private:
 	UPROPERTY(transient, duplicatetransient)
 	UMaterialInstanceDynamic* MaterialInstance;
 
-	/**  */
+	/** The grid used to find actual hit actual widgets once input has been translated to the components local space */
 	TSharedPtr<class FHittestGrid> HitTestGrid;
 
 	/** The slate 3D renderer used to render the user slate widget */
