- Added MaxInteractionDistance property to widget component
- Added input checks for distance & orientation (to make sure input is to front of widget)

diff --git a/Engine/Source/Runtime/UMG/Private/Components/WidgetComponent.cpp b/Engine/Source/Runtime/UMG/Private/Components/WidgetComponent.cpp
--- a/Engine/Source/Runtime/UMG/Private/Components/WidgetComponent.cpp
+++ b/Engine/Source/Runtime/UMG/Private/Components/WidgetComponent.cpp
@@ -122,14 +122,14 @@ public:
 				MeshBuilder.AddTriangle(VertexIndices[0], VertexIndices[2], VertexIndices[3]);
 
 				MeshBuilder.Draw(PDI, GetLocalToWorld(), MaterialInstance->GetRenderProxy(IsSelected()), SDPG_World);
-
 			}
 		}
-
-		//FColor CollisionColor = FColor(157,149,223,255);
-		//FTransform GeomTransform(GetLocalToWorld());
-		//BodySetup->AggGeom.DrawAggGeom(PDI, GeomTransform, GetSelectionColor(CollisionColor, false, IsHovered()), nullptr, false, false, true );
-
+		
+		// Visualize the collision geometry
+		/*FColor CollisionColor = FColor(157,149,223,255);
+		FTransform GeomTransform(GetLocalToWorld());
+		BodySetup->AggGeom.DrawAggGeom(PDI, GeomTransform, GetSelectionColor(CollisionColor, false, IsHovered()), nullptr, false, false, true );*/
+		
 #if !(UE_BUILD_SHIPPING || UE_BUILD_TEST)
 		RenderBounds(PDI, View->Family->EngineShowFlags, GetBounds(), IsSelected());
 #endif
@@ -202,7 +202,16 @@ public:
 					UWidgetComponent* WidgetComponent = Cast<UWidgetComponent>(HitResult.Component.Get());
 					if ( WidgetComponent )
 					{
-						return WidgetComponent->GetHitWidgetPath( HitResult, bIgnoreEnabledStatus );
+						// Make sure the player is interacting with the front of the widget
+						// For widget components, the "front" faces the Z (or Up vector) direction
+						if ( FVector::DotProduct( WidgetComponent->GetUpVector(), HitResult.ImpactPoint - HitResult.TraceStart ) < 0.f )
+						{
+							// Make sure the player is close enough to the widget to interact with it
+							if ( FVector::DistSquared( HitResult.TraceStart, HitResult.ImpactPoint ) <= FMath::Square( WidgetComponent->GetMaxInteractionDistance() ) )
+							{
+								return WidgetComponent->GetHitWidgetPath( HitResult, bIgnoreEnabledStatus );
+							}
+						}
 					}
 				}
 			}
@@ -288,6 +297,7 @@ private:
 UWidgetComponent::UWidgetComponent( const class FObjectInitializer& PCIP )
 	: Super( PCIP )
 	, DrawSize( FIntPoint( 500, 500 ) )
+	, MaxInteractionDistance( 1000.f )
 	, BackgroundColor( FLinearColor::Transparent )
 	, bIsOpaque( false )
 	, bIsTwoSided( false )
@@ -667,7 +677,6 @@ void UWidgetComponent::UpdateBodySetup( bool bDrawSizeChanged )
 	{
 		BodySetup = ConstructObject<UBodySetup>( UBodySetup::StaticClass(), this );
 		BodySetup->CollisionTraceFlag = CTF_UseSimpleAsComplex;
-		BodySetup->bDoubleSidedGeometry = false;
 		BodySetup->AggGeom.BoxElems.Add(FKBoxElem());
 
 		FKBoxElem* BoxElem = BodySetup->AggGeom.BoxElems.GetData();
diff --git a/Engine/Source/Runtime/UMG/Public/Components/WidgetComponent.h b/Engine/Source/Runtime/UMG/Public/Components/WidgetComponent.h
--- a/Engine/Source/Runtime/UMG/Public/Components/WidgetComponent.h
+++ b/Engine/Source/Runtime/UMG/Public/Components/WidgetComponent.h
@@ -67,17 +67,24 @@ public:
 	/** @return The draw size of the quad in the world */
 	const FIntPoint& GetDrawSize() const { return DrawSize; }
 
+	/** @return The max distance from which a player can interact with this widget */
+	float GetMaxInteractionDistance() const { return MaxInteractionDistance; }
+
 	/** @return True if the component is opaque */
 	bool IsOpaque() const { return bIsOpaque; }
 	
 private:
-	/** The size of the displayed quad */
-	UPROPERTY(EditAnywhere, Category=UI)
-	FIntPoint DrawSize;
-
 	/** The class of User Widget to create and display an instance of */
 	UPROPERTY(EditAnywhere, Category=UI)
 	TSubclassOf<UUserWidget> WidgetClass;
+	
+	/** The size of the displayed quad */
+	UPROPERTY(EditAnywhere, Category=UI)
+	FIntPoint DrawSize;
+	
+	/** The maximum distance from which a player can interact with this widget */
+	UPROPERTY(EditAnywhere, Category=UI, meta=(ClampMin="0.0", UIMax="5000.0", ClampMax="100000.0"))
+	float MaxInteractionDistance;
 
 	/** The background color of the component */
 	UPROPERTY(EditAnywhere, Category=Rendering)
