Added ability to toggle text wrapping in a UMG multi-line text block.  Auto-wrap is enabled by default

diff --git a/Engine/Source/Runtime/UMG/Private/Components/MultiLineEditableText.cpp b/Engine/Source/Runtime/UMG/Private/Components/MultiLineEditableText.cpp
--- a/Engine/Source/Runtime/UMG/Private/Components/MultiLineEditableText.cpp
+++ b/Engine/Source/Runtime/UMG/Private/Components/MultiLineEditableText.cpp
@@ -13,6 +13,7 @@ UMultiLineEditableText::UMultiLineEditableText(const FObjectInitializer& ObjectI
 	SMultiLineEditableText::FArguments Defaults;
 	WidgetStyle = *Defaults._TextStyle;
 
+	bAutoWrapText = true;
 	// HACK Special font initialization hack since there are no font assets yet for slate.
 	Font = FSlateFontInfo(TEXT("Slate/Fonts/Roboto-Bold.ttf"), 12);
 }
@@ -37,6 +38,8 @@ TSharedRef<SWidget> UMultiLineEditableText::RebuildWidget()
 	.TextStyle(&WidgetStyle)
 	.Font(FSlateFontInfo(FontPath, Font.Size))
 	.Justification(Justification)
+	.WrapTextAt( WrapTextAt )
+	.AutoWrapText( bAutoWrapText )
 //	.MinDesiredWidth(MinimumDesiredWidth)
 //	.IsCaretMovedWhenGainFocus(IsCaretMovedWhenGainFocus)
 //	.SelectAllTextWhenFocused(SelectAllTextWhenFocused)
diff --git a/Engine/Source/Runtime/UMG/Private/Components/MultiLineEditableTextBox.cpp b/Engine/Source/Runtime/UMG/Private/Components/MultiLineEditableTextBox.cpp
--- a/Engine/Source/Runtime/UMG/Private/Components/MultiLineEditableTextBox.cpp
+++ b/Engine/Source/Runtime/UMG/Private/Components/MultiLineEditableTextBox.cpp
@@ -17,6 +17,8 @@ UMultiLineEditableTextBox::UMultiLineEditableTextBox(const FObjectInitializer& O
 	// HACK Special font initialization hack since there are no font assets yet for slate.
 	Font = FSlateFontInfo(TEXT("Slate/Fonts/Roboto-Bold.ttf"), 12);
 
+	bAutoWrapText = true;
+
 	SMultiLineEditableTextBox::FArguments Defaults;
 	WidgetStyle = *Defaults._Style;
 	TextStyle = *Defaults._TextStyle;
@@ -46,6 +48,8 @@ TSharedRef<SWidget> UMultiLineEditableTextBox::RebuildWidget()
 		.ForegroundColor(ForegroundColor)
 		.BackgroundColor(BackgroundColor)
 		.ReadOnlyForegroundColor(ReadOnlyForegroundColor)
+		.AutoWrapText( bAutoWrapText )
+		.WrapTextAt( WrapTextAt )
 //		.MinDesiredWidth(MinimumDesiredWidth)
 //		.Padding(Padding)
 //		.IsCaretMovedWhenGainFocus(IsCaretMovedWhenGainFocus)
diff --git a/Engine/Source/Runtime/UMG/Public/Components/MultiLineEditableText.h b/Engine/Source/Runtime/UMG/Public/Components/MultiLineEditableText.h
--- a/Engine/Source/Runtime/UMG/Public/Components/MultiLineEditableText.h
+++ b/Engine/Source/Runtime/UMG/Public/Components/MultiLineEditableText.h
@@ -32,6 +32,14 @@ public:
 	UPROPERTY(EditAnywhere, BlueprintReadWrite, Category=Content)
 	TEnumAsByte<ETextJustify::Type> Justification;
 
+	/** Whether to wrap text automatically based on the widget's computed horizontal space.*/
+	UPROPERTY(EditDefaultsOnly, Category=Content)
+	bool bAutoWrapText;
+
+	/** Whether text wraps onto a new line when it's length exceeds this width; if this value is zero or negative, no wrapping occurs. */
+	UPROPERTY(EditDefaultsOnly, AdvancedDisplay, Category=Content)
+	float WrapTextAt;
+
 	/** Font color and opacity (overrides Style) */
 	UPROPERTY(EditAnywhere, BlueprintReadWrite, Category=Appearance)
 	FSlateFontInfo Font;
diff --git a/Engine/Source/Runtime/UMG/Public/Components/MultiLineEditableTextBox.h b/Engine/Source/Runtime/UMG/Public/Components/MultiLineEditableTextBox.h
--- a/Engine/Source/Runtime/UMG/Public/Components/MultiLineEditableTextBox.h
+++ b/Engine/Source/Runtime/UMG/Public/Components/MultiLineEditableTextBox.h
@@ -40,6 +40,14 @@ public:
 	UPROPERTY(EditDefaultsOnly, BlueprintReadOnly, Category=Content)
 	TEnumAsByte<ETextJustify::Type> Justification;
 
+	/** Whether to wrap text automatically based on the widget's computed horizontal space.*/
+	UPROPERTY(EditDefaultsOnly, Category=Content)
+	bool bAutoWrapText;
+
+	/** Whether text wraps onto a new line when it's length exceeds this width; if this value is zero or negative, no wrapping occurs. */
+	UPROPERTY(EditDefaultsOnly, AdvancedDisplay, Category=Content)
+	float WrapTextAt;
+
 	/** Font color and opacity (overrides Style) */
 	UPROPERTY(EditDefaultsOnly, BlueprintReadOnly, Category=Appearance)
 	FSlateFontInfo Font;
