UMG now uses the Roboto font asset by default

diff --git a/Engine/Source/Runtime/UMG/Private/Components/EditableText.cpp b/Engine/Source/Runtime/UMG/Private/Components/EditableText.cpp
--- a/Engine/Source/Runtime/UMG/Private/Components/EditableText.cpp
+++ b/Engine/Source/Runtime/UMG/Private/Components/EditableText.cpp
@@ -15,8 +15,8 @@ UEditableText::UEditableText(const FObjectInitializer& ObjectInitializer)
 
 	ColorAndOpacity = FLinearColor::Black;
 
-	// HACK Special font initialization hack since there are no font assets yet for slate.
-	Font = FSlateFontInfo(TEXT("Slate/Fonts/Roboto-Bold.ttf"), 12);
+	static ConstructorHelpers::FObjectFinder<UFont> RobotoFontObj(TEXT("/Engine/EngineFonts/Roboto"));
+	Font = FSlateFontInfo(RobotoFontObj.Object, 12, FName("Bold"));
 
 	// Grab other defaults from slate arguments.
 	IsReadOnly = Defaults._IsReadOnly.Get();
diff --git a/Engine/Source/Runtime/UMG/Private/Components/EditableTextBox.cpp b/Engine/Source/Runtime/UMG/Private/Components/EditableTextBox.cpp
--- a/Engine/Source/Runtime/UMG/Private/Components/EditableTextBox.cpp
+++ b/Engine/Source/Runtime/UMG/Private/Components/EditableTextBox.cpp
@@ -14,8 +14,8 @@ UEditableTextBox::UEditableTextBox(const FObjectInitializer& ObjectInitializer)
 	BackgroundColor = FLinearColor::White;
 	ReadOnlyForegroundColor = FLinearColor::Black;
 
-	// HACK Special font initialization hack since there are no font assets yet for slate.
-	Font = FSlateFontInfo(TEXT("Slate/Fonts/Roboto-Bold.ttf"), 12);
+	static ConstructorHelpers::FObjectFinder<UFont> RobotoFontObj(TEXT("/Engine/EngineFonts/Roboto"));
+	Font = FSlateFontInfo(RobotoFontObj.Object, 12, FName("Bold"));
 
 	// Grab other defaults from slate arguments.
 	SEditableTextBox::FArguments Defaults;
diff --git a/Engine/Source/Runtime/UMG/Private/Components/MultiLineEditableText.cpp b/Engine/Source/Runtime/UMG/Private/Components/MultiLineEditableText.cpp
--- a/Engine/Source/Runtime/UMG/Private/Components/MultiLineEditableText.cpp
+++ b/Engine/Source/Runtime/UMG/Private/Components/MultiLineEditableText.cpp
@@ -14,8 +14,9 @@ UMultiLineEditableText::UMultiLineEditableText(const FObjectInitializer& ObjectI
 	WidgetStyle = *Defaults._TextStyle;
 
 	bAutoWrapText = true;
-	// HACK Special font initialization hack since there are no font assets yet for slate.
-	Font = FSlateFontInfo(TEXT("Slate/Fonts/Roboto-Bold.ttf"), 12);
+	
+	static ConstructorHelpers::FObjectFinder<UFont> RobotoFontObj(TEXT("/Engine/EngineFonts/Roboto"));
+	Font = FSlateFontInfo(RobotoFontObj.Object, 12, FName("Bold"));
 }
 
 void UMultiLineEditableText::ReleaseSlateResources(bool bReleaseChildren)
diff --git a/Engine/Source/Runtime/UMG/Private/Components/MultiLineEditableTextBox.cpp b/Engine/Source/Runtime/UMG/Private/Components/MultiLineEditableTextBox.cpp
--- a/Engine/Source/Runtime/UMG/Private/Components/MultiLineEditableTextBox.cpp
+++ b/Engine/Source/Runtime/UMG/Private/Components/MultiLineEditableTextBox.cpp
@@ -14,8 +14,8 @@ UMultiLineEditableTextBox::UMultiLineEditableTextBox(const FObjectInitializer& O
 	BackgroundColor = FLinearColor::White;
 	ReadOnlyForegroundColor = FLinearColor::Black;
 
-	// HACK Special font initialization hack since there are no font assets yet for slate.
-	Font = FSlateFontInfo(TEXT("Slate/Fonts/Roboto-Bold.ttf"), 12);
+	static ConstructorHelpers::FObjectFinder<UFont> RobotoFontObj(TEXT("/Engine/EngineFonts/Roboto"));
+	Font = FSlateFontInfo(RobotoFontObj.Object, 12, FName("Bold"));
 
 	bAutoWrapText = true;
 
diff --git a/Engine/Source/Runtime/UMG/Private/Components/SpinBox.cpp b/Engine/Source/Runtime/UMG/Private/Components/SpinBox.cpp
--- a/Engine/Source/Runtime/UMG/Private/Components/SpinBox.cpp
+++ b/Engine/Source/Runtime/UMG/Private/Components/SpinBox.cpp
@@ -10,8 +10,8 @@
 USpinBox::USpinBox(const FObjectInitializer& ObjectInitializer)
 	: Super(ObjectInitializer)
 {
-	// HACK Special font initialization hack since there are no font assets yet for slate.
-	Font = FSlateFontInfo(TEXT("Slate/Fonts/Roboto-Bold.ttf"), 12);
+	static ConstructorHelpers::FObjectFinder<UFont> RobotoFontObj(TEXT("/Engine/EngineFonts/Roboto"));
+	Font = FSlateFontInfo(RobotoFontObj.Object, 12, FName("Bold"));
 
 	// Grab other defaults from slate arguments.
 	SSpinBox<float>::FArguments Defaults;
diff --git a/Engine/Source/Runtime/UMG/Private/Components/TextBlock.cpp b/Engine/Source/Runtime/UMG/Private/Components/TextBlock.cpp
--- a/Engine/Source/Runtime/UMG/Private/Components/TextBlock.cpp
+++ b/Engine/Source/Runtime/UMG/Private/Components/TextBlock.cpp
@@ -20,8 +20,8 @@ UTextBlock::UTextBlock(const FObjectInitializer& ObjectInitializer)
 	ShadowColorAndOpacity = FLinearColor::Transparent;
 	LineHeightPercentage = 1.0f;
 
-	// @TODO UMG HACK Special font initialization hack since there are no font assets yet for slate.
-	Font = FSlateFontInfo(TEXT("Slate/Fonts/Roboto-Bold.ttf"), 24);
+	static ConstructorHelpers::FObjectFinder<UFont> RobotoFontObj(TEXT("/Engine/EngineFonts/Roboto"));
+	Font = FSlateFontInfo(RobotoFontObj.Object, 24, FName("Bold"));
 }
 
 void UTextBlock::ReleaseSlateResources(bool bReleaseChildren)
