UMG - Exposing a UFunction to get access to the underlying UserWidget that was factoried by the 3D Component.  Also removing several bits of functionality that don't need to run in a server build.

diff --git a/Engine/Source/Runtime/UMG/Private/Components/WidgetComponent.cpp b/Engine/Source/Runtime/UMG/Private/Components/WidgetComponent.cpp
--- a/Engine/Source/Runtime/UMG/Private/Components/WidgetComponent.cpp
+++ b/Engine/Source/Runtime/UMG/Private/Components/WidgetComponent.cpp
@@ -381,6 +381,8 @@ void UWidgetComponent::OnRegister()
 		}
 	}
 
+#if !UE_SERVER
+
 	if( !MaterialInstance )
 	{
 		UMaterialInterface* Parent = nullptr;
@@ -398,11 +400,11 @@ void UWidgetComponent::OnRegister()
 
 	UpdateWidget();
 
-#if !UE_SERVER
 	if( !Renderer.IsValid() )
 	{
 		Renderer = FModuleManager::Get().GetModuleChecked<ISlateRHIRendererModule>("SlateRHIRenderer").CreateSlate3DRenderer();
 	}
+
 #endif // !UE_SERVER
 }
 
@@ -447,6 +449,7 @@ void UWidgetComponent::DestroyComponent()
 void UWidgetComponent::TickComponent(float DeltaTime, enum ELevelTick TickType, FActorComponentTickFunction *ThisTickFunction)
 {
 #if !UE_SERVER
+
 	const float RenderTimeThreshold = .5f;
 	if( IsVisible() && GetWorld()->TimeSince(LastRenderTime) <= RenderTimeThreshold ) // Don't bother ticking if it hasn't been rendered recently
 	{
@@ -492,6 +495,7 @@ void UWidgetComponent::TickComponent(float DeltaTime, enum ELevelTick TickType,
 			});
 
 	}
+
 #endif // !UE_SERVER
 }
 
@@ -689,6 +693,11 @@ void UWidgetComponent::UpdateBodySetup( bool bDrawSizeChanged )
 	}	
 }
 
+UUserWidget* UWidgetComponent::GetUserWidgetObject() const
+{
+	return Widget;
+}
+
 TArray<FWidgetAndPointer> UWidgetComponent::GetHitWidgetPath( const FHitResult& HitResult, bool bIgnoreEnabledStatus  )
 {
 	FVector2D LocalHitLocation = FVector2D( ComponentToWorld.InverseTransformPosition( HitResult.Location ) );
@@ -714,4 +723,4 @@ TArray<FWidgetAndPointer> UWidgetComponent::GetHitWidgetPath( const FHitResult&
 TSharedPtr<SWidget> UWidgetComponent::GetSlateWidget() const
 {
 	return SlateWidget;
-}
\ No newline at end of file
+}
diff --git a/Engine/Source/Runtime/UMG/Public/Components/WidgetComponent.h b/Engine/Source/Runtime/UMG/Public/Components/WidgetComponent.h
--- a/Engine/Source/Runtime/UMG/Public/Components/WidgetComponent.h
+++ b/Engine/Source/Runtime/UMG/Public/Components/WidgetComponent.h
@@ -50,7 +50,8 @@ public:
 	TSubclassOf<UUserWidget> GetWidgetClass() const { return WidgetClass; }
 
 	/** @return The user widget object displayed by this component */
-	UUserWidget* GetUserWidgetObject() const { return Widget; }
+	UFUNCTION(BlueprintCallable, Category=UI)
+	UUserWidget* GetUserWidgetObject() const;
 
 	/** @return List of widgets with their geometry and the cursor position transformed into this Widget component's space. */
 	TArray<FWidgetAndPointer> GetHitWidgetPath( const FHitResult& HitResult, bool bIgnoreEnabledStatus );
