UMG - Removing the textblock's TextStyle, unfortunatly there's a lot of overlap with the individual properties, so much so that having the style isn't really useful, it was just confusing people.

diff --git a/Engine/Source/Runtime/UMG/Private/Components/TextBlock.cpp b/Engine/Source/Runtime/UMG/Private/Components/TextBlock.cpp
--- a/Engine/Source/Runtime/UMG/Private/Components/TextBlock.cpp
+++ b/Engine/Source/Runtime/UMG/Private/Components/TextBlock.cpp
@@ -12,7 +12,7 @@ UTextBlock::UTextBlock(const FObjectInitializer& ObjectInitializer)
 {
 	bIsVariable = false;
 
-	WidgetStyle = *UTextBlock::GetDefaultStyle();
+//	WidgetStyle = *UTextBlock::GetDefaultStyle();
 
 	Text = LOCTEXT("TextBlockDefaultValue", "Text Block");
 	ShadowOffset = FVector2D(1.0f, 1.0f);
@@ -60,8 +60,8 @@ void UTextBlock::SetShadowOffset(FVector2D InShadowOffset)
 
 TSharedRef<SWidget> UTextBlock::RebuildWidget()
 {
-	MyTextBlock = SNew(STextBlock)
-		.TextStyle(&WidgetStyle);
+	MyTextBlock = SNew(STextBlock);
+//		.TextStyle(&WidgetStyle);
 
 	return MyTextBlock.ToSharedRef();
 }
@@ -114,11 +114,11 @@ void UTextBlock::PostLoad()
 	{
 		if ( Style_DEPRECATED != nullptr )
 		{
-			const FTextBlockStyle* StylePtr = Style_DEPRECATED->GetStyle<FTextBlockStyle>();
-			if ( StylePtr != nullptr )
-			{
-				WidgetStyle = *StylePtr;
-			}
+			//const FTextBlockStyle* StylePtr = Style_DEPRECATED->GetStyle<FTextBlockStyle>();
+			//if ( StylePtr != nullptr )
+			//{
+			//	WidgetStyle = *StylePtr;
+			//}
 
 			Style_DEPRECATED = nullptr;
 		}
diff --git a/Engine/Source/Runtime/UMG/Public/Components/TextBlock.h b/Engine/Source/Runtime/UMG/Public/Components/TextBlock.h
--- a/Engine/Source/Runtime/UMG/Public/Components/TextBlock.h
+++ b/Engine/Source/Runtime/UMG/Public/Components/TextBlock.h
@@ -43,8 +43,8 @@ public:
 
 public:
 	/** The text block style */
-	UPROPERTY(EditAnywhere, BlueprintReadWrite, Category="Style", meta=( DisplayName="Style" ))
-	FTextBlockStyle WidgetStyle;
+	//UPROPERTY(EditAnywhere, BlueprintReadWrite, Category="Style", meta=( DisplayName="Style" ))
+	//FTextBlockStyle WidgetStyle;
 
 	UPROPERTY()
 	USlateWidgetStyleAsset* Style_DEPRECATED;
