UMG - ScrollBox now has scroll to end/start UFunctions.

diff --git a/Engine/Source/Runtime/UMG/Private/Components/ScrollBox.cpp b/Engine/Source/Runtime/UMG/Private/Components/ScrollBox.cpp
--- a/Engine/Source/Runtime/UMG/Private/Components/ScrollBox.cpp
+++ b/Engine/Source/Runtime/UMG/Private/Components/ScrollBox.cpp
@@ -96,6 +96,22 @@ void UScrollBox::SetScrollOffset(float NewScrollOffset)
 	}
 }
 
+void UScrollBox::ScrollToStart()
+{
+	if ( MyScrollBox.IsValid() )
+	{
+		MyScrollBox->ScrollToStart();
+	}
+}
+
+void UScrollBox::ScrollToEnd()
+{
+	if ( MyScrollBox.IsValid() )
+	{
+		MyScrollBox->ScrollToEnd();
+	}
+}
+
 void UScrollBox::PostLoad()
 {
 	Super::PostLoad();
diff --git a/Engine/Source/Runtime/UMG/Public/Components/ScrollBox.h b/Engine/Source/Runtime/UMG/Public/Components/ScrollBox.h
--- a/Engine/Source/Runtime/UMG/Public/Components/ScrollBox.h
+++ b/Engine/Source/Runtime/UMG/Public/Components/ScrollBox.h
@@ -51,10 +51,21 @@ public:
 
 public:
 
-	/** Updates the scroll offset of the scrollbox */
+	/**
+	 * Updates the scroll offset of the scrollbox.
+	 * @param NewScrollOffset is in Slate Units.
+	 */
 	UFUNCTION(BlueprintCallable, Category="Widget")
 	void SetScrollOffset(float NewScrollOffset);
 
+	/** Scrolls the ScrollBox to the top instantly */
+	UFUNCTION(BlueprintCallable, Category="Widget")
+	void ScrollToStart();
+
+	/** Scrolls the ScrollBox to the bottom instantly during the next layout pass. */
+	UFUNCTION(BlueprintCallable, Category="Widget")
+	void ScrollToEnd();
+
 	// UWidget interface
 	virtual void SynchronizeProperties() override;
 	// End of UWidget interface
