Fix for crash when calling SetInputMode API.

diff --git a/Engine/Source/Runtime/UMG/Private/WidgetBlueprintLibrary.cpp b/Engine/Source/Runtime/UMG/Private/WidgetBlueprintLibrary.cpp
--- a/Engine/Source/Runtime/UMG/Private/WidgetBlueprintLibrary.cpp
+++ b/Engine/Source/Runtime/UMG/Private/WidgetBlueprintLibrary.cpp
@@ -40,7 +40,7 @@ UDragDropOperation* UWidgetBlueprintLibrary::CreateDragDropOperation(TSubclassOf
 	}
 }
 
-void UWidgetBlueprintLibrary::SetInputMode_UIOnly(APlayerController*const& Target, UWidget* InWidgetToFocus, bool bLockMouseToViewport)
+void UWidgetBlueprintLibrary::SetInputMode_UIOnly(APlayerController* Target, UWidget* InWidgetToFocus, bool bLockMouseToViewport)
 {
 	if (Target != nullptr)
 	{
@@ -55,7 +55,7 @@ void UWidgetBlueprintLibrary::SetInputMode_UIOnly(APlayerController*const& Targe
 	}
 }
 
-void UWidgetBlueprintLibrary::SetInputMode_GameAndUI(APlayerController*const& Target, UWidget* InWidgetToFocus, bool bLockMouseToViewport, bool bHideCursorDuringCapture)
+void UWidgetBlueprintLibrary::SetInputMode_GameAndUI(APlayerController* Target, UWidget* InWidgetToFocus, bool bLockMouseToViewport, bool bHideCursorDuringCapture)
 {
 	if (Target != nullptr)
 	{
@@ -71,7 +71,7 @@ void UWidgetBlueprintLibrary::SetInputMode_GameAndUI(APlayerController*const& Ta
 	}
 }
 
-void UWidgetBlueprintLibrary::SetInputMode_GameOnly(APlayerController*const& Target)
+void UWidgetBlueprintLibrary::SetInputMode_GameOnly(APlayerController* Target)
 {
 	if (Target != nullptr)
 	{
diff --git a/Engine/Source/Runtime/UMG/Public/Blueprint/WidgetBlueprintLibrary.h b/Engine/Source/Runtime/UMG/Public/Blueprint/WidgetBlueprintLibrary.h
--- a/Engine/Source/Runtime/UMG/Public/Blueprint/WidgetBlueprintLibrary.h
+++ b/Engine/Source/Runtime/UMG/Public/Blueprint/WidgetBlueprintLibrary.h
@@ -28,15 +28,15 @@ class UWidgetBlueprintLibrary : public UBlueprintFunctionLibrary
 	
 	/** Setup an input mode that allows only the UI to respond to user input. */
 	UFUNCTION(BlueprintCallable, Category = "Input")
-	static void SetInputMode_UIOnly(APlayerController*const& Target, UWidget* InWidgetToFocus = nullptr, bool bLockMouseToViewport = false);
+	static void SetInputMode_UIOnly(APlayerController* Target, UWidget* InWidgetToFocus = nullptr, bool bLockMouseToViewport = false);
 
 	/** Setup an input mode that allows only the UI to respond to user input, and if the UI doesn't handle it player input / player controller gets a chance. */
 	UFUNCTION(BlueprintCallable, Category = "Input")
-	static void SetInputMode_GameAndUI(APlayerController*const& Target, UWidget* InWidgetToFocus = nullptr, bool bLockMouseToViewport = false, bool bHideCursorDuringCapture = true);
+	static void SetInputMode_GameAndUI(APlayerController* Target, UWidget* InWidgetToFocus = nullptr, bool bLockMouseToViewport = false, bool bHideCursorDuringCapture = true);
 
 	/** Setup an input mode that allows only player input / player controller to respond to user input. */
 	UFUNCTION(BlueprintCallable, Category = "Input")
-	static void SetInputMode_GameOnly(APlayerController*const& Target);
+	static void SetInputMode_GameOnly(APlayerController* Target);
 
 
 	UFUNCTION(BlueprintCallable, Category="Focus")
