UMG - Marking lots of functions as being cosmetic.

diff --git a/Engine/Source/Editor/UMGEditor/Private/Nodes/K2Node_CreateDragDropOperation.cpp b/Engine/Source/Editor/UMGEditor/Private/Nodes/K2Node_CreateDragDropOperation.cpp
--- a/Engine/Source/Editor/UMGEditor/Private/Nodes/K2Node_CreateDragDropOperation.cpp
+++ b/Engine/Source/Editor/UMGEditor/Private/Nodes/K2Node_CreateDragDropOperation.cpp
@@ -49,6 +49,11 @@ void UK2Node_CreateDragDropOperation::GetMenuEntries(FGraphContextMenuBuilder& C
 	NodeAction->NodeTemplate = TemplateNode;
 }
 
+FName UK2Node_CreateDragDropOperation::GetCornerIcon() const
+{
+	return TEXT("Graph.Replication.ClientEvent");
+}
+
 void UK2Node_CreateDragDropOperation::AllocateDefaultPins()
 {
 	Super::AllocateDefaultPins();
diff --git a/Engine/Source/Editor/UMGEditor/Private/Nodes/K2Node_CreateDragDropOperation.h b/Engine/Source/Editor/UMGEditor/Private/Nodes/K2Node_CreateDragDropOperation.h
--- a/Engine/Source/Editor/UMGEditor/Private/Nodes/K2Node_CreateDragDropOperation.h
+++ b/Engine/Source/Editor/UMGEditor/Private/Nodes/K2Node_CreateDragDropOperation.h
@@ -17,8 +17,9 @@ class UMGEDITOR_API UK2Node_CreateDragDropOperation : public UK2Node_ConstructOb
 	// End UEdGraphNode interface.
 
 	// Begin UK2Node interface
-	void GetMenuEntries(FGraphContextMenuBuilder& ContextMenuBuilder) const;
+	virtual void GetMenuEntries(FGraphContextMenuBuilder& ContextMenuBuilder) const override;
 	virtual FText GetMenuCategory() const override;
+	virtual FName GetCornerIcon() const override;
 	// End UK2Node interface.
 
 protected:
diff --git a/Engine/Source/Editor/UMGEditor/Private/Nodes/K2Node_CreateWidget.cpp b/Engine/Source/Editor/UMGEditor/Private/Nodes/K2Node_CreateWidget.cpp
--- a/Engine/Source/Editor/UMGEditor/Private/Nodes/K2Node_CreateWidget.cpp
+++ b/Engine/Source/Editor/UMGEditor/Private/Nodes/K2Node_CreateWidget.cpp
@@ -73,6 +73,11 @@ void UK2Node_CreateWidget::GetMenuEntries(FGraphContextMenuBuilder& ContextMenuB
 	NodeAction->NodeTemplate = TemplateNode;
 }
 
+FName UK2Node_CreateWidget::GetCornerIcon() const
+{
+	return TEXT("Graph.Replication.ClientEvent");
+}
+
 UEdGraphPin* UK2Node_CreateWidget::GetOwningPlayerPin() const
 {
 	UEdGraphPin* Pin = FindPin(FK2Node_CreateWidgetHelper::OwningPlayerPinName);
diff --git a/Engine/Source/Editor/UMGEditor/Private/Nodes/K2Node_CreateWidget.h b/Engine/Source/Editor/UMGEditor/Private/Nodes/K2Node_CreateWidget.h
--- a/Engine/Source/Editor/UMGEditor/Private/Nodes/K2Node_CreateWidget.h
+++ b/Engine/Source/Editor/UMGEditor/Private/Nodes/K2Node_CreateWidget.h
@@ -18,8 +18,9 @@ class UMGEDITOR_API UK2Node_CreateWidget : public UK2Node_ConstructObjectFromCla
 	// End UEdGraphNode interface.
 
 	// Begin UK2Node interface
-	void GetMenuEntries(FGraphContextMenuBuilder& ContextMenuBuilder) const;
+	virtual void GetMenuEntries(FGraphContextMenuBuilder& ContextMenuBuilder) const override;
 	virtual FText GetMenuCategory() const override;
+	virtual FName GetCornerIcon() const override;
 	// End UK2Node interface.
 
 	/** Get the owning player pin */
diff --git a/Engine/Source/Runtime/UMG/Public/Blueprint/SlateBlueprintLibrary.h b/Engine/Source/Runtime/UMG/Public/Blueprint/SlateBlueprintLibrary.h
--- a/Engine/Source/Runtime/UMG/Public/Blueprint/SlateBlueprintLibrary.h
+++ b/Engine/Source/Runtime/UMG/Public/Blueprint/SlateBlueprintLibrary.h
@@ -9,6 +9,8 @@ class USlateBlueprintLibrary : public UBlueprintFunctionLibrary
 {
 	GENERATED_UCLASS_BODY()
 
+public:
+
 	/**
 	 * Absolute coordinates could be either desktop or window space depending on what space the root of the widget hierarchy is in.
 	 *
diff --git a/Engine/Source/Runtime/UMG/Public/Blueprint/UserWidget.h b/Engine/Source/Runtime/UMG/Public/Blueprint/UserWidget.h
--- a/Engine/Source/Runtime/UMG/Public/Blueprint/UserWidget.h
+++ b/Engine/Source/Runtime/UMG/Public/Blueprint/UserWidget.h
@@ -144,13 +144,13 @@ public:
 	 * Adds it to the game's viewport, defaults to filling the entire viewport area.
 	 * @param bModal If this dialog should steal keyboard/mouse focus and consume all input. Great for a fullscreen menu. Terrible for HUDs.
 	 */
-	UFUNCTION(BlueprintCallable, Category="User Interface|Viewport")
+	UFUNCTION(BlueprintCallable, BlueprintCosmetic, Category="User Interface|Viewport")
 	void AddToViewport();
 
 	/**
 	 * Removes the widget from the viewport.
 	 */
-	UFUNCTION(BlueprintCallable, Category="User Interface|Viewport", meta=( DeprecatedFunction, DeprecationMessage="Use RemoveFromParent instead" ))
+	UFUNCTION(BlueprintCallable, BlueprintCosmetic, Category="User Interface|Viewport", meta=( DeprecatedFunction, DeprecationMessage="Use RemoveFromParent instead" ))
 	void RemoveFromViewport();
 
 	/** Removes the widget from it's parent widget, including the viewport if it was added to the viewport. */
@@ -159,27 +159,27 @@ public:
 	/**
 	 * Sets the widgets position in the viewport.
 	 */
-	UFUNCTION(BlueprintCallable, Category="User Interface|Viewport")
+	UFUNCTION(BlueprintCallable, BlueprintCosmetic, Category="User Interface|Viewport")
 	void SetPositionInViewport(FVector2D Position);
 
 	/*  */
-	UFUNCTION(BlueprintCallable, Category="User Interface|Viewport")
+	UFUNCTION(BlueprintCallable, BlueprintCosmetic, Category="User Interface|Viewport")
 	void SetDesiredSizeInViewport(FVector2D Size);
 
 	/*  */
-	UFUNCTION(BlueprintCallable, Category="User Interface|Viewport")
+	UFUNCTION(BlueprintCallable, BlueprintCosmetic, Category="User Interface|Viewport")
 	void SetAnchorsInViewport(FAnchors Anchors);
 
 	/*  */
-	UFUNCTION(BlueprintCallable, Category="User Interface|Viewport")
+	UFUNCTION(BlueprintCallable, BlueprintCosmetic, Category="User Interface|Viewport")
 	void SetAlignmentInViewport(FVector2D Alignment);
 
 	/*  */
-	UFUNCTION(BlueprintPure, Category="Appearance", meta=( DeprecatedFunction, DeprecationMessage="Use IsInViewport instead" ))
+	UFUNCTION(BlueprintPure, BlueprintCosmetic, Category="Appearance", meta=( DeprecatedFunction, DeprecationMessage="Use IsInViewport instead" ))
 	bool GetIsVisible() const;
 
 	/* @return true if the widget was added to the viewport using AddToViewport. */
-	UFUNCTION(BlueprintPure, Category="Appearance")
+	UFUNCTION(BlueprintPure, BlueprintCosmetic, Category="Appearance")
 	bool IsInViewport() const;
 
 	/** Sets the player context associated with this UI. */
@@ -192,28 +192,28 @@ public:
 	 * Gets the local player associated with this UI.
 	 * @return The owning local player.
 	 */
-	UFUNCTION(BlueprintCallable, Category="Player")
+	UFUNCTION(BlueprintCallable, BlueprintCosmetic, Category="Player")
 	class ULocalPlayer* GetOwningLocalPlayer() const;
 
 	/**
 	 * Gets the player controller associated with this UI.
 	 * @return The player controller that owns the UI.
 	 */
-	UFUNCTION(BlueprintCallable, Category="Player")
+	UFUNCTION(BlueprintCallable, BlueprintCosmetic, Category="Player")
 	class APlayerController* GetOwningPlayer() const;
 
 	/**
 	 * Gets the player pawn associated with this UI.
 	 * @return Gets the owning player pawn that's owned by the player controller assigned to this widget.
 	 */
-	UFUNCTION(BlueprintCallable, Category="Player")
+	UFUNCTION(BlueprintCallable, BlueprintCosmetic, Category="Player")
 	class APawn* GetOwningPlayerPawn() const;
 
 	/**
 	 * Called after the underlying slate widget is constructed.  Depending on how the slate object is used
 	 * this event may be called multiple times due to adding and removing from the hierarchy.
 	 */
-	UFUNCTION(BlueprintNativeEvent, Category="User Interface", meta=(Keywords="Begin Play"))
+	UFUNCTION(BlueprintNativeEvent, BlueprintCosmetic, Category="User Interface", meta=( Keywords="Begin Play" ))
 	void Construct();
 
 	/**
@@ -222,10 +222,10 @@ public:
 	 * @param  MyGeometry The space allotted for this widget
 	 * @param  InDeltaTime  Real time passed since last tick
 	 */
-	UFUNCTION(BlueprintNativeEvent, Category="User Interface")
+	UFUNCTION(BlueprintNativeEvent, BlueprintCosmetic, Category="User Interface")
 	void Tick(FGeometry MyGeometry, float InDeltaTime);
 
-	UFUNCTION(BlueprintNativeEvent, Category="User Interface | Painting")
+	UFUNCTION(BlueprintNativeEvent, BlueprintCosmetic, Category="User Interface | Painting")
 	void OnPaint(UPARAM(ref) FPaintContext& Context) const;
 
 	/**
@@ -235,10 +235,10 @@ public:
 	 * @param InFocusEvent  FocusEvent
 	 * @return  Returns whether the event was handled, along with other possible actions
 	 */
-	UFUNCTION(BlueprintNativeEvent, Category="Input")
+	UFUNCTION(BlueprintNativeEvent, BlueprintCosmetic, Category="Input")
 	FEventReply OnFocusReceived(FGeometry MyGeometry, FFocusEvent InFocusEvent);
 
-	UFUNCTION(BlueprintNativeEvent, meta = (DeprecatedFunction, DeprecationMessage = "Use OnFocusReceived() instead"), Category = "Keyboard")
+	UFUNCTION(BlueprintNativeEvent, BlueprintCosmetic, meta = (DeprecatedFunction, DeprecationMessage = "Use OnFocusReceived() instead"), Category = "Keyboard")
 	FEventReply OnKeyboardFocusReceived(FGeometry MyGeometry, FKeyboardFocusEvent InKeyboardFocusEvent);
 
 
@@ -247,10 +247,10 @@ public:
 	 *
 	 * @param  InFocusEvent  FocusEvent
 	 */
-	UFUNCTION(BlueprintNativeEvent, Category="Input")
+	UFUNCTION(BlueprintNativeEvent, BlueprintCosmetic, Category="Input")
 	void OnFocusLost(FFocusEvent InFocusEvent);
 
-	UFUNCTION(BlueprintNativeEvent, meta = (DeprecatedFunction, DeprecationMessage = "Use OnFocusLost() instead"), Category = "Keyboard")
+	UFUNCTION(BlueprintNativeEvent, BlueprintCosmetic, meta = (DeprecatedFunction, DeprecationMessage = "Use OnFocusLost() instead"), Category = "Keyboard")
 	void OnKeyboardFocusLost(FKeyboardFocusEvent InKeyboardFocusEvent);
 
 
@@ -261,7 +261,7 @@ public:
 	 * @param  InCharacterEvent  Character event
 	 * @return  Returns whether the event was handled, along with other possible actions
 	 */
-	UFUNCTION(BlueprintNativeEvent, Category="Input")
+	UFUNCTION(BlueprintNativeEvent, BlueprintCosmetic, Category="Input")
 	FEventReply OnKeyChar(FGeometry MyGeometry, FCharacterEvent InCharacterEvent);
 
 	/**
@@ -285,7 +285,7 @@ public:
 	 * @param  InKeyEvent  Key event
 	 * @return  Returns whether the event was handled, along with other possible actions
 	 */
-	UFUNCTION(BlueprintNativeEvent, Category="Input")
+	UFUNCTION(BlueprintNativeEvent, BlueprintCosmetic, Category="Input")
 	FEventReply OnKeyDown(FGeometry MyGeometry, FKeyEvent InKeyEvent);
 
 	/**
@@ -295,7 +295,7 @@ public:
 	 * @param  InKeyEvent  Key event
 	 * @return  Returns whether the event was handled, along with other possible actions
 	 */
-	UFUNCTION(BlueprintNativeEvent, Category="Input")
+	UFUNCTION(BlueprintNativeEvent, BlueprintCosmetic, Category="Input")
 	FEventReply OnKeyUp(FGeometry MyGeometry, FKeyEvent InKeyEvent);
 
 	/**
@@ -315,7 +315,7 @@ public:
 	 * @param MouseEvent Information about the input event
 	 * @return Whether the event was handled along with possible requests for the system to take action.
 	 */
-	UFUNCTION(BlueprintNativeEvent, Category="Mouse")
+	UFUNCTION(BlueprintNativeEvent, BlueprintCosmetic, Category="Mouse")
 	FEventReply OnMouseButtonDown(FGeometry MyGeometry, const FPointerEvent& MouseEvent);
 
 	/**
@@ -329,7 +329,7 @@ public:
 	 * @param MouseEvent Information about the input event
 	 * @return Whether the event was handled along with possible requests for the system to take action.
 	 */
-	UFUNCTION(BlueprintNativeEvent, Category="Mouse")
+	UFUNCTION(BlueprintNativeEvent, BlueprintCosmetic, Category="Mouse")
 	FEventReply OnPreviewMouseButtonDown(FGeometry MyGeometry, const FPointerEvent& MouseEvent);
 
 	/**
@@ -339,7 +339,7 @@ public:
 	 * @param MouseEvent Information about the input event
 	 * @return Whether the event was handled along with possible requests for the system to take action.
 	 */
-	UFUNCTION(BlueprintNativeEvent, Category="Mouse")
+	UFUNCTION(BlueprintNativeEvent, BlueprintCosmetic, Category="Mouse")
 	FEventReply OnMouseButtonUp(FGeometry MyGeometry, const FPointerEvent& MouseEvent);
 
 	/**
@@ -349,7 +349,7 @@ public:
 	 * @param MouseEvent Information about the input event
 	 * @return Whether the event was handled along with possible requests for the system to take action.
 	 */
-	UFUNCTION(BlueprintNativeEvent, Category="Mouse")
+	UFUNCTION(BlueprintNativeEvent, BlueprintCosmetic, Category="Mouse")
 	FEventReply OnMouseMove(FGeometry MyGeometry, const FPointerEvent& MouseEvent);
 
 	/**
@@ -358,7 +358,7 @@ public:
 	 * @param MyGeometry The Geometry of the widget receiving the event
 	 * @param MouseEvent Information about the input event
 	 */
-	UFUNCTION(BlueprintNativeEvent, Category="Mouse")
+	UFUNCTION(BlueprintNativeEvent, BlueprintCosmetic, Category="Mouse")
 	void OnMouseEnter(FGeometry MyGeometry, const FPointerEvent& MouseEvent);
 
 	/**
@@ -366,7 +366,7 @@ public:
 	 *
 	 * @param MouseEvent Information about the input event
 	 */
-	UFUNCTION(BlueprintNativeEvent, Category="Mouse")
+	UFUNCTION(BlueprintNativeEvent, BlueprintCosmetic, Category="Mouse")
 	void OnMouseLeave(const FPointerEvent& MouseEvent);
 
 	/**
@@ -375,7 +375,7 @@ public:
 	 * @param  MouseEvent  Mouse event
 	 * @return  Returns whether the event was handled, along with other possible actions
 	 */
-	UFUNCTION(BlueprintNativeEvent, Category="Mouse")
+	UFUNCTION(BlueprintNativeEvent, BlueprintCosmetic, Category="Mouse")
 	FEventReply OnMouseWheel(FGeometry MyGeometry, const FPointerEvent& MouseEvent);
 
 	/**
@@ -385,7 +385,7 @@ public:
 	 * @param  InMouseEvent  Mouse button event
 	 * @return  Returns whether the event was handled, along with other possible actions
 	 */
-	UFUNCTION(BlueprintNativeEvent, Category="Mouse")
+	UFUNCTION(BlueprintNativeEvent, BlueprintCosmetic, Category="Mouse")
 	FEventReply OnMouseButtonDoubleClick(FGeometry InMyGeometry, const FPointerEvent& InMouseEvent);
 
 	// TODO
@@ -402,7 +402,7 @@ public:
 	 * @param  PointerEvent  MouseMove that triggered the drag
 	 * @param  Operation     The drag operation that was detected.
 	 */
-	UFUNCTION(BlueprintNativeEvent, Category="Drag and Drop")
+	UFUNCTION(BlueprintNativeEvent, BlueprintCosmetic, Category="Drag and Drop")
 	void OnDragDetected(FGeometry MyGeometry, const FPointerEvent& PointerEvent, UDragDropOperation*& Operation);
 
 	/**
@@ -412,7 +412,7 @@ public:
 	 * @param  PointerEvent  Last mouse event from when the drag was canceled.
 	 * @param  Operation     The drag operation that was canceled.
 	 */
-	UFUNCTION(BlueprintNativeEvent, Category="Drag and Drop")
+	UFUNCTION(BlueprintNativeEvent, BlueprintCosmetic, Category="Drag and Drop")
 	void OnDragCancelled(const FPointerEvent& PointerEvent, UDragDropOperation* Operation);
 	
 	/**
@@ -422,7 +422,7 @@ public:
 	 * @param PointerEvent   The mouse event from when the drag entered the widget.
 	 * @param Operation      The drag operation that entered the widget.
 	 */
-	UFUNCTION(BlueprintNativeEvent, Category="Drag and Drop")
+	UFUNCTION(BlueprintNativeEvent, BlueprintCosmetic, Category="Drag and Drop")
 	void OnDragEnter(FGeometry MyGeometry, FPointerEvent PointerEvent, UDragDropOperation* Operation);
 
 	/**
@@ -431,7 +431,7 @@ public:
 	 * @param PointerEvent   The mouse event from when the drag left the widget.
 	 * @param Operation      The drag operation that entered the widget.
 	 */
-	UFUNCTION(BlueprintNativeEvent, Category="Drag and Drop")
+	UFUNCTION(BlueprintNativeEvent, BlueprintCosmetic, Category="Drag and Drop")
 	void OnDragLeave(FPointerEvent PointerEvent, UDragDropOperation* Operation);
 
 	/**
@@ -443,7 +443,7 @@ public:
 	 *
 	 * @return 'true' to indicate that you handled the drag over operation.  Returning 'false' will cause the operation to continue to bubble up.
 	 */
-	UFUNCTION(BlueprintNativeEvent, Category="Drag and Drop")
+	UFUNCTION(BlueprintNativeEvent, BlueprintCosmetic, Category="Drag and Drop")
 	bool OnDragOver(FGeometry MyGeometry, FPointerEvent PointerEvent, UDragDropOperation* Operation);
 
 	/**
@@ -455,7 +455,7 @@ public:
 	 * 
 	 * @return 'true' to indicate you handled the drop operation.
 	 */
-	UFUNCTION(BlueprintNativeEvent, Category="Drag and Drop")
+	UFUNCTION(BlueprintNativeEvent, BlueprintCosmetic, Category="Drag and Drop")
 	bool OnDrop(FGeometry MyGeometry, FPointerEvent PointerEvent, UDragDropOperation* Operation);
 
 
@@ -475,7 +475,7 @@ public:
 	 * @param  GestureEvent  gesture event
 	 * @return  Returns whether the event was handled, along with other possible actions
 	 */
-	UFUNCTION(BlueprintNativeEvent, Category="Touch Input")
+	UFUNCTION(BlueprintNativeEvent, BlueprintCosmetic, Category="Touch Input")
 	FEventReply OnTouchGesture(FGeometry MyGeometry, const FPointerEvent& GestureEvent);
 
 	/**
@@ -484,7 +484,7 @@ public:
 	 * @param MyGeometry    The geometry of the widget receiving the event.
 	 * @param InTouchEvent	The touch event generated
 	 */
-	UFUNCTION(BlueprintNativeEvent, Category="Touch Input")
+	UFUNCTION(BlueprintNativeEvent, BlueprintCosmetic, Category="Touch Input")
 	FEventReply OnTouchStarted(FGeometry MyGeometry, const FPointerEvent& InTouchEvent);
 
 	/**
@@ -493,7 +493,7 @@ public:
 	 * @param MyGeometry    The geometry of the widget receiving the event.
 	 * @param InTouchEvent	The touch event generated
 	 */
-	UFUNCTION(BlueprintNativeEvent, Category="Touch Input")
+	UFUNCTION(BlueprintNativeEvent, BlueprintCosmetic, Category="Touch Input")
 	FEventReply OnTouchMoved(FGeometry MyGeometry, const FPointerEvent& InTouchEvent);
 
 	/**
@@ -502,7 +502,7 @@ public:
 	 * @param MyGeometry    The geometry of the widget receiving the event.
 	 * @param InTouchEvent	The touch event generated
 	 */
-	UFUNCTION(BlueprintNativeEvent, Category="Touch Input")
+	UFUNCTION(BlueprintNativeEvent, BlueprintCosmetic, Category="Touch Input")
 	FEventReply OnTouchEnded(FGeometry MyGeometry, const FPointerEvent& InTouchEvent);
 
 	/**
@@ -512,7 +512,7 @@ public:
 	 * @param MyGeometry    The geometry of the widget receiving the event.
 	 * @param MotionEvent	The motion event generated
 	 */
-	UFUNCTION(BlueprintNativeEvent, Category="Touch Input")
+	UFUNCTION(BlueprintNativeEvent, BlueprintCosmetic, Category="Touch Input")
 	FEventReply OnMotionDetected(FGeometry MyGeometry, FMotionEvent InMotionEvent);
 
 public:
@@ -522,7 +522,7 @@ public:
 	 * 
 	 * @param InColorAndOpacity	The tint to apply to all child widgets.
 	 */
-	UFUNCTION(BlueprintCallable, Category="Appearance")
+	UFUNCTION(BlueprintCallable, BlueprintCosmetic, Category="Appearance")
 	void SetColorAndOpacity(FLinearColor InColorAndOpacity);
 
 	/**
@@ -531,7 +531,7 @@ public:
 	 * 
 	 * @param InForegroundColor	The foreground color.
 	 */
-	UFUNCTION(BlueprintCallable, Category="Appearance")
+	UFUNCTION(BlueprintCallable, BlueprintCosmetic, Category="Appearance")
 	void SetForegroundColor(FSlateColor InForegroundColor);
 
 	/**
@@ -539,7 +539,7 @@ public:
 	 * 
 	 * @param The name of the animation to play
 	 */
-	UFUNCTION(BlueprintCallable, Category="User Interface|Animation")
+	UFUNCTION(BlueprintCallable, BlueprintCosmetic, Category="User Interface|Animation")
 	void PlayAnimation(const UWidgetAnimation* InAnimation);
 
 	/**
@@ -547,7 +547,7 @@ public:
 	 * 
 	 * @param The name of the animation to stop
 	 */
-	UFUNCTION(BlueprintCallable, Category="User Interface|Animation")
+	UFUNCTION(BlueprintCallable, BlueprintCosmetic, Category="User Interface|Animation")
 	void StopAnimation(const UWidgetAnimation* InAnimation);
 
 	/** Called when a sequence player is finished playing an animation */
@@ -558,7 +558,7 @@ public:
 	*
 	* @param The sound to play
 	*/
-	UFUNCTION(BlueprintCallable, Category = "Sound")
+	UFUNCTION(BlueprintCallable, BlueprintCosmetic, Category = "Sound")
 	void PlaySound(class USoundBase* SoundToPlay);
 
 	/** @returns The UObject wrapper for a given SWidget */
diff --git a/Engine/Source/Runtime/UMG/Public/Blueprint/WidgetBlueprintLibrary.h b/Engine/Source/Runtime/UMG/Public/Blueprint/WidgetBlueprintLibrary.h
--- a/Engine/Source/Runtime/UMG/Public/Blueprint/WidgetBlueprintLibrary.h
+++ b/Engine/Source/Runtime/UMG/Public/Blueprint/WidgetBlueprintLibrary.h
@@ -15,8 +15,10 @@ class UWidgetBlueprintLibrary : public UBlueprintFunctionLibrary
 {
 	GENERATED_UCLASS_BODY()
 
+public:
+
 	/** Creates a widget */
-	UFUNCTION(BlueprintCallable, meta=( WorldContext="WorldContextObject", FriendlyName = "Create Widget", BlueprintInternalUseOnly = "true" ), Category="Widget")
+	UFUNCTION(BlueprintCallable, BlueprintCosmetic, meta=( WorldContext="WorldContextObject", FriendlyName="Create Widget", BlueprintInternalUseOnly="true" ), Category="Widget")
 	static class UUserWidget* Create(UObject* WorldContextObject, TSubclassOf<class UUserWidget> WidgetType, APlayerController* OwningPlayer);
 
 	/**
@@ -27,19 +29,19 @@ class UWidgetBlueprintLibrary : public UBlueprintFunctionLibrary
 	static UDragDropOperation* CreateDragDropOperation(TSubclassOf<UDragDropOperation> OperationClass);
 	
 	/** Setup an input mode that allows only the UI to respond to user input. */
-	UFUNCTION(BlueprintCallable, Category = "Input")
+	UFUNCTION(BlueprintCallable, BlueprintCosmetic, Category="Input")
 	static void SetInputMode_UIOnly(APlayerController* Target, UWidget* InWidgetToFocus = nullptr, bool bLockMouseToViewport = false);
 
 	/** Setup an input mode that allows only the UI to respond to user input, and if the UI doesn't handle it player input / player controller gets a chance. */
-	UFUNCTION(BlueprintCallable, Category = "Input")
+	UFUNCTION(BlueprintCallable, BlueprintCosmetic, Category="Input")
 	static void SetInputMode_GameAndUI(APlayerController* Target, UWidget* InWidgetToFocus = nullptr, bool bLockMouseToViewport = false, bool bHideCursorDuringCapture = true);
 
 	/** Setup an input mode that allows only player input / player controller to respond to user input. */
-	UFUNCTION(BlueprintCallable, Category = "Input")
+	UFUNCTION(BlueprintCallable, BlueprintCosmetic, Category="Input")
 	static void SetInputMode_GameOnly(APlayerController* Target);
 
-
-	UFUNCTION(BlueprintCallable, Category="Focus")
+	/** */
+	UFUNCTION(BlueprintCallable, BlueprintCosmetic, Category="Focus")
 	static void SetFocusToGameViewport();
 
 	/** Draws a box */
@@ -125,13 +127,13 @@ class UWidgetBlueprintLibrary : public UBlueprintFunctionLibrary
 	/**
 	 * Returns true if a drag/drop event is occurring that a widget can handle.
 	 */
-	UFUNCTION(BlueprintPure, Category="Widget|Drag and Drop")
+	UFUNCTION(BlueprintPure, BlueprintCosmetic, Category="Widget|Drag and Drop")
 	static bool IsDragDropping();
 
 	/**
 	 * Returns the drag and drop operation that is currently occuring if any, otherwise nothing.
 	 */
-	UFUNCTION(BlueprintPure, Category="Widget|Drag and Drop")
+	UFUNCTION(BlueprintPure, BlueprintCosmetic, Category="Widget|Drag and Drop")
 	static UDragDropOperation* GetDragDroppingContent();
 
 	/**
@@ -180,6 +182,6 @@ class UWidgetBlueprintLibrary : public UBlueprintFunctionLibrary
 	static UMaterialInstanceDynamic* GetDynamicMaterial(FSlateBrush& Brush);
 
 	/** Closes any popup menu */
-	UFUNCTION(BlueprintCallable, Category="Widget|Menu")
+	UFUNCTION(BlueprintCallable, BlueprintCosmetic, Category="Widget|Menu")
 	static void DismissAllMenus();
 };
diff --git a/Engine/Source/Runtime/UMG/Public/Blueprint/WidgetLayoutLibrary.h b/Engine/Source/Runtime/UMG/Public/Blueprint/WidgetLayoutLibrary.h
--- a/Engine/Source/Runtime/UMG/Public/Blueprint/WidgetLayoutLibrary.h
+++ b/Engine/Source/Runtime/UMG/Public/Blueprint/WidgetLayoutLibrary.h
@@ -9,6 +9,8 @@ class UWidgetLayoutLibrary : public UBlueprintFunctionLibrary
 {
 	GENERATED_UCLASS_BODY()
 
+public:
+
 	/**  */
 	UFUNCTION(BlueprintPure, Category="Viewport", meta=( HidePin="WorldContextObject", DefaultToSelf="WorldContextObject" ))
 	static float GetViewportScale(UObject* WorldContextObject);
@@ -35,14 +37,14 @@ class UWidgetLayoutLibrary : public UBlueprintFunctionLibrary
 	 * Gets the slot object on the child widget as a Horizontal Box Slot, allowing you to manipulate its information.
 	 * @param Widget The child widget of a Horizontal Box.
 	 */
-	UFUNCTION(BlueprintPure, Category = "Slot")
+	UFUNCTION(BlueprintPure, Category="Slot")
 	static UHorizontalBoxSlot* SlotAsHorizontalBoxSlot(UWidget* Widget);
 
 	/**
 	 * Gets the slot object on the child widget as a Overlay Slot, allowing you to manipulate layout information.
 	 * @param Widget The child widget of a overlay panel.
 	 */
-	UFUNCTION(BlueprintPure, Category = "Slot")
+	UFUNCTION(BlueprintPure, Category="Slot")
 	static UOverlaySlot* SlotAsOverlaySlot(UWidget* Widget);
 
 	/**
