[Github PR 576] TheCodez - Added function to WidgetTree to find a widget of a given type

diff --git a/Engine/Source/Runtime/UMG/Public/Blueprint/WidgetTree.h b/Engine/Source/Runtime/UMG/Public/Blueprint/WidgetTree.h
--- a/Engine/Source/Runtime/UMG/Public/Blueprint/WidgetTree.h
+++ b/Engine/Source/Runtime/UMG/Public/Blueprint/WidgetTree.h
@@ -20,6 +20,13 @@ public:
 	/** Finds a widget in the tree using the native widget as the key. */
 	UWidget* FindWidget(TSharedRef<SWidget> InWidget) const;
 
+	/** Finds the widget in the tree by name and casts the return to the desired type. */
+	template<class WidgetClass>
+	FORCEINLINE WidgetClass* FindWidget(const FName& Name) const
+	{
+		return Cast<WidgetClass>(FindWidget(Name));
+	}
+
 	/** Removes the widget from the hierarchy and all sub widgets. */
 	bool RemoveWidget(UWidget* Widget);
 
@@ -97,6 +104,7 @@ public:
 		}
 	}
 
+	// UObject interface
 	virtual void PreSave() override
 	{
 		AllWidgets.Empty();
@@ -105,13 +113,7 @@ public:
 
 		Super::PreSave();
 	}
-
-	virtual void PostLoad() override
-	{
-		//AllWidgets.Empty();
-
-		Super::PostLoad();
-	}
+	// End of UObject interface
 
 public:
 	/** The root widget of the tree */
