Fixed FWidget3DSceneProxy::GetDynamicMeshElements

diff --git a/Engine/Source/Runtime/UMG/Private/Components/WidgetComponent.cpp b/Engine/Source/Runtime/UMG/Private/Components/WidgetComponent.cpp
--- a/Engine/Source/Runtime/UMG/Private/Components/WidgetComponent.cpp
+++ b/Engine/Source/Runtime/UMG/Private/Components/WidgetComponent.cpp
@@ -69,27 +69,27 @@ public:
 
 				int32 VertexIndices[4];
 
-				VertexIndices[0] = MeshBuilder.AddVertex(FVector(U, V, 0), FVector2D(0, 0), FVector(1, 0, 0), FVector(0, 1, 0), FVector(0, 0, 1), FColor::White);
-				VertexIndices[1] = MeshBuilder.AddVertex(FVector(U, VL, 0), FVector2D(0, 1), FVector(1, 0, 0), FVector(0, 1, 0), FVector(0, 0, 1), FColor::White);
-				VertexIndices[2] = MeshBuilder.AddVertex(FVector(UL, VL, 0), FVector2D(1, 1), FVector(1, 0, 0), FVector(0, 1, 0), FVector(0, 0, 1), FColor::White);
-				VertexIndices[3] = MeshBuilder.AddVertex(FVector(UL, V, 0), FVector2D(1, 0), FVector(1, 0, 0), FVector(0, 1, 0), FVector(0, 0, 1), FColor::White);
-
-				MeshBuilder.AddTriangle(VertexIndices[0], VertexIndices[1], VertexIndices[2]);
-				MeshBuilder.AddTriangle(VertexIndices[0], VertexIndices[2], VertexIndices[3]);
-			}
-
-
-			for(int32 ViewIndex = 0; ViewIndex < Views.Num(); ViewIndex++)
-			{
-				if(VisibilityMap & (1 << ViewIndex))
+				for (int32 ViewIndex = 0; ViewIndex < Views.Num(); ViewIndex++)
 				{
-					FPrimitiveDrawInterface* PDI = Collector.GetPDI(ViewIndex);
+					if (VisibilityMap & (1 << ViewIndex))
+					{
+						VertexIndices[0] = MeshBuilder.AddVertex(FVector(U, V, 0), FVector2D(0, 0), FVector(1, 0, 0), FVector(0, 1, 0), FVector(0, 0, 1), FColor::White);
+						VertexIndices[1] = MeshBuilder.AddVertex(FVector(U, VL, 0), FVector2D(0, 1), FVector(1, 0, 0), FVector(0, 1, 0), FVector(0, 0, 1), FColor::White);
+						VertexIndices[2] = MeshBuilder.AddVertex(FVector(UL, VL, 0), FVector2D(1, 1), FVector(1, 0, 0), FVector(0, 1, 0), FVector(0, 0, 1), FColor::White);
+						VertexIndices[3] = MeshBuilder.AddVertex(FVector(UL, V, 0), FVector2D(1, 0), FVector(1, 0, 0), FVector(0, 1, 0), FVector(0, 0, 1), FColor::White);
 
-					MeshBuilder.Draw(PDI, LocalToWorld, MaterialInstance->GetRenderProxy(IsSelected()), SDPG_World);
+						MeshBuilder.AddTriangle(VertexIndices[0], VertexIndices[1], VertexIndices[2]);
+						MeshBuilder.AddTriangle(VertexIndices[0], VertexIndices[2], VertexIndices[3]);
+
+						MeshBuilder.GetMesh(LocalToWorld, MaterialInstance->GetRenderProxy(IsSelected()), SDPG_World, false, true, ViewIndex, Collector);
+					}
 				}
 			}
 		}
 
+#if !(UE_BUILD_SHIPPING || UE_BUILD_TEST)
+		RenderBounds(Collector.GetPDI(0), ViewFamily.EngineShowFlags, GetBounds(), IsSelected());
+#endif
 	}
 	/** 
 	 * Draw the scene proxy as a dynamic element
