UMG - Broke named slot restoring, fix.

diff --git a/Engine/Source/Runtime/UMG/Private/UserWidget.cpp b/Engine/Source/Runtime/UMG/Private/UserWidget.cpp
--- a/Engine/Source/Runtime/UMG/Private/UserWidget.cpp
+++ b/Engine/Source/Runtime/UMG/Private/UserWidget.cpp
@@ -62,8 +62,9 @@ void UUserWidget::Initialize()
 			}
 		});
 
+		// TODO Find a way to remove the bindings from the table.  Currently they're still needed.
 		// Clear the named slot bindings table.
-		NamedSlotBindings.Reset();
+		//NamedSlotBindings.Reset();
 	}
 }
 
