UMG - Adding more code doc to the WidgetTree.

diff --git a/Engine/Source/Runtime/UMG/Public/Blueprint/WidgetTree.h b/Engine/Source/Runtime/UMG/Public/Blueprint/WidgetTree.h
--- a/Engine/Source/Runtime/UMG/Public/Blueprint/WidgetTree.h
+++ b/Engine/Source/Runtime/UMG/Public/Blueprint/WidgetTree.h
@@ -39,7 +39,11 @@ public:
 	/** Gathers descendant child widgets of a parent widget. */
 	void GetChildWidgets(UWidget* Parent, TArray<UWidget*>& Widgets) const;
 
-	/**  */
+	/**
+	 * Iterates through all widgets including widgets contained in named slots, other than
+	 * investigating named slots, this code does not dive into foreign WidgetTrees, as would exist
+	 * inside another user widget.
+	 */
 	template <typename Predicate>
 	FORCEINLINE void ForEachWidget(Predicate Pred) const
 	{
@@ -51,10 +55,15 @@ public:
 		}
 	}
 
-	/**  */
+	/**
+	 * Iterates through all child widgets including widgets contained in named slots, other than
+	 * investigating named slots, this code does not dive into foreign WidgetTrees, as would exist
+	 * inside another user widget.
+	 */
 	template <typename Predicate>
 	FORCEINLINE void ForWidgetAndChildren(UWidget* Widget, Predicate Pred) const
 	{
+		// Search for any named slot with content that we need to dive into.
 		if ( INamedSlotInterface* NamedSlotHost = Cast<INamedSlotInterface>(Widget) )
 		{
 			TArray<FName> SlotNames;
@@ -71,6 +80,7 @@ public:
 			}
 		}
 
+		// Search standard children.
 		if ( UPanelWidget* PanelParent = Cast<UPanelWidget>(Widget) )
 		{
 			for ( int32 ChildIndex = 0; ChildIndex < PanelParent->GetChildrenCount(); ChildIndex++ )
