UMG - Adding an icon for the named slot.

diff --git a/Engine/Source/Runtime/UMG/Private/Components/NamedSlot.cpp b/Engine/Source/Runtime/UMG/Private/Components/NamedSlot.cpp
--- a/Engine/Source/Runtime/UMG/Private/Components/NamedSlot.cpp
+++ b/Engine/Source/Runtime/UMG/Private/Components/NamedSlot.cpp
@@ -90,7 +90,7 @@ void UNamedSlot::OnSlotRemoved(UPanelSlot* Slot)
 
 const FSlateBrush* UNamedSlot::GetEditorIcon()
 {
-	return FUMGStyle::Get().GetBrush("Widget.Spacer");
+	return FUMGStyle::Get().GetBrush("Widget.NamedSlot");
 }
 
 const FText UNamedSlot::GetPaletteCategory()
diff --git a/Engine/Source/Runtime/UMG/Private/UMGStyle.cpp b/Engine/Source/Runtime/UMG/Private/UMGStyle.cpp
--- a/Engine/Source/Runtime/UMG/Private/UMGStyle.cpp
+++ b/Engine/Source/Runtime/UMG/Private/UMGStyle.cpp
@@ -80,6 +80,7 @@ TSharedRef< FSlateStyleSet > FUMGStyle::Create()
 	Style->Set("Widget.SpinBox", new IMAGE_BRUSH(TEXT("SpinBox"), Icon16x16));
 	Style->Set("Widget.Grid", new IMAGE_BRUSH(TEXT("Grid"), Icon16x16));
 	Style->Set("Widget.WrapBox", new IMAGE_BRUSH(TEXT("WrapBox"), Icon16x16));
+	Style->Set("Widget.NamedSlot", new IMAGE_BRUSH(TEXT("NamedSlot"), Icon16x16));
 
 	Style->Set("Widget.UserWidget", new IMAGE_BRUSH(TEXT("UserWidget"), Icon16x16));
 
