Fix crash calling UWidgetComponent::UpdateWidget during cooking

diff --git a/Engine/Source/Runtime/UMG/Private/Components/WidgetComponent.cpp b/Engine/Source/Runtime/UMG/Private/Components/WidgetComponent.cpp
--- a/Engine/Source/Runtime/UMG/Private/Components/WidgetComponent.cpp
+++ b/Engine/Source/Runtime/UMG/Private/Components/WidgetComponent.cpp
@@ -581,44 +581,48 @@ void UWidgetComponent::PostEditChangeProperty(FPropertyChangedEvent& PropertyCha
 
 void UWidgetComponent::UpdateWidget()
 {
-	if(WidgetClass && !Widget && GetWorld() )
-	{
-		Widget = CreateWidget<UUserWidget>( GetWorld(), WidgetClass );
-	}
-	else if( !WidgetClass && Widget )
-	{
-		Widget = nullptr;
-	}
+	// Don't do any work if Slate is not initialized
+	if(FSlateApplication::IsInitialized())
+	{		
+		if(WidgetClass && !Widget && GetWorld() )
+		{
+			Widget = CreateWidget<UUserWidget>( GetWorld(), WidgetClass );
+		}
+		else if( !WidgetClass && Widget )
+		{
+			Widget = nullptr;
+		}
 
-	if(!SlateWidget.IsValid())
-	{
-		SlateWidget = SNew(SVirtualWindow).Size(DrawSize);
+		if(!SlateWidget.IsValid())
+		{
+			SlateWidget = SNew(SVirtualWindow).Size(DrawSize);
 
-	}
+		}
 	
-	if( !HitTestGrid.IsValid() )
-	{
-		HitTestGrid = MakeShareable( new FHittestGrid );
-	}
-
-	{
-		SlateWidget->Resize( DrawSize );
-
-		if( Widget )
+		if( !HitTestGrid.IsValid() )
 		{
-			SlateWidget->SetContent( Widget->TakeWidget() );
+			HitTestGrid = MakeShareable( new FHittestGrid );
 		}
-		else
+
 		{
-			SlateWidget->SetContent( SNullWidget::NullWidget );
-		}
+			SlateWidget->Resize( DrawSize );
 
-	}
+			if( Widget )
+			{
+				SlateWidget->SetContent( Widget->TakeWidget() );
+			}
+			else
+			{
+				SlateWidget->SetContent( SNullWidget::NullWidget );
+			}
 
-	if( Widget && !GetWorld()->IsGameWorld() )
-	{
-		// Prevent native ticking of editor component previews
-		Widget->SetIsDesignTime( true );
+		}
+
+		if( Widget && !GetWorld()->IsGameWorld() )
+		{
+			// Prevent native ticking of editor component previews
+			Widget->SetIsDesignTime( true );
+		}
 	}
 }
 
