UMG - ScrollBar should not be a panel widget.

diff --git a/Engine/Source/Runtime/UMG/Public/Components/ScrollBar.h b/Engine/Source/Runtime/UMG/Public/Components/ScrollBar.h
--- a/Engine/Source/Runtime/UMG/Public/Components/ScrollBar.h
+++ b/Engine/Source/Runtime/UMG/Public/Components/ScrollBar.h
@@ -6,7 +6,7 @@
 
 /** */
 UCLASS(Experimental)
-class UMG_API UScrollBar : public UPanelWidget
+class UMG_API UScrollBar : public UWidget
 {
 	GENERATED_UCLASS_BODY()
 
