UMG - Dead code removal from textblock.

diff --git a/Engine/Source/Runtime/UMG/Public/Components/TextBlock.h b/Engine/Source/Runtime/UMG/Public/Components/TextBlock.h
--- a/Engine/Source/Runtime/UMG/Public/Components/TextBlock.h
+++ b/Engine/Source/Runtime/UMG/Public/Components/TextBlock.h
@@ -42,10 +42,6 @@ public:
 	void SetShadowOffset(FVector2D InShadowOffset);
 
 public:
-	/** The text block style */
-	//UPROPERTY(EditAnywhere, BlueprintReadWrite, Category="Style", meta=( DisplayName="Style" ))
-	//FTextBlockStyle WidgetStyle;
-
 	UPROPERTY()
 	USlateWidgetStyleAsset* Style_DEPRECATED;
 
@@ -135,8 +131,6 @@ public:
 	virtual void PostLoad() override;
 	// End of UObject interface
 
-	static const FTextBlockStyle* GetDefaultStyle();
-
 #if WITH_EDITOR
 	// UWidget interface
 	virtual const FSlateBrush* GetEditorIcon() override;
