UMG - Deadcode removal from textblock.

diff --git a/Engine/Source/Runtime/UMG/Private/Components/TextBlock.cpp b/Engine/Source/Runtime/UMG/Private/Components/TextBlock.cpp
--- a/Engine/Source/Runtime/UMG/Private/Components/TextBlock.cpp
+++ b/Engine/Source/Runtime/UMG/Private/Components/TextBlock.cpp
@@ -12,8 +12,6 @@ UTextBlock::UTextBlock(const FObjectInitializer& ObjectInitializer)
 {
 	bIsVariable = false;
 
-//	WidgetStyle = *UTextBlock::GetDefaultStyle();
-
 	Text = LOCTEXT("TextBlockDefaultValue", "Text Block");
 	ShadowOffset = FVector2D(1.0f, 1.0f);
 	ColorAndOpacity = FLinearColor::White;
@@ -114,23 +112,11 @@ void UTextBlock::PostLoad()
 	{
 		if ( Style_DEPRECATED != nullptr )
 		{
-			//const FTextBlockStyle* StylePtr = Style_DEPRECATED->GetStyle<FTextBlockStyle>();
-			//if ( StylePtr != nullptr )
-			//{
-			//	WidgetStyle = *StylePtr;
-			//}
-
 			Style_DEPRECATED = nullptr;
 		}
 	}
 }
 
-const FTextBlockStyle* UTextBlock::GetDefaultStyle()
-{
-	static const STextBlock::FArguments Defaults;
-	return Defaults._TextStyle;
-}
-
 #if WITH_EDITOR
 
 FString UTextBlock::GetLabelMetadata() const
