UMG - Adding a new ReplaceChild function to the PanelWidget for editor usage.

diff --git a/Engine/Source/Runtime/UMG/Private/Components/PanelWidget.cpp b/Engine/Source/Runtime/UMG/Private/Components/PanelWidget.cpp
--- a/Engine/Source/Runtime/UMG/Private/Components/PanelWidget.cpp
+++ b/Engine/Source/Runtime/UMG/Private/Components/PanelWidget.cpp
@@ -137,6 +137,17 @@ bool UPanelWidget::ReplaceChildAt(int32 Index, UWidget* Content)
 
 #if WITH_EDITOR
 
+bool UPanelWidget::ReplaceChild(UWidget* CurrentChild, UWidget* NewChild)
+{
+	int32 Index = GetChildIndex(CurrentChild);
+	if ( Index != -1 )
+	{
+		return ReplaceChildAt(Index, NewChild);
+	}
+
+	return false;
+}
+
 UPanelSlot* UPanelWidget::InsertChildAt(int32 Index, UWidget* Content)
 {
 	UPanelSlot* NewSlot = AddChild(Content);
diff --git a/Engine/Source/Runtime/UMG/Public/Components/PanelWidget.h b/Engine/Source/Runtime/UMG/Public/Components/PanelWidget.h
--- a/Engine/Source/Runtime/UMG/Public/Components/PanelWidget.h
+++ b/Engine/Source/Runtime/UMG/Public/Components/PanelWidget.h
@@ -57,6 +57,14 @@ public:
 
 #if WITH_EDITOR
 
+	/**
+	 * Swaps the child widget out of the slot, and replaces it with the new child widget.
+	 * @param CurrentChild The existing child widget being removed.
+	 * @param NewChild The new child widget being inserted.
+	 * @return true if the CurrentChild was found and the swap occurred, otherwise false.
+	 */
+	virtual bool ReplaceChild(UWidget* CurrentChild, UWidget* NewChild);
+
 	/**
 	 * Inserts a widget at a specific index.  This does not update the live slate version, it requires
 	 * a rebuild of the whole UI to see a change.
