[Github PR 569] EverNewJoy - Get All Widgets Of Class added to the widget blueprint library.  Allows accessing all widgets in the current world, optionally allows just getting the top level ones in the viewport.

diff --git a/Engine/Source/Runtime/UMG/Private/WidgetBlueprintLibrary.cpp b/Engine/Source/Runtime/UMG/Private/WidgetBlueprintLibrary.cpp
--- a/Engine/Source/Runtime/UMG/Private/WidgetBlueprintLibrary.cpp
+++ b/Engine/Source/Runtime/UMG/Private/WidgetBlueprintLibrary.cpp
@@ -345,4 +345,43 @@ void UWidgetBlueprintLibrary::DismissAllMenus()
 	FSlateApplication::Get().DismissAllMenus();
 }
 
+void UWidgetBlueprintLibrary::GetAllWidgetsOfClass(UObject* WorldContextObject, TArray<UUserWidget*>& FoundWidgets, TSubclassOf<UUserWidget> WidgetClass, bool TopLevelOnly)
+{
+	//Prevent possibility of an ever-growing array if user uses this in a loop
+	FoundWidgets.Empty();
+	
+	if ( !WidgetClass || !WorldContextObject )
+	{
+		return;
+	}
+	 
+	const UWorld* World = GEngine->GetWorldFromContextObject(WorldContextObject);
+	if ( !World )
+	{
+		return;
+	}
+
+	for ( TObjectIterator<UUserWidget> Itr; Itr; ++Itr )
+	{
+		UUserWidget* LiveWidget = *Itr;
+
+		if ( LiveWidget->GetWorld() != World )
+		{
+			continue;
+		}
+		
+		if ( TopLevelOnly )
+		{
+			if ( LiveWidget->IsInViewport() )
+			{
+				FoundWidgets.Add(LiveWidget);
+			}
+		}
+		else
+		{
+			FoundWidgets.Add(LiveWidget);
+		}
+	}
+}
+
 #undef LOCTEXT_NAMESPACE
\ No newline at end of file
diff --git a/Engine/Source/Runtime/UMG/Public/Blueprint/WidgetBlueprintLibrary.h b/Engine/Source/Runtime/UMG/Public/Blueprint/WidgetBlueprintLibrary.h
--- a/Engine/Source/Runtime/UMG/Public/Blueprint/WidgetBlueprintLibrary.h
+++ b/Engine/Source/Runtime/UMG/Public/Blueprint/WidgetBlueprintLibrary.h
@@ -184,4 +184,13 @@ public:
 	/** Closes any popup menu */
 	UFUNCTION(BlueprintCallable, BlueprintCosmetic, Category="Widget|Menu")
 	static void DismissAllMenus();
+
+	/**
+	 * Find all widgets of a certain class.
+	 * @param FoundWidgets The widgets that were found matching the filter.
+	 * @param WidgetClass The widget class to filter by.
+	 * @param TopLevelOnly Only the widgets that are direct children of the viewport will be returned.
+	 */
+	UFUNCTION(BlueprintCallable, BlueprintCosmetic, Category="Widget", meta=( WorldContext="WorldContextObject" ))
+	static void GetAllWidgetsOfClass(UObject* WorldContextObject, TArray<UUserWidget*>& FoundWidgets, TSubclassOf<UUserWidget> WidgetClass, bool TopLevelOnly = true);
 };
