UMG - Adding a function that converts raw mouse position to a DPI scaled position.

diff --git a/Engine/Source/Runtime/UMG/Private/WidgetLayoutLibrary.cpp b/Engine/Source/Runtime/UMG/Private/WidgetLayoutLibrary.cpp
--- a/Engine/Source/Runtime/UMG/Private/WidgetLayoutLibrary.cpp
+++ b/Engine/Source/Runtime/UMG/Private/WidgetLayoutLibrary.cpp
@@ -32,6 +32,20 @@ float UWidgetLayoutLibrary::GetViewportScale(UObject* WorldContextObject)
 	return 1;
 }
 
+bool UWidgetLayoutLibrary::GetMousePositionScaledByDPI(APlayerController* Player, float& LocationX, float& LocationY)
+{
+	if ( Player && Player->GetMousePosition(LocationX, LocationY) )
+	{
+		float Scale = UWidgetLayoutLibrary::GetViewportScale(Player);
+		LocationX = LocationX * ( 1.0f / Scale );
+		LocationY = LocationY * ( 1.0f / Scale );
+
+		return true;
+	}
+
+	return false;
+}
+
 FVector2D UWidgetLayoutLibrary::GetViewportSize(UObject* WorldContextObject)
 {
 	UWorld* World = GEngine->GetWorldFromContextObject(WorldContextObject);
diff --git a/Engine/Source/Runtime/UMG/Public/Blueprint/WidgetLayoutLibrary.h b/Engine/Source/Runtime/UMG/Public/Blueprint/WidgetLayoutLibrary.h
--- a/Engine/Source/Runtime/UMG/Public/Blueprint/WidgetLayoutLibrary.h
+++ b/Engine/Source/Runtime/UMG/Public/Blueprint/WidgetLayoutLibrary.h
@@ -11,14 +11,27 @@ class UWidgetLayoutLibrary : public UBlueprintFunctionLibrary
 
 public:
 
-	/**  */
-	UFUNCTION(BlueprintPure, Category="Viewport", meta=( HidePin="WorldContextObject", DefaultToSelf="WorldContextObject" ))
+	/**
+	 * Gets the current DPI Scale being applied to the viewport and all the Widgets.
+	 */
+	UFUNCTION(BlueprintPure, BlueprintCosmetic, Category="Viewport", meta=( HidePin="WorldContextObject", DefaultToSelf="WorldContextObject" ))
 	static float GetViewportScale(UObject* WorldContextObject);
 
-	/**  */
-	UFUNCTION(BlueprintPure, Category="Viewport", meta=( HidePin="WorldContextObject", DefaultToSelf="WorldContextObject" ))
+	/**
+	 * Gets the size of the game viewport.
+	 */
+	UFUNCTION(BlueprintPure, BlueprintCosmetic, Category="Viewport", meta=( HidePin="WorldContextObject", DefaultToSelf="WorldContextObject" ))
 	static FVector2D GetViewportSize(UObject* WorldContextObject);
 
+	/**
+	 * Gets the mouse position of the player controller, scaled by the DPI.  If you're trying to go from raw mouse screenspace coordinates
+	 * to fullscreen widget space, you'll need to transform the mouse into DPI Scaled space.  This function performs that scaling.
+	 *
+	 * MousePositionScaledByDPI = MousePosition * (1 / ViewportScale).
+	 */
+	UFUNCTION(BlueprintPure, BlueprintCosmetic, Category="Viewport")
+	static bool GetMousePositionScaledByDPI(APlayerController* Player, float& LocationX, float& LocationY);
+
 	/**
 	 * Gets the slot object on the child widget as a Canvas Slot, allowing you to manipulate layout information.
 	 * @param Widget The child widget of a canvas panel.
