UMG - Adding a bit of metadata to events for buttons, checkbox and editabletextbox.

diff --git a/Engine/Source/Runtime/UMG/Public/Components/Button.h b/Engine/Source/Runtime/UMG/Public/Components/Button.h
--- a/Engine/Source/Runtime/UMG/Public/Components/Button.h
+++ b/Engine/Source/Runtime/UMG/Public/Components/Button.h
@@ -53,7 +53,7 @@ public:
 public:
 
 	/** Called when the button is clicked */
-	UPROPERTY(BlueprintAssignable)
+	UPROPERTY(BlueprintAssignable, Category="Button|Event")
 	FOnButtonClickedEvent OnClicked;
 
 public:
diff --git a/Engine/Source/Runtime/UMG/Public/Components/CheckBox.h b/Engine/Source/Runtime/UMG/Public/Components/CheckBox.h
--- a/Engine/Source/Runtime/UMG/Public/Components/CheckBox.h
+++ b/Engine/Source/Runtime/UMG/Public/Components/CheckBox.h
@@ -90,7 +90,7 @@ public:
 public:
 
 	/** Called when the checked state has changed */
-	UPROPERTY(BlueprintAssignable)
+	UPROPERTY(BlueprintAssignable, Category="CheckBox|Event")
 	FOnCheckBoxComponentStateChanged OnCheckStateChanged;
 
 public:
diff --git a/Engine/Source/Runtime/UMG/Public/Components/EditableTextBox.h b/Engine/Source/Runtime/UMG/Public/Components/EditableTextBox.h
--- a/Engine/Source/Runtime/UMG/Public/Components/EditableTextBox.h
+++ b/Engine/Source/Runtime/UMG/Public/Components/EditableTextBox.h
@@ -99,12 +99,14 @@ public:
 	UPROPERTY(EditDefaultsOnly, BlueprintReadOnly, Category=Behavior, AdvancedDisplay)
 	bool SelectAllTextOnCommit;
 
+public:
+
 	/** Called whenever the text is changed interactively by the user */
-	UPROPERTY(BlueprintAssignable, Category="Widget Event")
+	UPROPERTY(BlueprintAssignable, Category="TextBox|Event")
 	FOnEditableTextBoxChangedEvent OnTextChanged;
 
 	/** Called whenever the text is committed.  This happens when the user presses enter or the text box loses focus. */
-	UPROPERTY(BlueprintAssignable, Category="Widget Event")
+	UPROPERTY(BlueprintAssignable, Category="TextBox|Event")
 	FOnEditableTextBoxCommittedEvent OnTextCommitted;
 
 	/** Provide a alternative mechanism for error reporting. */
