Fixed crash opening UMG content examples map

diff --git a/Engine/Source/Runtime/UMG/Private/Components/WidgetComponent.cpp b/Engine/Source/Runtime/UMG/Private/Components/WidgetComponent.cpp
--- a/Engine/Source/Runtime/UMG/Private/Components/WidgetComponent.cpp
+++ b/Engine/Source/Runtime/UMG/Private/Components/WidgetComponent.cpp
@@ -538,7 +538,10 @@ void UWidgetComponent::ApplyComponentInstanceData(FComponentInstanceDataBase* Co
 	FWidgetComponentInstanceData* WidgetInstanceData  = static_cast<FWidgetComponentInstanceData*>(ComponentInstanceData);
 
 	RenderTarget = WidgetInstanceData->RenderTarget;
-	MaterialInstance->SetTextureParameterValue("SlateUI", RenderTarget);
+	if( MaterialInstance )
+	{
+		MaterialInstance->SetTextureParameterValue("SlateUI", RenderTarget);
+	}
 
 	MarkRenderStateDirty();
 }
