Changed display names on Blueprint functions for the UI editor so the user can better identify between them.

Tooltip for functions will display target class.

diff --git a/Engine/Source/Editor/BlueprintGraph/Private/K2Node_CallFunction.cpp b/Engine/Source/Editor/BlueprintGraph/Private/K2Node_CallFunction.cpp
--- a/Engine/Source/Editor/BlueprintGraph/Private/K2Node_CallFunction.cpp
+++ b/Engine/Source/Editor/BlueprintGraph/Private/K2Node_CallFunction.cpp
@@ -470,13 +470,13 @@ FText UK2Node_CallFunction::GetNodeTitle(ENodeTitleType::Type TitleType) const
 		ContextString = GetFunctionContextString();
 	}
 	else
-		{
+	{
 		FunctionName = FunctionReference.GetMemberName().ToString();
 		if ((GEditor != NULL) && (GetDefault<UEditorStyleSettings>()->bShowFriendlyNames))
 		{
 			FunctionName = FName::NameToDisplayString(FunctionName, false);
 		}
-		}
+	}
 
 	if(TitleType == ENodeTitleType::FullTitle)
 	{
@@ -1249,7 +1249,19 @@ FString UK2Node_CallFunction::GetDefaultTooltipForFunction(const UFunction* Func
 		Tooltip.Trim();
 		Tooltip.TrimTrailing();
 
-		return Tooltip;
+		UClass* CurrentSelfClass = (Function != NULL) ? Function->GetOwnerClass() : NULL;
+		UClass const* TrueSelfClass = CurrentSelfClass;
+		if (CurrentSelfClass && CurrentSelfClass->ClassGeneratedBy)
+		{
+			TrueSelfClass = CurrentSelfClass->GetAuthoritativeClass();
+		}
+
+		FText TargetDisplayText = (TrueSelfClass != NULL) ? TrueSelfClass->GetDisplayNameText() : LOCTEXT("None", "None");
+
+		FFormatNamedArguments Args;
+		Args.Add(TEXT("TargetName"), TargetDisplayText);
+		Args.Add(TEXT("Tooltip"), FText::FromString(Tooltip));
+		return FText::Format(LOCTEXT("CallFunction_Tooltip", "{Tooltip}\n\nTarget is {TargetName}"), Args).ToString();
 	}
 	else
 	{
diff --git a/Engine/Source/Runtime/UMG/Public/Components/EditableText.h b/Engine/Source/Runtime/UMG/Public/Components/EditableText.h
--- a/Engine/Source/Runtime/UMG/Public/Components/EditableText.h
+++ b/Engine/Source/Runtime/UMG/Public/Components/EditableText.h
@@ -105,18 +105,18 @@ public:
 public:
 
 	/** Called whenever the text is changed interactively by the user */
-	UPROPERTY(BlueprintAssignable, Category="Widget Event")
+	UPROPERTY(BlueprintAssignable, Category="Widget Event", meta=(DisplayName="OnTextChanged (Editable Text)"))
 	FOnEditableTextChangedEvent OnTextChanged;
 
 	/** Called whenever the text is committed.  This happens when the user presses enter or the text box loses focus. */
-	UPROPERTY(BlueprintAssignable, Category="Widget Event")
+	UPROPERTY(BlueprintAssignable, Category="Widget Event", meta=(DisplayName="OnTextCommitted (Editable Text)"))
 	FOnEditableTextCommittedEvent OnTextCommitted;
 
 	/**
 	 * Gets the widget text
 	 * @return The widget text
 	 */
-	UFUNCTION(BlueprintCallable, Category = "Widget")
+	UFUNCTION(BlueprintCallable, Category = "Widget", meta=(FriendlyName="GetText (Editable Text)"))
 	FText GetText() const;
 
 	/**
@@ -124,7 +124,7 @@ public:
 	 * Warning: This will wipe any binding created for the Text property!
 	 * @param InText The text to assign to the widget
 	 */
-	UFUNCTION(BlueprintCallable, Category = "Widget")
+	UFUNCTION(BlueprintCallable, Category = "Widget", meta=(FriendlyName="SetText (Editable Text)"))
 	void SetText(FText InText);
 
 	UFUNCTION(BlueprintCallable, Category = "Widget")
diff --git a/Engine/Source/Runtime/UMG/Public/Components/EditableTextBox.h b/Engine/Source/Runtime/UMG/Public/Components/EditableTextBox.h
--- a/Engine/Source/Runtime/UMG/Public/Components/EditableTextBox.h
+++ b/Engine/Source/Runtime/UMG/Public/Components/EditableTextBox.h
@@ -115,14 +115,14 @@ public:
 public:
 
 	/**  */
-	UFUNCTION(BlueprintCallable, Category="Widget")
+	UFUNCTION(BlueprintCallable, Category="Widget", meta=(FriendlyName="GetText (Text Box)"))
 	FText GetText() const;
 
 	/**  */
-	UFUNCTION(BlueprintCallable, Category="Widget")
+	UFUNCTION(BlueprintCallable, Category="Widget", meta=(FriendlyName="SetText (Text Box)"))
 	void SetText(FText InText);
 
-	UFUNCTION(BlueprintCallable, Category="Widget")
+	UFUNCTION(BlueprintCallable, Category="Widget",  meta=(FriendlyName="SetError (Text Box)"))
 	void SetError(FText InError);
 
 	UFUNCTION(BlueprintCallable, Category="Widget")
diff --git a/Engine/Source/Runtime/UMG/Public/Components/MultiLineEditableText.h b/Engine/Source/Runtime/UMG/Public/Components/MultiLineEditableText.h
--- a/Engine/Source/Runtime/UMG/Public/Components/MultiLineEditableText.h
+++ b/Engine/Source/Runtime/UMG/Public/Components/MultiLineEditableText.h
@@ -45,21 +45,21 @@ public:
 	FSlateFontInfo Font;
 
 	/** Called whenever the text is changed interactively by the user */
-	UPROPERTY(BlueprintAssignable, Category="Widget Event")
+	UPROPERTY(BlueprintAssignable, Category="Widget Event", meta=(DisplayName="OnTextChanged (Multi-Line Editable Text)"))
 	FOnMultiLineEditableTextChangedEvent OnTextChanged;
 
 	/** Called whenever the text is committed.  This happens when the user presses enter or the text box loses focus. */
-	UPROPERTY(BlueprintAssignable, Category="Widget Event")
+	UPROPERTY(BlueprintAssignable, Category="Widget Event", meta=(DisplayName="OnTextCommitted (Multi-Line Editable Text)"))
 	FOnMultiLineEditableTextCommittedEvent OnTextCommitted;
 
 public:
 
 	/**  */
-	UFUNCTION(BlueprintCallable, Category="Widget")
+	UFUNCTION(BlueprintCallable, Category="Widget", meta=(FriendlyName="GetText (Multi-Line Editable Text)"))
 	FText GetText() const;
 
 	/**  */
-	UFUNCTION(BlueprintCallable, Category="Widget")
+	UFUNCTION(BlueprintCallable, Category="Widget", meta=(FriendlyName="SetText (Multi-Line Editable Text)"))
 	void SetText(FText InText);
 	
 	// UWidget interface
diff --git a/Engine/Source/Runtime/UMG/Public/Components/MultiLineEditableTextBox.h b/Engine/Source/Runtime/UMG/Public/Components/MultiLineEditableTextBox.h
--- a/Engine/Source/Runtime/UMG/Public/Components/MultiLineEditableTextBox.h
+++ b/Engine/Source/Runtime/UMG/Public/Components/MultiLineEditableTextBox.h
@@ -65,11 +65,11 @@ public:
 	FLinearColor ReadOnlyForegroundColor;
 
 	/** Called whenever the text is changed interactively by the user */
-	UPROPERTY(BlueprintAssignable, Category="Widget Event")
+	UPROPERTY(BlueprintAssignable, Category="Widget Event", meta=(DisplayName="OnTextChanged (Multi-Line Text Box)"))
 	FOnMultiLineEditableTextBoxChangedEvent OnTextChanged;
 
 	/** Called whenever the text is committed.  This happens when the user presses enter or the text box loses focus. */
-	UPROPERTY(BlueprintAssignable, Category="Widget Event")
+	UPROPERTY(BlueprintAssignable, Category="Widget Event", meta=(DisplayName="OnTextCommitted (Multi-Line Text Box)"))
 	FOnMultiLineEditableTextBoxCommittedEvent OnTextCommitted;
 
 	/** Provide a alternative mechanism for error reporting. */
@@ -78,14 +78,14 @@ public:
 public:
 
 	/**  */
-	UFUNCTION(BlueprintCallable, Category="Widget")
+	UFUNCTION(BlueprintCallable, Category="Widget", meta=(FriendlyName="GetText (Multi-Line Text Box)"))
 	FText GetText() const;
 
 	/**  */
-	UFUNCTION(BlueprintCallable, Category="Widget")
+	UFUNCTION(BlueprintCallable, Category="Widget", meta=(FriendlyName="SetText (Multi-Line Text Box)"))
 	void SetText(FText InText);
 
-	UFUNCTION(BlueprintCallable, Category="Widget")
+	UFUNCTION(BlueprintCallable, Category="Widget", meta=(FriendlyName="SetError (Multi-Line Text Box)"))
 	void SetError(FText InError);
 
 	//TODO UMG Add Set ReadOnlyForegroundColor
diff --git a/Engine/Source/Runtime/UMG/Public/Components/TextBlock.h b/Engine/Source/Runtime/UMG/Public/Components/TextBlock.h
--- a/Engine/Source/Runtime/UMG/Public/Components/TextBlock.h
+++ b/Engine/Source/Runtime/UMG/Public/Components/TextBlock.h
@@ -108,7 +108,7 @@ public:
 	 * Gets the widget text
 	 * @return The widget text
 	 */
-	UFUNCTION(BlueprintCallable, Category="Widget")
+	UFUNCTION(BlueprintCallable, Category="Widget", meta=(FriendlyName="GetText (Text)"))
 	FText GetText() const;
 
 	/**
@@ -116,7 +116,7 @@ public:
 	 * Warning: This will wipe any binding created for the Text property!
 	 * @param InText The text to assign to the widget
 	 */
-	UFUNCTION(BlueprintCallable, Category="Widget")
+	UFUNCTION(BlueprintCallable, Category="Widget", meta=(FriendlyName="SetText (Text)"))
 	void SetText(FText InText);
 
 	// UWidget interface
