Fixed UMG widget visibility being reset to "Visible" on load (regression)

diff --git a/Engine/Source/Runtime/UMG/Private/Components/CanvasPanel.cpp b/Engine/Source/Runtime/UMG/Private/Components/CanvasPanel.cpp
--- a/Engine/Source/Runtime/UMG/Private/Components/CanvasPanel.cpp
+++ b/Engine/Source/Runtime/UMG/Private/Components/CanvasPanel.cpp
@@ -13,7 +13,7 @@ UCanvasPanel::UCanvasPanel(const FObjectInitializer& ObjectInitializer)
 	bIsVariable = false;
 
 	SConstraintCanvas::FArguments Defaults;
-	Visibility = UWidget::ConvertRuntimeToSerializedVisibility(Defaults._Visibility.Get());
+	Visiblity_DEPRECATED = Visibility = UWidget::ConvertRuntimeToSerializedVisibility(Defaults._Visibility.Get());
 }
 
 void UCanvasPanel::ReleaseSlateResources(bool bReleaseChildren)
diff --git a/Engine/Source/Runtime/UMG/Private/Components/GridPanel.cpp b/Engine/Source/Runtime/UMG/Private/Components/GridPanel.cpp
--- a/Engine/Source/Runtime/UMG/Private/Components/GridPanel.cpp
+++ b/Engine/Source/Runtime/UMG/Private/Components/GridPanel.cpp
@@ -13,7 +13,7 @@ UGridPanel::UGridPanel(const FObjectInitializer& ObjectInitializer)
 	bIsVariable = false;
 
 	SGridPanel::FArguments Defaults;
-	Visibility = UWidget::ConvertRuntimeToSerializedVisibility(Defaults._Visibility.Get());
+	Visiblity_DEPRECATED = Visibility = UWidget::ConvertRuntimeToSerializedVisibility(Defaults._Visibility.Get());
 }
 
 void UGridPanel::ReleaseSlateResources(bool bReleaseChildren)
diff --git a/Engine/Source/Runtime/UMG/Private/Components/HorizontalBox.cpp b/Engine/Source/Runtime/UMG/Private/Components/HorizontalBox.cpp
--- a/Engine/Source/Runtime/UMG/Private/Components/HorizontalBox.cpp
+++ b/Engine/Source/Runtime/UMG/Private/Components/HorizontalBox.cpp
@@ -13,7 +13,7 @@ UHorizontalBox::UHorizontalBox(const FObjectInitializer& ObjectInitializer)
 	bIsVariable = false;
 
 	SHorizontalBox::FArguments Defaults;
-	Visibility = UWidget::ConvertRuntimeToSerializedVisibility(Defaults._Visibility.Get());
+	Visiblity_DEPRECATED = Visibility = UWidget::ConvertRuntimeToSerializedVisibility(Defaults._Visibility.Get());
 }
 
 void UHorizontalBox::ReleaseSlateResources(bool bReleaseChildren)
diff --git a/Engine/Source/Runtime/UMG/Private/Components/Overlay.cpp b/Engine/Source/Runtime/UMG/Private/Components/Overlay.cpp
--- a/Engine/Source/Runtime/UMG/Private/Components/Overlay.cpp
+++ b/Engine/Source/Runtime/UMG/Private/Components/Overlay.cpp
@@ -13,7 +13,7 @@ UOverlay::UOverlay(const FObjectInitializer& ObjectInitializer)
 	bIsVariable = false;
 
 	SOverlay::FArguments Defaults;
-	Visibility = UWidget::ConvertRuntimeToSerializedVisibility(Defaults._Visibility.Get());
+	Visiblity_DEPRECATED = Visibility = UWidget::ConvertRuntimeToSerializedVisibility(Defaults._Visibility.Get());
 }
 
 void UOverlay::ReleaseSlateResources(bool bReleaseChildren)
diff --git a/Engine/Source/Runtime/UMG/Private/Components/ScrollBox.cpp b/Engine/Source/Runtime/UMG/Private/Components/ScrollBox.cpp
--- a/Engine/Source/Runtime/UMG/Private/Components/ScrollBox.cpp
+++ b/Engine/Source/Runtime/UMG/Private/Components/ScrollBox.cpp
@@ -15,7 +15,7 @@ UScrollBox::UScrollBox(const FObjectInitializer& ObjectInitializer)
 	Orientation = Orient_Vertical;
 
 	SScrollBox::FArguments Defaults;
-	Visibility = UWidget::ConvertRuntimeToSerializedVisibility(Defaults._Visibility.Get());
+	Visiblity_DEPRECATED = Visibility = UWidget::ConvertRuntimeToSerializedVisibility(Defaults._Visibility.Get());
 
 	WidgetStyle = *Defaults._Style;
 	WidgetBarStyle = *Defaults._ScrollBarStyle;
diff --git a/Engine/Source/Runtime/UMG/Private/Components/UniformGridPanel.cpp b/Engine/Source/Runtime/UMG/Private/Components/UniformGridPanel.cpp
--- a/Engine/Source/Runtime/UMG/Private/Components/UniformGridPanel.cpp
+++ b/Engine/Source/Runtime/UMG/Private/Components/UniformGridPanel.cpp
@@ -13,7 +13,7 @@ UUniformGridPanel::UUniformGridPanel(const FObjectInitializer& ObjectInitializer
 	bIsVariable = false;
 
 	SUniformGridPanel::FArguments Defaults;
-	Visibility = UWidget::ConvertRuntimeToSerializedVisibility(Defaults._Visibility.Get());
+	Visiblity_DEPRECATED = Visibility = UWidget::ConvertRuntimeToSerializedVisibility(Defaults._Visibility.Get());
 }
 
 void UUniformGridPanel::ReleaseSlateResources(bool bReleaseChildren)
diff --git a/Engine/Source/Runtime/UMG/Private/Components/VerticalBox.cpp b/Engine/Source/Runtime/UMG/Private/Components/VerticalBox.cpp
--- a/Engine/Source/Runtime/UMG/Private/Components/VerticalBox.cpp
+++ b/Engine/Source/Runtime/UMG/Private/Components/VerticalBox.cpp
@@ -13,7 +13,7 @@ UVerticalBox::UVerticalBox(const FObjectInitializer& ObjectInitializer)
 	bIsVariable = false;
 
 	SVerticalBox::FArguments Defaults;
-	Visibility = UWidget::ConvertRuntimeToSerializedVisibility(Defaults._Visibility.Get());
+	Visiblity_DEPRECATED = Visibility = UWidget::ConvertRuntimeToSerializedVisibility(Defaults._Visibility.Get());
 }
 
 void UVerticalBox::ReleaseSlateResources(bool bReleaseChildren)
diff --git a/Engine/Source/Runtime/UMG/Private/Components/Widget.cpp b/Engine/Source/Runtime/UMG/Private/Components/Widget.cpp
--- a/Engine/Source/Runtime/UMG/Private/Components/Widget.cpp
+++ b/Engine/Source/Runtime/UMG/Private/Components/Widget.cpp
@@ -13,7 +13,7 @@ UWidget::UWidget(const FObjectInitializer& ObjectInitializer)
 	bIsEnabled = true;
 	bIsVariable = true;
 	bDesignTime = false;
-	Visibility = ESlateVisibility::Visible;
+	Visiblity_DEPRECATED = Visibility = ESlateVisibility::Visible;	
 	RenderTransformPivot = FVector2D(0.5f, 0.5f);
 
 	//TODO UMG ToolTipWidget
diff --git a/Engine/Source/Runtime/UMG/Private/Components/WidgetSwitcher.cpp b/Engine/Source/Runtime/UMG/Private/Components/WidgetSwitcher.cpp
--- a/Engine/Source/Runtime/UMG/Private/Components/WidgetSwitcher.cpp
+++ b/Engine/Source/Runtime/UMG/Private/Components/WidgetSwitcher.cpp
@@ -14,7 +14,7 @@ UWidgetSwitcher::UWidgetSwitcher(const FObjectInitializer& ObjectInitializer)
 	bIsVariable = true;
 
 	SWidgetSwitcher::FArguments Defaults;
-	Visibility = UWidget::ConvertRuntimeToSerializedVisibility(Defaults._Visibility.Get());
+	Visiblity_DEPRECATED = Visibility = Visiblity_DEPRECATED = UWidget::ConvertRuntimeToSerializedVisibility(Defaults._Visibility.Get());
 }
 
 void UWidgetSwitcher::ReleaseSlateResources(bool bReleaseChildren)
diff --git a/Engine/Source/Runtime/UMG/Private/Components/WrapBox.cpp b/Engine/Source/Runtime/UMG/Private/Components/WrapBox.cpp
--- a/Engine/Source/Runtime/UMG/Private/Components/WrapBox.cpp
+++ b/Engine/Source/Runtime/UMG/Private/Components/WrapBox.cpp
@@ -13,7 +13,7 @@ UWrapBox::UWrapBox(const FObjectInitializer& ObjectInitializer)
 	bIsVariable = false;
 
 	SWrapBox::FArguments Defaults;
-	Visibility = UWidget::ConvertRuntimeToSerializedVisibility(Defaults._Visibility.Get());
+	Visibility = Visiblity_DEPRECATED = UWidget::ConvertRuntimeToSerializedVisibility(Defaults._Visibility.Get());
 }
 
 void UWrapBox::ReleaseSlateResources(bool bReleaseChildren)
diff --git a/Engine/Source/Runtime/UMG/Private/UserWidget.cpp b/Engine/Source/Runtime/UMG/Private/UserWidget.cpp
--- a/Engine/Source/Runtime/UMG/Private/UserWidget.cpp
+++ b/Engine/Source/Runtime/UMG/Private/UserWidget.cpp
@@ -18,7 +18,7 @@ UUserWidget::UUserWidget(const FObjectInitializer& ObjectInitializer)
 	: Super(ObjectInitializer)
 {
 	ViewportAnchors = FAnchors(0, 0, 1, 1);
-	Visibility = ESlateVisibility::SelfHitTestInvisible;
+	Visiblity_DEPRECATED = Visibility = ESlateVisibility::SelfHitTestInvisible;
 
 	bInitialized = false;
 	bSupportsKeyboardFocus = true;
