UE-5544 Having a normal pin connected to a pass-by-reference pin can lead to a crash when using PIE

diff --git a/Engine/Source/Runtime/UMG/Public/Blueprint/UserWidget.h b/Engine/Source/Runtime/UMG/Public/Blueprint/UserWidget.h
--- a/Engine/Source/Runtime/UMG/Public/Blueprint/UserWidget.h
+++ b/Engine/Source/Runtime/UMG/Public/Blueprint/UserWidget.h
@@ -57,12 +57,9 @@ public:
 	/** We override the assignment operator to allow generated code to compile with the const ref member. */
 	void operator=( const FPaintContext& Other )
 	{
-		const_cast<FGeometry&>( AllottedGeometry ) = Other.AllottedGeometry;
-		const_cast<FSlateRect&>( MyClippingRect ) = Other.MyClippingRect;
-		OutDrawElements = Other.OutDrawElements;
-		LayerId = Other.LayerId;
-		const_cast<FWidgetStyle&>( WidgetStyle ) = Other.WidgetStyle;
-		bParentEnabled = Other.bParentEnabled;
+		FPaintContext* Ptr = this;
+		Ptr->~FPaintContext();
+		new(Ptr) FPaintContext(Other.AllottedGeometry, Other.MyClippingRect, Other.OutDrawElements, Other.LayerId, Other.WidgetStyle, Other.bParentEnabled);
 	}
 
 public:
