Fixed a crash when compiling a widget blueprint that's contained within another widget blueprint being displayed as a 3D widget component.

diff --git a/Engine/Source/Runtime/UMG/Private/Components/WidgetComponent.cpp b/Engine/Source/Runtime/UMG/Private/Components/WidgetComponent.cpp
--- a/Engine/Source/Runtime/UMG/Private/Components/WidgetComponent.cpp
+++ b/Engine/Source/Runtime/UMG/Private/Components/WidgetComponent.cpp
@@ -9,7 +9,6 @@
 #endif // !UE_SERVER
 #include "DynamicMeshBuilder.h"
 
-
 class SVirtualWindow : public SWindow
 {
 	SLATE_BEGIN_ARGS( SVirtualWindow )
@@ -43,7 +42,7 @@ public:
 		, BodySetup( InComponent->GetBodySetup() )
 		, bIsOpaque( InComponent->IsOpaque() )
 	{
-		bWillEverBeLit = false;
+		bWillEverBeLit = false;	
 	}
 
 	~FWidget3DSceneProxy()
@@ -398,8 +397,6 @@ void UWidgetComponent::OnRegister()
 		MaterialInstance = UMaterialInstanceDynamic::Create(Parent, this);
 	}
 
-	UpdateWidget();
-
 	if( !Renderer.IsValid() )
 	{
 		Renderer = FModuleManager::Get().GetModuleChecked<ISlateRHIRendererModule>("SlateRHIRenderer").CreateSlate3DRenderer();
@@ -451,8 +448,16 @@ void UWidgetComponent::TickComponent(float DeltaTime, enum ELevelTick TickType,
 #if !UE_SERVER
 
 	const float RenderTimeThreshold = .5f;
+
 	if( IsVisible() && GetWorld()->TimeSince(LastRenderTime) <= RenderTimeThreshold ) // Don't bother ticking if it hasn't been rendered recently
 	{
+		UpdateWidget();
+
+		if (!SlateWidget.IsValid())
+		{
+			return;
+		}
+
 		SlateWidget->SlatePrepass();
 
 		FGeometry WindowGeometry = FGeometry::MakeRoot( DrawSize, FSlateLayoutTransform() );
@@ -657,6 +662,7 @@ void UWidgetComponent::UpdateRenderTarget()
 		if ( RenderTarget->SizeX != DrawSize.X || RenderTarget->SizeY != DrawSize.Y )
 		{
 			RenderTarget->InitCustomFormat( DrawSize.X, DrawSize.Y, PF_B8G8R8A8, false );
+			MarkRenderStateDirty();
 		}
 
 		// Update the clear color
@@ -675,14 +681,13 @@ void UWidgetComponent::UpdateRenderTarget()
 				bClearColorChanged = true;
 			}
 		}
-		
-		MarkRenderStateDirty();
 	}
 
 	// If the clear color of the render target changed, update the BackColor of the material to match
 	if ( bClearColorChanged )
 	{
 		MaterialInstance->SetVectorParameterValue( "BackColor", RenderTarget->ClearColor );
+		MarkRenderStateDirty();
 	}
 }
 
