Can play looped animations on UMG widgets

diff --git a/Engine/Source/Runtime/UMG/Private/Animation/UMGSequencePlayer.cpp b/Engine/Source/Runtime/UMG/Private/Animation/UMGSequencePlayer.cpp
--- a/Engine/Source/Runtime/UMG/Private/Animation/UMGSequencePlayer.cpp
+++ b/Engine/Source/Runtime/UMG/Private/Animation/UMGSequencePlayer.cpp
@@ -63,25 +63,24 @@ void UUMGSequencePlayer::Tick(float DeltaTime)
 			RootMovieSceneInstance->Update(TimeCursorPosition, LastTimePosition, *this);
 		}
 
-		bool bIsFinished = false;
-
+		// Has a loop been completed?
 		if ( TimeCursorPosition >= TimeRange.GetUpperBoundValue() )
 		{
-			bIsFinished = true;
+			NumLoopsCompleted++;
+			TimeCursorPosition = 0.0;
 		}
-
-		if (bIsFinished)
+		
+		// Have we looped the requested number of times?
+		if (NumLoopsToPlay != 0 && NumLoopsCompleted >= NumLoopsToPlay)
 		{
-
 			PlayerStatus = EMovieScenePlayerStatus::Stopped;
 
 			OnSequenceFinishedPlayingEvent.Broadcast( *this );
-
 		}
 	}
 }
 
-void UUMGSequencePlayer::Play()
+void UUMGSequencePlayer::Play(float StartAtTime, int32 InNumLoopsToPlay)
 {
 	UMovieScene* MovieScene = Animation->MovieScene;
 
@@ -89,8 +88,9 @@ void UUMGSequencePlayer::Play()
 
 	RootMovieSceneInstance->RefreshInstance( *this );
 
-	// @todo UMG Support user specified start times
-	TimeCursorPosition = 0;
+	TimeCursorPosition = StaticCast<double>(StartAtTime);
+	NumLoopsToPlay = InNumLoopsToPlay;
+	NumLoopsCompleted = 0;
 
 	PlayerStatus = EMovieScenePlayerStatus::Playing;
 }
diff --git a/Engine/Source/Runtime/UMG/Private/UserWidget.cpp b/Engine/Source/Runtime/UMG/Private/UserWidget.cpp
--- a/Engine/Source/Runtime/UMG/Private/UserWidget.cpp
+++ b/Engine/Source/Runtime/UMG/Private/UserWidget.cpp
@@ -342,7 +342,12 @@ FEventReply UUserWidget::OnMotionDetected_Implementation(FGeometry MyGeometry, F
 	return UWidgetBlueprintLibrary::Unhandled();
 }
 
-void UUserWidget::PlayAnimation(const UWidgetAnimation* InAnimation)
+void UUserWidget::OnAnimationFinished_Implementation( const UWidgetAnimation* Animation )
+{
+	
+}
+
+void UUserWidget::PlayAnimation(const UWidgetAnimation* InAnimation, float StartAtTime, int32 NumberOfLoops)
 {
 	if( InAnimation )
 	{
@@ -358,11 +363,11 @@ void UUserWidget::PlayAnimation(const UWidgetAnimation* InAnimation)
 
 			NewPlayer->InitSequencePlayer( *InAnimation, *this );
 
-			NewPlayer->Play();
+			NewPlayer->Play( StartAtTime, NumberOfLoops );
 		}
 		else
 		{
-			(*FoundPlayer)->Play();
+			( *FoundPlayer )->Play( StartAtTime, NumberOfLoops );
 		}
 	}
 }
@@ -383,6 +388,7 @@ void UUserWidget::StopAnimation(const UWidgetAnimation* InAnimation)
 
 void UUserWidget::OnAnimationFinishedPlaying( UUMGSequencePlayer& Player )
 {
+	OnAnimationFinished( Player.GetAnimation() );
 	StoppedSequencePlayers.Add( &Player );
 }
 
diff --git a/Engine/Source/Runtime/UMG/Public/Animation/UMGSequencePlayer.h b/Engine/Source/Runtime/UMG/Public/Animation/UMGSequencePlayer.h
--- a/Engine/Source/Runtime/UMG/Public/Animation/UMGSequencePlayer.h
+++ b/Engine/Source/Runtime/UMG/Public/Animation/UMGSequencePlayer.h
@@ -20,7 +20,7 @@ public:
 	void Tick( float DeltaTime );
 
 	/** Begins playing or restarts an animation */
-	void Play();
+	void Play( float StartAtTime, int32 InNumLoopsToPlay );
 
 	/** Stops a running animation and resets time */
 	void Stop();
@@ -65,4 +65,10 @@ private:
 
 	/** Delegate to call when a sequence has finished playing */
 	FOnSequenceFinishedPlaying OnSequenceFinishedPlayingEvent;
+
+	/** The number of times to loop the animation playback */
+	int32 NumLoopsToPlay;
+
+	/** The number of loops completed since the last call to Play() */
+	int32 NumLoopsCompleted;
 };
diff --git a/Engine/Source/Runtime/UMG/Public/Blueprint/UserWidget.h b/Engine/Source/Runtime/UMG/Public/Blueprint/UserWidget.h
--- a/Engine/Source/Runtime/UMG/Public/Blueprint/UserWidget.h
+++ b/Engine/Source/Runtime/UMG/Public/Blueprint/UserWidget.h
@@ -513,6 +513,14 @@ public:
 	UFUNCTION(BlueprintNativeEvent, BlueprintCosmetic, Category="Touch Input")
 	FEventReply OnMotionDetected(FGeometry MyGeometry, FMotionEvent InMotionEvent);
 
+	/**
+	 * Called when an animation has either played all the way through or is stopped
+	 *
+	 * @param Animation The animation that has finished playing
+	 */
+	UFUNCTION( BlueprintNativeEvent, BlueprintCosmetic, Category = "Animation" )
+	void OnAnimationFinished( const UWidgetAnimation* Animation );
+
 public:
 
 	/**
@@ -533,12 +541,14 @@ public:
 	void SetForegroundColor(FSlateColor InForegroundColor);
 
 	/**
-	 * Plays an animation in this widget
+	 * Plays an animation in this widget a specified number of times
 	 * 
-	 * @param The name of the animation to play
+	 * @param InAnimation The animation to play
+	 * @param StartAtTime The time in the animation from which to start playing. For looped animations, this will only affect the first playback of the animation.
+	 * @param NumLoopsToPlay The number of times to loop this animation (0 to loop indefinitely)
 	 */
 	UFUNCTION(BlueprintCallable, BlueprintCosmetic, Category="User Interface|Animation")
-	void PlayAnimation(const UWidgetAnimation* InAnimation);
+	void PlayAnimation(const UWidgetAnimation* InAnimation, float StartAtTime = 0.0f, int32 NumLoopsToPlay = 1);
 
 	/**
 	 * Stops an already running animation in this widget
