UMG - The 3D Experimental widget now allows you to specifiy a Pivot (defaults to center).  Also now allows Space to be specified, defaults to worlds space, setting it to Screen causes the widgets to be constructed and added to the viewport.  Keeps the widget in the same relative point based on world to screen projection.

diff --git a/Engine/Source/Runtime/UMG/Private/Components/WidgetComponent.cpp b/Engine/Source/Runtime/UMG/Private/Components/WidgetComponent.cpp
--- a/Engine/Source/Runtime/UMG/Private/Components/WidgetComponent.cpp
+++ b/Engine/Source/Runtime/UMG/Private/Components/WidgetComponent.cpp
@@ -41,6 +41,7 @@ public:
 		, MaterialInstance( InComponent->GetMaterialInstance() )
 		, BodySetup( InComponent->GetBodySetup() )
 		, bIsOpaque( InComponent->IsOpaque() )
+		, Pivot( InComponent->GetPivot() )
 	{
 		bWillEverBeLit = false;	
 	}
@@ -59,18 +60,18 @@ public:
 		if( RenderTarget )
 		{
 			FTextureResource* TextureResource = RenderTarget->Resource;
-			if(TextureResource)
+			if ( TextureResource )
 			{
-				float U = 0.0f;
-				float V = 0.0f;
-				float UL = RenderTarget->SizeX;
-				float VL = RenderTarget->SizeY;
+				float U = -RenderTarget->SizeX * Pivot.X;
+				float V = -RenderTarget->SizeY * Pivot.Y;
+				float UL = RenderTarget->SizeX * ( 1.0f - Pivot.X );
+				float VL = RenderTarget->SizeY * ( 1.0f - Pivot.Y );
 
 				int32 VertexIndices[4];
 
-				for (int32 ViewIndex = 0; ViewIndex < Views.Num(); ViewIndex++)
+				for ( int32 ViewIndex = 0; ViewIndex < Views.Num(); ViewIndex++ )
 				{
-					if (VisibilityMap & (1 << ViewIndex))
+					if ( VisibilityMap & ( 1 << ViewIndex ) )
 					{
 						VertexIndices[0] = MeshBuilder.AddVertex(FVector(U, V, 0), FVector2D(0, 0), FVector(1, 0, 0), FVector(0, 1, 0), FVector(0, 0, 1), FColor::White);
 						VertexIndices[1] = MeshBuilder.AddVertex(FVector(U, VL, 0), FVector2D(0, 1), FVector(1, 0, 0), FVector(0, 1, 0), FVector(0, 0, 1), FColor::White);
@@ -98,15 +99,15 @@ public:
 	 */
 	virtual void DrawDynamicElements(FPrimitiveDrawInterface* PDI, const FSceneView* View) override
 	{
-		if( RenderTarget )
+		if ( RenderTarget )
 		{
 			FTextureResource* TextureResource = RenderTarget->Resource;
-			if(TextureResource)
+			if ( TextureResource )
 			{
-				float U = 0.0f;
-				float V = 0.0f;
-				float UL = RenderTarget->SizeX;
-				float VL = RenderTarget->SizeY;
+				float U = -RenderTarget->SizeX * Pivot.X;
+				float V = -RenderTarget->SizeY * Pivot.Y;
+				float UL = RenderTarget->SizeX * ( 1.0f - Pivot.X );
+				float VL = RenderTarget->SizeY * ( 1.0f - Pivot.Y );
 
 				FDynamicMeshBuilder MeshBuilder;
 
@@ -168,6 +169,7 @@ public:
 
 private:
 	FVector Origin;
+	FVector2D Pivot;
 	ISlate3DRenderer& Renderer;
 	UTextureRenderTarget2D* RenderTarget;
 	UMaterialInstanceDynamic* MaterialInstance;
@@ -325,16 +327,35 @@ UWidgetComponent::UWidgetComponent( const class FObjectInitializer& PCIP )
 	LastLocalHitLocation = FVector2D::ZeroVector;
 	//bGenerateOverlapEvents = false;
 	bUseEditorCompositing = false;
+
+	Space = EWidgetSpace::World;
+	Pivot = FVector2D(0.5, 0.5);
+
+	PrimaryComponentTick.TickGroup = TG_PostUpdateWork;
 }
 
 FPrimitiveSceneProxy* UWidgetComponent::CreateSceneProxy()
 {
-	return new FWidget3DSceneProxy(this,*Renderer);
+	if ( Space != EWidgetSpace::Screen )
+	{
+		return new FWidget3DSceneProxy(this, *Renderer);
+	}
+	else
+	{
+		return nullptr;
+	}
 }
 
 FBoxSphereBounds UWidgetComponent::CalcBounds(const FTransform & LocalToWorld) const
 {
-	return FBoxSphereBounds( FVector( DrawSize.X/2.0f, DrawSize.Y/2.0f, .5f), FVector(DrawSize.X/2,DrawSize.Y/2,1.0f), DrawSize.Size()/2 ).TransformBy( LocalToWorld );
+	if ( Space != EWidgetSpace::Screen )
+	{
+		return FBoxSphereBounds(FVector(DrawSize.X / 2.0f, DrawSize.Y / 2.0f, .5f), FVector(DrawSize.X / 2, DrawSize.Y / 2, 1.0f), DrawSize.Size() / 2).TransformBy(LocalToWorld);
+	}
+	else
+	{
+		return FBoxSphereBounds(FSphere()).TransformBy(LocalToWorld);
+	}
 }
 
 UBodySetup* UWidgetComponent::GetBodySetup() 
@@ -345,61 +366,76 @@ UBodySetup* UWidgetComponent::GetBodySetup()
 
 FCollisionShape UWidgetComponent::GetCollisionShape(float Inflation) const
 {
-	FVector	Extent = ( FVector( DrawSize.Y, DrawSize.Y, 1.0f ) * ComponentToWorld.GetScale3D() ) + Inflation;
-	if( Inflation < 0.0f )
+	if ( Space != EWidgetSpace::Screen )
 	{
-		// Don't shrink below zero size.
-		Extent = Extent.ComponentMax(FVector::ZeroVector);
-	}
+		FVector	Extent = ( FVector(DrawSize.Y, DrawSize.Y, 1.0f) * ComponentToWorld.GetScale3D() ) + Inflation;
+		if ( Inflation < 0.0f )
+		{
+			// Don't shrink below zero size.
+			Extent = Extent.ComponentMax(FVector::ZeroVector);
+		}
 
-	return FCollisionShape::MakeBox(Extent);
+		return FCollisionShape::MakeBox(Extent);
+	}
+	else
+	{
+		return FCollisionShape::MakeBox(FVector::ZeroVector);
+	}
 }
 
 void UWidgetComponent::OnRegister()
 {
 	Super::OnRegister();
 
-	if( GetWorld()->IsGameWorld() )
-	{
-		TSharedPtr<SViewport> GameViewportWidget = GEngine->GetGameViewportWidget();
+#if !UE_SERVER
 
-		if( GameViewportWidget.IsValid() )
+	if ( Space != EWidgetSpace::Screen )
+	{
+		if ( GetWorld()->IsGameWorld() )
 		{
-			TSharedPtr<ICustomHitTestPath> CustomHitTestPath = GameViewportWidget->GetCustomHitTestPath();
-			if( !CustomHitTestPath.IsValid() )
+			TSharedPtr<SViewport> GameViewportWidget = GEngine->GetGameViewportWidget();
+
+			if ( GameViewportWidget.IsValid() )
 			{
-				CustomHitTestPath = MakeShareable(new FWidget3DHitTester(GetWorld()));
-				GameViewportWidget->SetCustomHitTestPath( CustomHitTestPath );
+				TSharedPtr<ICustomHitTestPath> CustomHitTestPath = GameViewportWidget->GetCustomHitTestPath();
+				if ( !CustomHitTestPath.IsValid() )
+				{
+					CustomHitTestPath = MakeShareable(new FWidget3DHitTester(GetWorld()));
+					GameViewportWidget->SetCustomHitTestPath(CustomHitTestPath);
+				}
+
+				TSharedPtr<FWidget3DHitTester> WidgetHitTester = StaticCastSharedPtr<FWidget3DHitTester>(CustomHitTestPath);
+				if ( WidgetHitTester->GetWorld() == GetWorld() )
+				{
+					WidgetHitTester->RegisterWidgetComponent(this);
+				}
 			}
+		}
 
-			TSharedPtr<FWidget3DHitTester> WidgetHitTester = StaticCastSharedPtr<FWidget3DHitTester>(CustomHitTestPath);
-			if( WidgetHitTester->GetWorld() == GetWorld() )
+		if ( !MaterialInstance )
+		{
+			UMaterialInterface* Parent = nullptr;
+			if ( bIsOpaque )
 			{
-				WidgetHitTester->RegisterWidgetComponent( this );
+				Parent = bIsTwoSided ? OpaqueMaterial : OpaqueMaterial_OneSided;
 			}
+			else
+			{
+				Parent = bIsTwoSided ? TranslucentMaterial : TranslucentMaterial_OneSided;
+			}
+
+			MaterialInstance = UMaterialInstanceDynamic::Create(Parent, this);
 		}
 	}
 
-#if !UE_SERVER
+	InitWidget();
 
-	if( !MaterialInstance )
+	if ( Space != EWidgetSpace::Screen )
 	{
-		UMaterialInterface* Parent = nullptr;
-		if (bIsOpaque)
+		if ( !Renderer.IsValid() )
 		{
-			Parent = bIsTwoSided ? OpaqueMaterial : OpaqueMaterial_OneSided;
+			Renderer = FModuleManager::Get().GetModuleChecked<ISlateRHIRendererModule>("SlateRHIRenderer").CreateSlate3DRenderer();
 		}
-		else
-		{
-			Parent = bIsTwoSided ? TranslucentMaterial : TranslucentMaterial_OneSided;
-		}
-
-		MaterialInstance = UMaterialInstanceDynamic::Create(Parent, this);
-	}
-
-	if( !Renderer.IsValid() )
-	{
-		Renderer = FModuleManager::Get().GetModuleChecked<ISlateRHIRendererModule>("SlateRHIRenderer").CreateSlate3DRenderer();
 	}
 
 #endif // !UE_SERVER
@@ -427,10 +463,11 @@ void UWidgetComponent::OnUnregister()
 		}
 	}
 
-	if( Widget )
+	if ( Widget )
 	{
+		Widget->RemoveFromParent();
 		Widget->MarkPendingKill();
-		Widget = NULL;
+		Widget = nullptr;
 	}
 
 	SlateWidget.Reset();
@@ -440,65 +477,119 @@ void UWidgetComponent::OnUnregister()
 void UWidgetComponent::DestroyComponent()
 {
 	Super::DestroyComponent();
+
 	Renderer.Reset();
+
+	if ( Widget )
+	{
+		Widget->RemoveFromParent();
+		Widget->MarkPendingKill();
+		Widget = nullptr;
+	}
+
+	SlateWidget.Reset();
 }
 
 void UWidgetComponent::TickComponent(float DeltaTime, enum ELevelTick TickType, FActorComponentTickFunction *ThisTickFunction)
 {
 #if !UE_SERVER
 
-	const float RenderTimeThreshold = .5f;
+	UpdateWidget();
 
-	if( IsVisible() && GetWorld()->TimeSince(LastRenderTime) <= RenderTimeThreshold ) // Don't bother ticking if it hasn't been rendered recently
+	if ( Widget == nullptr )
 	{
-		UpdateWidget();
+		return;
+	}
 
-		if (!SlateWidget.IsValid())
+	if ( Space != EWidgetSpace::Screen )
+	{
+		if ( SlateWidget.IsValid() )
 		{
 			return;
 		}
 
-		SlateWidget->SlatePrepass();
+		const float RenderTimeThreshold = .5f;
+		if ( IsVisible() && GetWorld()->TimeSince(LastRenderTime) <= RenderTimeThreshold ) // Don't bother ticking if it hasn't been rendered recently
+		{
+			SlateWidget->SlatePrepass();
 
-		FGeometry WindowGeometry = FGeometry::MakeRoot( DrawSize, FSlateLayoutTransform() );
+			FGeometry WindowGeometry = FGeometry::MakeRoot(DrawSize, FSlateLayoutTransform());
 
-		SlateWidget->TickWidgetsRecursively(WindowGeometry, FApp::GetCurrentTime(), DeltaTime);
+			SlateWidget->TickWidgetsRecursively(WindowGeometry, FApp::GetCurrentTime(), DeltaTime);
 
-		// Ticking can cause geometry changes.  Recompute
-		SlateWidget->SlatePrepass();
+			// Ticking can cause geometry changes.  Recompute
+			SlateWidget->SlatePrepass();
 
-		// Get the free buffer & add our virtual window
-		FSlateDrawBuffer& DrawBuffer = Renderer->GetDrawBuffer();
-		FSlateWindowElementList& WindowElementList = DrawBuffer.AddWindowElementList(SlateWidget.ToSharedRef());
+			// Get the free buffer & add our virtual window
+			FSlateDrawBuffer& DrawBuffer = Renderer->GetDrawBuffer();
+			FSlateWindowElementList& WindowElementList = DrawBuffer.AddWindowElementList(SlateWidget.ToSharedRef());
 
-		int32 MaxLayerId = 0;
-		{
-			// Prepare the test grid 
-			HitTestGrid->ClearGridForNewFrame(WindowGeometry.GetClippingRect());
-
-			// Paint the window
-			MaxLayerId = SlateWidget->Paint(
-				FPaintArgs(SlateWidget.ToSharedRef(), *HitTestGrid, FVector2D::ZeroVector, FSlateApplication::Get().GetCurrentTime(), FSlateApplication::Get().GetDeltaTime()),
-				WindowGeometry, WindowGeometry.GetClippingRect(),
-				WindowElementList,
-				0,
-				FWidgetStyle(),
-				SlateWidget->IsEnabled());
-		}
+			int32 MaxLayerId = 0;
+			{
+				// Prepare the test grid 
+				HitTestGrid->ClearGridForNewFrame(WindowGeometry.GetClippingRect());
+
+				// Paint the window
+				MaxLayerId = SlateWidget->Paint(
+					FPaintArgs(SlateWidget.ToSharedRef(), *HitTestGrid, FVector2D::ZeroVector, FSlateApplication::Get().GetCurrentTime(), FSlateApplication::Get().GetDeltaTime()),
+					WindowGeometry, WindowGeometry.GetClippingRect(),
+					WindowElementList,
+					0,
+					FWidgetStyle(),
+					SlateWidget->IsEnabled());
+			}
 
-		Renderer->DrawWindow_GameThread(DrawBuffer);
+			Renderer->DrawWindow_GameThread(DrawBuffer);
 
-		UpdateRenderTarget();
+			UpdateRenderTarget();
+
+			// Enqueue a command to unlock the draw buffer after all windows have been drawn
+			ENQUEUE_UNIQUE_RENDER_COMMAND_THREEPARAMETER(UWidgetComponentRenderToTexture,
+				FSlateDrawBuffer&, InDrawBuffer, DrawBuffer,
+				UTextureRenderTarget2D*, InRenderTarget, RenderTarget,
+				ISlate3DRenderer*, InRenderer, Renderer.Get(),
+				{
+					InRenderer->DrawWindowToTarget_RenderThread(RHICmdList, InRenderTarget, InDrawBuffer);
+				});
+
+		}
+	}
+	else
+	{
+		if ( Widget && !Widget->IsDesignTime() )
+		{
+			ULocalPlayer* const TargetPlayer = GEngine->GetLocalPlayerFromControllerId(GetWorld(), 0);
+			APlayerController* PlayerController = TargetPlayer->PlayerController;
 
-		// Enqueue a command to unlock the draw buffer after all windows have been drawn
-		ENQUEUE_UNIQUE_RENDER_COMMAND_THREEPARAMETER(UWidgetComponentRenderToTexture,
-			FSlateDrawBuffer&, InDrawBuffer, DrawBuffer,
-			UTextureRenderTarget2D*, InRenderTarget, RenderTarget,
-			ISlate3DRenderer*, InRenderer, Renderer.Get(),
+			if ( TargetPlayer && PlayerController )
 			{
-				InRenderer->DrawWindowToTarget_RenderThread( RHICmdList, InRenderTarget, InDrawBuffer );	
-			});
+				if ( !Widget->IsInViewport() )
+				{
+					Widget->AddToViewport();
+				}
 
+				FVector WorldLocation = GetComponentLocation();
+
+				FVector2D ScreenLocation;
+				bool bProjected = PlayerController->ProjectWorldLocationToScreen(WorldLocation, ScreenLocation);
+				if ( bProjected )
+				{
+					Widget->SetVisibility(ESlateVisibility::SelfHitTestInvisible);
+				}
+				else
+				{
+					Widget->SetVisibility(ESlateVisibility::Hidden);
+				}
+				Widget->SetPositionInViewport(ScreenLocation);
+				Widget->SetAlignmentInViewport(Pivot);
+			}
+
+			// If the component isn't visible, and the widget is on the viewport, remove it.
+			if ( !IsVisible() && Widget->IsInViewport() )
+			{
+				Widget->RemoveFromParent();
+			}
+		}
 	}
 
 #endif // !UE_SERVER
@@ -587,49 +678,55 @@ void UWidgetComponent::PostEditChangeProperty(FPropertyChangedEvent& PropertyCha
 }
 #endif
 
-void UWidgetComponent::UpdateWidget()
+void UWidgetComponent::InitWidget()
 {
 	// Don't do any work if Slate is not initialized
-	if(FSlateApplication::IsInitialized())
-	{		
-		if(WidgetClass && !Widget && GetWorld() )
+	if ( FSlateApplication::IsInitialized() )
+	{
+		if ( WidgetClass && Widget == nullptr && GetWorld() )
 		{
-			Widget = CreateWidget<UUserWidget>( GetWorld(), WidgetClass );
+			Widget = CreateWidget<UUserWidget>(GetWorld(), WidgetClass);
 		}
-		else if( !WidgetClass && Widget )
+		else if ( !WidgetClass && Widget )
 		{
 			Widget = nullptr;
 		}
 
-		if(!SlateWidget.IsValid())
-		{
-			SlateWidget = SNew(SVirtualWindow).Size(DrawSize);
-
-		}
-	
-		if( !HitTestGrid.IsValid() )
+		if ( Widget && !GetWorld()->IsGameWorld() )
 		{
-			HitTestGrid = MakeShareable( new FHittestGrid );
+			// Prevent native ticking of editor component previews
+			Widget->SetIsDesignTime(true);
 		}
+	}
+}
 
+void UWidgetComponent::UpdateWidget()
+{
+	// Don't do any work if Slate is not initialized
+	if ( FSlateApplication::IsInitialized() )
+	{
+		if ( Space != EWidgetSpace::Screen )
 		{
-			SlateWidget->Resize( DrawSize );
-
-			if( Widget )
+			if ( !SlateWidget.IsValid() )
 			{
-				SlateWidget->SetContent( Widget->TakeWidget() );
+				SlateWidget = SNew(SVirtualWindow).Size(DrawSize);
 			}
-			else
+
+			if ( !HitTestGrid.IsValid() )
 			{
-				SlateWidget->SetContent( SNullWidget::NullWidget );
+				HitTestGrid = MakeShareable(new FHittestGrid);
 			}
 
-		}
+			SlateWidget->Resize(DrawSize);
 
-		if( Widget && !GetWorld()->IsGameWorld() )
-		{
-			// Prevent native ticking of editor component previews
-			Widget->SetIsDesignTime( true );
+			if ( Widget )
+			{
+				SlateWidget->SetContent(Widget->TakeWidget());
+			}
+			else
+			{
+				SlateWidget->SetContent(SNullWidget::NullWidget);
+			}
 		}
 	}
 }
diff --git a/Engine/Source/Runtime/UMG/Public/Components/WidgetComponent.h b/Engine/Source/Runtime/UMG/Public/Components/WidgetComponent.h
--- a/Engine/Source/Runtime/UMG/Public/Components/WidgetComponent.h
+++ b/Engine/Source/Runtime/UMG/Public/Components/WidgetComponent.h
@@ -6,6 +6,13 @@
 
 struct FVirtualPointerPosition;
 
+UENUM()
+enum class EWidgetSpace : uint8
+{
+	World,
+	Screen
+};
+
 /**
 * Beware! This feature is experimental and may be substantially changed or removed in future releases.
 * A 3D instance of a Widget Blueprint that can be interacted with in the world.
@@ -34,7 +41,10 @@ public:
 	virtual void PostEditChangeProperty(FPropertyChangedEvent& PropertyChangedEvent);
 #endif
 
-	/** Ensures the user widget is initialized and updates its size and content. */
+	/** Ensures the user widget is initialized */
+	void InitWidget();
+
+	/** Ensures the 3d window is created its size and content. */
 	void UpdateWidget();
 
 	/** Ensure the render target is initialized and updates it if needed. */
@@ -73,15 +83,26 @@ public:
 
 	/** @return True if the component is opaque */
 	bool IsOpaque() const { return bIsOpaque; }
+
+	/** @return The pivot point where the UI is rendered about the origin. */
+	FVector2D GetPivot() const { return Pivot; }
 	
 private:
+	/** The class of User Widget to create and display an instance of */
+	UPROPERTY(EditAnywhere, Category=UI)
+	EWidgetSpace Space;
+
 	/** The class of User Widget to create and display an instance of */
 	UPROPERTY(EditAnywhere, Category=UI)
 	TSubclassOf<UUserWidget> WidgetClass;
 	
-	/** The size of the displayed quad */
-	UPROPERTY(EditAnywhere, Category=UI)
+	/** The size of the displayed quad. */
+	UPROPERTY(EditAnywhere, Category="UI")
 	FIntPoint DrawSize;
+
+	/** The Alignment/Pivot point that the widget is placed at relative to the position. */
+	UPROPERTY(EditAnywhere, Category=UI)
+	FVector2D Pivot;
 	
 	/** The maximum distance from which a player can interact with this widget */
 	UPROPERTY(EditAnywhere, Category=UI, meta=(ClampMin="0.0", UIMax="5000.0", ClampMax="100000.0"))
