UMG - Adding some getters to a few slot types, exposing a few more things as readonly to blueprints.

diff --git a/Engine/Source/Runtime/UMG/Private/Components/CanvasPanelSlot.cpp b/Engine/Source/Runtime/UMG/Private/Components/CanvasPanelSlot.cpp
--- a/Engine/Source/Runtime/UMG/Private/Components/CanvasPanelSlot.cpp
+++ b/Engine/Source/Runtime/UMG/Private/Components/CanvasPanelSlot.cpp
@@ -7,7 +7,7 @@
 
 UCanvasPanelSlot::UCanvasPanelSlot(const FObjectInitializer& ObjectInitializer)
 	: Super(ObjectInitializer)
-	, Slot(NULL)
+	, Slot(nullptr)
 {
 	LayoutData.Offsets = FMargin(0, 0, 100, 30);
 	LayoutData.Anchors = FAnchors(0.0f, 0.0f);
@@ -20,14 +20,14 @@ void UCanvasPanelSlot::ReleaseSlateResources(bool bReleaseChildren)
 {
 	Super::ReleaseSlateResources(bReleaseChildren);
 
-	Slot = NULL;
+	Slot = nullptr;
 }
 
 void UCanvasPanelSlot::BuildSlot(TSharedRef<SConstraintCanvas> Canvas)
 {
 	Slot = &Canvas->AddSlot()
 		[
-			Content == NULL ? SNullWidget::NullWidget : Content->TakeWidget()
+			Content == nullptr ? SNullWidget::NullWidget : Content->TakeWidget()
 		];
 
 	SynchronizeProperties();
@@ -124,6 +124,23 @@ bool UCanvasPanelSlot::CanResize(const FVector2D& Direction) const
 	return true;
 }
 
+void UCanvasPanelSlot::SetLayout(const FAnchorData& InLayoutData)
+{
+	LayoutData = InLayoutData;
+
+	if ( Slot )
+	{
+		Slot->Offset(LayoutData.Offsets);
+		Slot->Anchors(LayoutData.Anchors);
+		Slot->Alignment(LayoutData.Alignment);
+	}
+}
+
+FAnchorData UCanvasPanelSlot::GetLayout() const
+{
+	return LayoutData;
+}
+
 void UCanvasPanelSlot::SetPosition(FVector2D InPosition)
 {
 	LayoutData.Offsets.Left = InPosition.X;
@@ -135,6 +152,17 @@ void UCanvasPanelSlot::SetPosition(FVector2D InPosition)
 	}
 }
 
+FVector2D UCanvasPanelSlot::GetPosition() const
+{
+	if ( Slot )
+	{
+		FMargin Offsets = Slot->OffsetAttr.Get();
+		return FVector2D(Offsets.Left, Offsets.Top);
+	}
+
+	return FVector2D(LayoutData.Offsets.Left, LayoutData.Offsets.Top);
+}
+
 void UCanvasPanelSlot::SetSize(FVector2D InSize)
 {
 	LayoutData.Offsets.Right = InSize.X;
@@ -146,6 +174,17 @@ void UCanvasPanelSlot::SetSize(FVector2D InSize)
 	}
 }
 
+FVector2D UCanvasPanelSlot::GetSize() const
+{
+	if ( Slot )
+	{
+		FMargin Offsets = Slot->OffsetAttr.Get();
+		return FVector2D(Offsets.Right, Offsets.Bottom);
+	}
+
+	return FVector2D(LayoutData.Offsets.Right, LayoutData.Offsets.Bottom);
+}
+
 void UCanvasPanelSlot::SetOffsets(FMargin InOffset)
 {
 	LayoutData.Offsets = InOffset;
@@ -155,6 +194,16 @@ void UCanvasPanelSlot::SetOffsets(FMargin InOffset)
 	}
 }
 
+FMargin UCanvasPanelSlot::GetOffsets() const
+{
+	if ( Slot )
+	{
+		return Slot->OffsetAttr.Get();
+	}
+
+	return LayoutData.Offsets;
+}
+
 void UCanvasPanelSlot::SetAnchors(FAnchors InAnchors)
 {
 	LayoutData.Anchors = InAnchors;
@@ -164,6 +213,16 @@ void UCanvasPanelSlot::SetAnchors(FAnchors InAnchors)
 	}
 }
 
+FAnchors UCanvasPanelSlot::GetAnchors() const
+{
+	if ( Slot )
+	{
+		return Slot->AnchorsAttr.Get();
+	}
+
+	return LayoutData.Anchors;
+}
+
 void UCanvasPanelSlot::SetAlignment(FVector2D InAlignment)
 {
 	LayoutData.Alignment = InAlignment;
@@ -173,6 +232,16 @@ void UCanvasPanelSlot::SetAlignment(FVector2D InAlignment)
 	}
 }
 
+FVector2D UCanvasPanelSlot::GetAlignment() const
+{
+	if ( Slot )
+	{
+		return Slot->AlignmentAttr.Get();
+	}
+
+	return LayoutData.Alignment;
+}
+
 void UCanvasPanelSlot::SetAutoSize(bool InbAutoSize)
 {
 	bAutoSize = InbAutoSize;
@@ -182,6 +251,16 @@ void UCanvasPanelSlot::SetAutoSize(bool InbAutoSize)
 	}
 }
 
+bool UCanvasPanelSlot::GetAutoSize() const
+{
+	if ( Slot )
+	{
+		return Slot->AutoSizeAttr.Get();
+	}
+
+	return bAutoSize;
+}
+
 void UCanvasPanelSlot::SetZOrder(int32 InZOrder)
 {
 	ZOrder = InZOrder;
@@ -191,6 +270,16 @@ void UCanvasPanelSlot::SetZOrder(int32 InZOrder)
 	}
 }
 
+int32 UCanvasPanelSlot::GetZOrder() const
+{
+	if ( Slot )
+	{
+		return Slot->ZOrderAttr.Get();
+	}
+
+	return ZOrder;
+}
+
 void UCanvasPanelSlot::SetMinimum(FVector2D InMinimumAnchors)
 {
 	LayoutData.Anchors.Minimum = InMinimumAnchors;
diff --git a/Engine/Source/Runtime/UMG/Public/Components/ButtonSlot.h b/Engine/Source/Runtime/UMG/Public/Components/ButtonSlot.h
--- a/Engine/Source/Runtime/UMG/Public/Components/ButtonSlot.h
+++ b/Engine/Source/Runtime/UMG/Public/Components/ButtonSlot.h
@@ -15,15 +15,15 @@ class UMG_API UButtonSlot : public UPanelSlot
 public:
 	
 	/** The padding area between the slot and the content it contains. */
-	UPROPERTY(EditDefaultsOnly, Category="Layout (Button Slot)")
+	UPROPERTY(EditDefaultsOnly, BlueprintReadOnly, Category="Layout (Button Slot)")
 	FMargin Padding;
 
 	/** The alignment of the object horizontally. */
-	UPROPERTY(EditDefaultsOnly, Category="Layout (Button Slot)")
+	UPROPERTY(EditDefaultsOnly, BlueprintReadOnly, Category="Layout (Button Slot)")
 	TEnumAsByte<EHorizontalAlignment> HorizontalAlignment;
 
 	/** The alignment of the object vertically. */
-	UPROPERTY(EditDefaultsOnly, Category="Layout (Button Slot)")
+	UPROPERTY(EditDefaultsOnly, BlueprintReadOnly, Category="Layout (Button Slot)")
 	TEnumAsByte<EVerticalAlignment> VerticalAlignment;
 
 public:
diff --git a/Engine/Source/Runtime/UMG/Public/Components/CanvasPanelSlot.h b/Engine/Source/Runtime/UMG/Public/Components/CanvasPanelSlot.h
--- a/Engine/Source/Runtime/UMG/Public/Components/CanvasPanelSlot.h
+++ b/Engine/Source/Runtime/UMG/Public/Components/CanvasPanelSlot.h
@@ -33,6 +33,8 @@ class UMG_API UCanvasPanelSlot : public UPanelSlot
 {
 	GENERATED_UCLASS_BODY()
 
+public:
+
 	/** The anchoring information for the slot */
 	UPROPERTY(EditDefaultsOnly, BlueprintReadOnly, Category="Layout (Canvas Slot)")
 	FAnchorData LayoutData;
@@ -45,44 +47,73 @@ class UMG_API UCanvasPanelSlot : public UPanelSlot
 	UPROPERTY(EditDefaultsOnly, BlueprintReadOnly, Category="Layout (Canvas Slot)")
 	int32 ZOrder;
 
-	void BuildSlot(TSharedRef<SConstraintCanvas> Canvas);
-
-	virtual void SetDesiredPosition(FVector2D InPosition) override;
-
-	virtual void SetDesiredSize(FVector2D InSize) override;
-
-	virtual void Resize(const FVector2D& Direction, const FVector2D& Amount) override;
+public:
+	/** Sets the layout data of the slot */
+	UFUNCTION(BlueprintCallable, Category="Layout (Canvas Slot)")
+	void SetLayout(const FAnchorData& InLayoutData);
 
-	virtual bool CanResize(const FVector2D& Direction) const override;
+	/** Gets the layout data of the slot */
+	UFUNCTION(BlueprintCallable, Category="Layout (Canvas Slot)")
+	FAnchorData GetLayout() const;
 
 	/** Sets the position of the slot */
 	UFUNCTION(BlueprintCallable, Category="Layout (Canvas Slot)")
 	void SetPosition(FVector2D InPosition);
 
+	/** Gets the position of the slot */
+	UFUNCTION(BlueprintCallable, Category="Layout (Canvas Slot)")
+	FVector2D GetPosition() const;
+
 	/** Sets the size of the slot */
 	UFUNCTION(BlueprintCallable, Category="Layout (Canvas Slot)")
 	void SetSize(FVector2D InSize);
 
-	/** Sets the position of the slot */
+	/** Gets the size of the slot */
+	UFUNCTION(BlueprintCallable, Category="Layout (Canvas Slot)")
+	FVector2D GetSize() const;
+
+	/** Sets the offset data of the slot, which could be position and size, or margins depending on the anchor points */
 	UFUNCTION(BlueprintCallable, Category="Layout (Canvas Slot)")
 	void SetOffsets(FMargin InOffset);
+
+	/** Gets the offset data of the slot, which could be position and size, or margins depending on the anchor points */
+	UFUNCTION(BlueprintCallable, Category="Layout (Canvas Slot)")
+	FMargin GetOffsets() const;
 	
 	/** Sets the anchors on the slot */
 	UFUNCTION(BlueprintCallable, Category="Layout (Canvas Slot)")
 	void SetAnchors(FAnchors InAnchors);
 
+	/** Gets the anchors on the slot */
+	UFUNCTION(BlueprintCallable, Category="Layout (Canvas Slot)")
+	FAnchors GetAnchors() const;
+
 	/** Sets the alignment on the slot */
 	UFUNCTION(BlueprintCallable, Category="Layout (Canvas Slot)")
 	void SetAlignment(FVector2D InAlignment);
 
-	/** Sets the slot to be auto-sized */
+	/** Gets the alignment on the slot */
+	UFUNCTION(BlueprintCallable, Category="Layout (Canvas Slot)")
+	FVector2D GetAlignment() const;
+
+	/** Sets if the slot to be auto-sized */
 	UFUNCTION(BlueprintCallable, Category="Layout (Canvas Slot)")
 	void SetAutoSize(bool InbAutoSize);
 
+	/** Gets if the slot to be auto-sized */
+	UFUNCTION(BlueprintCallable, Category="Layout (Canvas Slot)")
+	bool GetAutoSize() const;
+
 	/** Sets the z-order on the slot */
 	UFUNCTION(BlueprintCallable, Category="Layout (Canvas Slot)")
 	void SetZOrder(int32 InZOrder);
 
+	/** Gets the z-order on the slot */
+	UFUNCTION(BlueprintCallable, Category="Layout (Canvas Slot)")
+	int32 GetZOrder() const;
+
+public:
+
 	/** Sets the anchors on the slot */
 	UFUNCTION()
 	void SetMinimum(FVector2D InMinimumAnchors);
@@ -91,6 +122,18 @@ class UMG_API UCanvasPanelSlot : public UPanelSlot
 	UFUNCTION()
 	void SetMaximum(FVector2D InMaximumAnchors);
 
+public:
+
+	void BuildSlot(TSharedRef<SConstraintCanvas> Canvas);
+
+	virtual void SetDesiredPosition(FVector2D InPosition) override;
+
+	virtual void SetDesiredSize(FVector2D InSize) override;
+
+	virtual void Resize(const FVector2D& Direction, const FVector2D& Amount) override;
+
+	virtual bool CanResize(const FVector2D& Direction) const override;
+
 	// UPanelSlot interface
 	virtual void SynchronizeProperties() override;
 	// End of UPanelSlot interface
diff --git a/Engine/Source/Runtime/UMG/Public/Components/HorizontalBoxSlot.h b/Engine/Source/Runtime/UMG/Public/Components/HorizontalBoxSlot.h
--- a/Engine/Source/Runtime/UMG/Public/Components/HorizontalBoxSlot.h
+++ b/Engine/Source/Runtime/UMG/Public/Components/HorizontalBoxSlot.h
@@ -14,19 +14,21 @@ class UMG_API UHorizontalBoxSlot : public UPanelSlot
 public:
 	
 	/** The amount of padding between the slots parent and the content. */
-	UPROPERTY(EditDefaultsOnly, Category="Layout (Horizontal Box Slot)")
+	UPROPERTY(EditDefaultsOnly, BlueprintReadOnly, Category="Layout (Horizontal Box Slot)")
 	FMargin Padding;
 
 	/** How much space this slot should occupy in the direction of the panel. */
-	UPROPERTY(EditDefaultsOnly, Category="Layout (Horizontal Box Slot)")
+	UPROPERTY(EditDefaultsOnly, BlueprintReadOnly, Category="Layout (Horizontal Box Slot)")
 	FSlateChildSize Size;
 
-	UPROPERTY(EditDefaultsOnly, Category="Layout (Horizontal Box Slot)")
+	UPROPERTY(EditDefaultsOnly, BlueprintReadOnly, Category="Layout (Horizontal Box Slot)")
 	TEnumAsByte<EHorizontalAlignment> HorizontalAlignment;
 
-	UPROPERTY(EditDefaultsOnly, Category="Layout (Horizontal Box Slot)")
+	UPROPERTY(EditDefaultsOnly, BlueprintReadOnly, Category="Layout (Horizontal Box Slot)")
 	TEnumAsByte<EVerticalAlignment> VerticalAlignment;
 
+public:
+
 	UFUNCTION(BlueprintCallable, Category="Layout(Horizontal Box Slot)")
 	void SetPadding(FMargin InPadding);
 
