UMG - Fixing a qnan-issue by initializing the bounds correctly in the widget component.

diff --git a/Engine/Source/Runtime/UMG/Private/Components/WidgetComponent.cpp b/Engine/Source/Runtime/UMG/Private/Components/WidgetComponent.cpp
--- a/Engine/Source/Runtime/UMG/Private/Components/WidgetComponent.cpp
+++ b/Engine/Source/Runtime/UMG/Private/Components/WidgetComponent.cpp
@@ -354,7 +354,7 @@ FBoxSphereBounds UWidgetComponent::CalcBounds(const FTransform & LocalToWorld) c
 	}
 	else
 	{
-		return FBoxSphereBounds(FSphere()).TransformBy(LocalToWorld);
+		return FBoxSphereBounds(ForceInit).TransformBy(LocalToWorld);
 	}
 }
 
