Fixing mac build initializing the fields in the proper order.

diff --git a/Engine/Source/Runtime/UMG/Private/Binding/DynamicPropertyPath.cpp b/Engine/Source/Runtime/UMG/Private/Binding/DynamicPropertyPath.cpp
--- a/Engine/Source/Runtime/UMG/Private/Binding/DynamicPropertyPath.cpp
+++ b/Engine/Source/Runtime/UMG/Private/Binding/DynamicPropertyPath.cpp
@@ -7,17 +7,17 @@
 
 FPropertyPathSegment::FPropertyPathSegment()
 	: Name(NAME_None)
+	, ArrayIndex(INDEX_NONE)
 	, Struct(nullptr)
 	, Field(nullptr)
-	, ArrayIndex(INDEX_NONE)
 {
 
 }
 
 FPropertyPathSegment::FPropertyPathSegment(FString SegmentName)
-	: Struct(nullptr)
+	: ArrayIndex(INDEX_NONE)
+	, Struct(nullptr)
 	, Field(nullptr)
-	, ArrayIndex(INDEX_NONE)
 {
 	// Parse the property name and (optional) array index
 	int32 ArrayPos = SegmentName.Find(TEXT("["));
