Build Fix: Adding header to fix non-unity build.

diff --git a/Engine/Source/Runtime/UMG/Private/Binding/BrushBinding.cpp b/Engine/Source/Runtime/UMG/Private/Binding/BrushBinding.cpp
--- a/Engine/Source/Runtime/UMG/Private/Binding/BrushBinding.cpp
+++ b/Engine/Source/Runtime/UMG/Private/Binding/BrushBinding.cpp
@@ -1,6 +1,7 @@
 // Copyright 1998-2014 Epic Games, Inc. All Rights Reserved.
 
 #include "UMGPrivatePCH.h"
+#include "WidgetBlueprintLibrary.h"
 #include "BrushBinding.h"
 
 #define LOCTEXT_NAMESPACE "UMG"
