Removed DrawDynamicElements, PreRenderView and dependencies.  These have been disabled since 11/04, and the GetDynamicMeshElements path is now the only path for dynamic mesh rendering.

diff --git a/Engine/Source/Runtime/UMG/Private/Components/WidgetComponent.cpp b/Engine/Source/Runtime/UMG/Private/Components/WidgetComponent.cpp
--- a/Engine/Source/Runtime/UMG/Private/Components/WidgetComponent.cpp
+++ b/Engine/Source/Runtime/UMG/Private/Components/WidgetComponent.cpp
@@ -91,49 +91,6 @@ public:
 		RenderBounds(Collector.GetPDI(0), ViewFamily.EngineShowFlags, GetBounds(), IsSelected());
 #endif
 	}
-	/** 
-	 * Draw the scene proxy as a dynamic element
-	 *
-	 * @param	PDI - draw interface to render to
-	 * @param	View - current view
-	 */
-	virtual void DrawDynamicElements(FPrimitiveDrawInterface* PDI, const FSceneView* View) override
-	{
-		if ( RenderTarget )
-		{
-			FTextureResource* TextureResource = RenderTarget->Resource;
-			if ( TextureResource )
-			{
-				float U = -RenderTarget->SizeX * Pivot.X;
-				float V = -RenderTarget->SizeY * Pivot.Y;
-				float UL = RenderTarget->SizeX * ( 1.0f - Pivot.X );
-				float VL = RenderTarget->SizeY * ( 1.0f - Pivot.Y );
-
-				FDynamicMeshBuilder MeshBuilder;
-
-				int32 VertexIndices[4];
-
-				VertexIndices[0] = MeshBuilder.AddVertex(FVector(U, V, 0), FVector2D(0, 0), FVector(1, 0, 0), FVector(0, 1, 0), FVector(0, 0, 1), FColor::White);
-				VertexIndices[1] = MeshBuilder.AddVertex(FVector(U, VL, 0), FVector2D(0, 1), FVector(1, 0, 0), FVector(0, 1, 0), FVector(0, 0, 1), FColor::White);
-				VertexIndices[2] = MeshBuilder.AddVertex(FVector(UL, VL, 0), FVector2D(1, 1), FVector(1, 0, 0), FVector(0, 1, 0), FVector(0, 0, 1), FColor::White);
-				VertexIndices[3] = MeshBuilder.AddVertex(FVector(UL, V, 0), FVector2D(1, 0), FVector(1, 0, 0), FVector(0, 1, 0), FVector(0, 0, 1), FColor::White);
-
-				MeshBuilder.AddTriangle(VertexIndices[0], VertexIndices[1], VertexIndices[2]);
-				MeshBuilder.AddTriangle(VertexIndices[0], VertexIndices[2], VertexIndices[3]);
-
-				MeshBuilder.Draw(PDI, GetLocalToWorld(), MaterialInstance->GetRenderProxy(IsSelected()), SDPG_World);
-			}
-		}
-		
-		// Visualize the collision geometry
-		/*FColor CollisionColor = FColor(157,149,223,255);
-		FTransform GeomTransform(GetLocalToWorld());
-		BodySetup->AggGeom.DrawAggGeom(PDI, GeomTransform, GetSelectionColor(CollisionColor, false, IsHovered()), nullptr, false, false, true );*/
-		
-#if !(UE_BUILD_SHIPPING || UE_BUILD_TEST)
-		RenderBounds(PDI, View->Family->EngineShowFlags, GetBounds(), IsSelected());
-#endif
-	}
 
 	virtual FPrimitiveViewRelevance GetViewRelevance(const FSceneView* View) override
 	{
