Fix logic bug in widget component preventing any component from being rendered

diff --git a/Engine/Source/Runtime/UMG/Private/Components/WidgetComponent.cpp b/Engine/Source/Runtime/UMG/Private/Components/WidgetComponent.cpp
--- a/Engine/Source/Runtime/UMG/Private/Components/WidgetComponent.cpp
+++ b/Engine/Source/Runtime/UMG/Private/Components/WidgetComponent.cpp
@@ -460,7 +460,7 @@ void UWidgetComponent::TickComponent(float DeltaTime, enum ELevelTick TickType,
 
 	if ( Space != EWidgetSpace::Screen )
 	{
-		if ( SlateWidget.IsValid() )
+		if ( !SlateWidget.IsValid() )
 		{
 			return;
 		}
