UMG - Widget components now have a load step to load old versions pivot data as 0,0.  Even though the default is now 0.5,0.5, to preserve old look for old assets.  Fixing several bugs with pivot and extent calculation in the component.  There were several bad calculations of extent, as well as issues with calling virtual functions in the constructor.

diff --git a/Engine/Source/Runtime/Core/Public/UObject/ObjectVersion.h b/Engine/Source/Runtime/Core/Public/UObject/ObjectVersion.h
--- a/Engine/Source/Runtime/Core/Public/UObject/ObjectVersion.h
+++ b/Engine/Source/Runtime/Core/Public/UObject/ObjectVersion.h
@@ -653,6 +653,8 @@ enum EUnrealEngineObjectUE4Version
 	VER_UE4_MONTAGE_BRANCHING_POINT_REMOVAL,
 	// Enforce const-correctness in Blueprint implementations of native C++ const class methods
 	VER_UE4_BLUEPRINT_ENFORCE_CONST_IN_FUNCTION_OVERRIDES,
+	// Added pivot to widget components, need to load old versions as a 0,0 pivot, new default is 0.5,0.5
+	VER_UE4_ADD_PIVOT_TO_WIDGET_COMPONENT,
 	// -----<new versions can be added before this line>-------------------------------------------------
 	// - this needs to be the last line (see note below)
 	VER_UE4_AUTOMATIC_VERSION_PLUS_ONE,
diff --git a/Engine/Source/Runtime/UMG/Private/Components/WidgetComponent.cpp b/Engine/Source/Runtime/UMG/Private/Components/WidgetComponent.cpp
--- a/Engine/Source/Runtime/UMG/Private/Components/WidgetComponent.cpp
+++ b/Engine/Source/Runtime/UMG/Private/Components/WidgetComponent.cpp
@@ -263,7 +263,7 @@ UWidgetComponent::UWidgetComponent( const FObjectInitializer& PCIP )
 	PrimaryComponentTick.bCanEverTick = true;
 	bTickInEditor = true;
 	
-	SetRelativeRotation( FRotator( 0.f, 0.f, 90.f ) );
+	RelativeRotation = FRotator(0.f, 0.f, 90.f);
 	
 	BodyInstance.SetCollisionProfileName(FName(TEXT("UI")));
 
@@ -307,7 +307,10 @@ FBoxSphereBounds UWidgetComponent::CalcBounds(const FTransform & LocalToWorld) c
 {
 	if ( Space != EWidgetSpace::Screen )
 	{
-		return FBoxSphereBounds(FVector(DrawSize.X / 2.0f, DrawSize.Y / 2.0f, .5f), FVector(DrawSize.X / 2, DrawSize.Y / 2, 1.0f), DrawSize.Size() / 2).TransformBy(LocalToWorld);
+		const FVector Origin = FVector(DrawSize.X / 2.0f - ( DrawSize.X * Pivot.X ), DrawSize.Y / 2.0f - ( DrawSize.Y * Pivot.Y ), .5f);
+		const FVector BoxExtent = FVector(DrawSize.X / 2.0f, DrawSize.Y / 2.0f, 1.0f);
+
+		return FBoxSphereBounds(Origin, BoxExtent, DrawSize.Size() / 2.0f).TransformBy(LocalToWorld);
 	}
 	else
 	{
@@ -325,14 +328,14 @@ FCollisionShape UWidgetComponent::GetCollisionShape(float Inflation) const
 {
 	if ( Space != EWidgetSpace::Screen )
 	{
-		FVector	Extent = ( FVector(DrawSize.Y, DrawSize.Y, 1.0f) * ComponentToWorld.GetScale3D() ) + Inflation;
+		FVector	BoxHalfExtent = ( FVector(DrawSize.X * 0.5f, DrawSize.Y * 0.5f, 1.0f) * ComponentToWorld.GetScale3D() ) + Inflation;
 		if ( Inflation < 0.0f )
 		{
 			// Don't shrink below zero size.
-			Extent = Extent.ComponentMax(FVector::ZeroVector);
+			BoxHalfExtent = BoxHalfExtent.ComponentMax(FVector::ZeroVector);
 		}
 
-		return FCollisionShape::MakeBox(Extent);
+		return FCollisionShape::MakeBox(BoxHalfExtent);
 	}
 	else
 	{
@@ -607,6 +610,7 @@ void UWidgetComponent::PostEditChangeProperty(FPropertyChangedEvent& PropertyCha
 	if( Property && PropertyChangedEvent.ChangeType != EPropertyChangeType::Interactive )
 	{
 		static FName DrawSizeName("DrawSize");
+		static FName PivotName("Pivot");
 		static FName WidgetClassName("WidgetClass");
 		static FName IsOpaqueName("bIsOpaque");
 		static FName IsTwoSidedName("bIsTwoSided");
@@ -618,7 +622,7 @@ void UWidgetComponent::PostEditChangeProperty(FPropertyChangedEvent& PropertyCha
 		{
 			UpdateWidget();
 
-			if( PropertyName == DrawSizeName )
+			if ( PropertyName == DrawSizeName || PropertyName == PivotName )
 			{
 				UpdateBodySetup(true);
 			}
@@ -755,8 +759,10 @@ void UWidgetComponent::UpdateBodySetup( bool bDrawSizeChanged )
 
 		FKBoxElem* BoxElem = BodySetup->AggGeom.BoxElems.GetData();
 
+		const FVector Origin = FVector(DrawSize.X / 2.0f - ( DrawSize.X * Pivot.X ), DrawSize.Y / 2.0f - ( DrawSize.Y * Pivot.Y ), .5f);
+
 		BoxElem->SetTransform(FTransform::Identity);
-		BoxElem->Center = FVector(DrawSize.X/2.0f, DrawSize.Y/2.0f, .5f);
+		BoxElem->Center = Origin;
 		BoxElem->X = DrawSize.X;
 		BoxElem->Y = DrawSize.Y;
 		BoxElem->Z = 1.0f;
@@ -793,4 +799,14 @@ TArray<FWidgetAndPointer> UWidgetComponent::GetHitWidgetPath( const FHitResult&
 TSharedPtr<SWidget> UWidgetComponent::GetSlateWidget() const
 {
 	return SlateWidget;
-}
\ No newline at end of file
+}
+
+void UWidgetComponent::PostLoad()
+{
+	Super::PostLoad();
+
+	if ( GetLinkerUE4Version() < VER_UE4_ADD_PIVOT_TO_WIDGET_COMPONENT )
+	{
+		Pivot = FVector2D(0, 0);
+	}
+}
diff --git a/Engine/Source/Runtime/UMG/Public/Components/WidgetComponent.h b/Engine/Source/Runtime/UMG/Public/Components/WidgetComponent.h
--- a/Engine/Source/Runtime/UMG/Public/Components/WidgetComponent.h
+++ b/Engine/Source/Runtime/UMG/Public/Components/WidgetComponent.h
@@ -37,6 +37,10 @@ public:
 	virtual FName GetComponentInstanceDataType() const override;
 	virtual void ApplyComponentInstanceData(FComponentInstanceDataBase* ComponentInstanceData) override;
 
+	// Begin UObject
+	virtual void PostLoad() override;
+	// End UObject
+
 #if WITH_EDITORONLY_DATA
 	virtual void PostEditChangeProperty(FPropertyChangedEvent& PropertyChangedEvent);
 #endif
