UMG - The UWidget and UWidgetTree now has a way to query the UWorld that they are currently part of, this fixes a crash in the MenuAnchor not having access to the current world when attempting to spawn a popup.

diff --git a/Engine/Source/Runtime/UMG/Private/Components/MenuAnchor.cpp b/Engine/Source/Runtime/UMG/Private/Components/MenuAnchor.cpp
--- a/Engine/Source/Runtime/UMG/Private/Components/MenuAnchor.cpp
+++ b/Engine/Source/Runtime/UMG/Private/Components/MenuAnchor.cpp
@@ -68,10 +68,12 @@ TSharedRef<SWidget> UMenuAnchor::HandleGetMenuContent()
 	{
 		if ( MenuClass != nullptr && !MenuClass->HasAnyClassFlags(CLASS_Abstract) )
 		{
-			UUserWidget* MenuWidget = CreateWidget<UUserWidget>(GetWorld(), MenuClass);
-			if ( MenuWidget )
+			if ( UWorld* World = GetWorld() )
 			{
-				SlateMenuWidget = MenuWidget->TakeWidget();
+				if ( UUserWidget* MenuWidget = CreateWidget<UUserWidget>(World, MenuClass) )
+				{
+					SlateMenuWidget = MenuWidget->TakeWidget();
+				}
 			}
 		}
 	}
diff --git a/Engine/Source/Runtime/UMG/Private/Components/Widget.cpp b/Engine/Source/Runtime/UMG/Private/Components/Widget.cpp
--- a/Engine/Source/Runtime/UMG/Private/Components/Widget.cpp
+++ b/Engine/Source/Runtime/UMG/Private/Components/Widget.cpp
@@ -505,6 +505,18 @@ void UWidget::SetIsDesignTime(bool bInDesignTime)
 	bDesignTime = bInDesignTime;
 }
 
+UWorld* UWidget::GetWorld() const
+{
+	// UWidget's are given world scope by their owning user widget.  We can get that through the widget tree that should
+	// be the outer of this widget.
+	if ( UWidgetTree* OwningTree = Cast<UWidgetTree>(GetOuter()) )
+	{
+		return OwningTree->GetWorld();
+	}
+
+	return nullptr;
+}
+
 void UWidget::PostLoad()
 {
 	Super::PostLoad();
diff --git a/Engine/Source/Runtime/UMG/Private/WidgetTree.cpp b/Engine/Source/Runtime/UMG/Private/WidgetTree.cpp
--- a/Engine/Source/Runtime/UMG/Private/WidgetTree.cpp
+++ b/Engine/Source/Runtime/UMG/Private/WidgetTree.cpp
@@ -10,6 +10,17 @@ UWidgetTree::UWidgetTree(const FObjectInitializer& ObjectInitializer)
 {
 }
 
+UWorld* UWidgetTree::GetWorld() const
+{
+	// The outer of a widget tree should be a user widget
+	if ( UUserWidget* OwningWidget = Cast<UUserWidget>(GetOuter()) )
+	{
+		return OwningWidget->GetWorld();
+	}
+
+	return nullptr;
+}
+
 UWidget* UWidgetTree::FindWidget(const FName& Name) const
 {
 	UWidget* FoundWidget = nullptr;
diff --git a/Engine/Source/Runtime/UMG/Public/Blueprint/WidgetTree.h b/Engine/Source/Runtime/UMG/Public/Blueprint/WidgetTree.h
--- a/Engine/Source/Runtime/UMG/Public/Blueprint/WidgetTree.h
+++ b/Engine/Source/Runtime/UMG/Public/Blueprint/WidgetTree.h
@@ -14,6 +14,10 @@ class UMG_API UWidgetTree : public UObject
 
 public:
 
+	// Begin UObject
+	virtual UWorld* GetWorld() const override;
+	// End UObject
+
 	/** Finds the widget in the tree by name. */
 	UWidget* FindWidget(const FName& Name) const;
 
diff --git a/Engine/Source/Runtime/UMG/Public/Components/Widget.h b/Engine/Source/Runtime/UMG/Public/Components/Widget.h
--- a/Engine/Source/Runtime/UMG/Public/Components/Widget.h
+++ b/Engine/Source/Runtime/UMG/Public/Components/Widget.h
@@ -283,6 +283,7 @@ public:
 	static TSubclassOf<class UPropertyBinding> FindBinderClassForDestination(UProperty* Property);
 
 	// Begin UObject
+	virtual UWorld* GetWorld() const override;
 	virtual void PostLoad() override;
 	// End UObject
 	
