UMG - Adding safer handling for the case when a user explictly attempts to destroy a UUserWidget.  It can now force the removal of the widget reference from the SObjectWidget, or the parent widget.

diff --git a/Engine/Source/Runtime/UMG/Private/Slate/SObjectWidget.cpp b/Engine/Source/Runtime/UMG/Private/Slate/SObjectWidget.cpp
--- a/Engine/Source/Runtime/UMG/Private/Slate/SObjectWidget.cpp
+++ b/Engine/Source/Runtime/UMG/Private/Slate/SObjectWidget.cpp
@@ -15,16 +15,29 @@ void SObjectWidget::Construct(const FArguments& InArgs, UUserWidget* InWidgetObj
 }
 
 SObjectWidget::~SObjectWidget(void)
+{
+	ResetWidget();
+}
+
+void SObjectWidget::ResetWidget()
 {
 	if ( UObjectInitialized() && WidgetObject )
 	{
 		// NOTE: When the SObjectWidget gets released we know that the User Widget has
 		// been removed from the slate widget hierarchy.  When this occurs, we need to 
-		// immediately release all slate widget widgets to deletion from taking n-frames
-		// due to widget nesting.
+		// immediately release all slate widget widgets to prevent deletion from taking
+		// n-frames due to widget nesting.
 		const bool bReleaseChildren = true;
 		WidgetObject->ReleaseSlateResources(bReleaseChildren);
+
+		WidgetObject = nullptr;
 	}
+
+	// Remove slate widget from our container
+	ChildSlot
+	[
+		SNullWidget::NullWidget
+	];
 }
 
 void SObjectWidget::AddReferencedObjects(FReferenceCollector& Collector)
diff --git a/Engine/Source/Runtime/UMG/Private/UserWidget.cpp b/Engine/Source/Runtime/UMG/Private/UserWidget.cpp
--- a/Engine/Source/Runtime/UMG/Private/UserWidget.cpp
+++ b/Engine/Source/Runtime/UMG/Private/UserWidget.cpp
@@ -68,6 +68,25 @@ void UUserWidget::Initialize()
 	}
 }
 
+void UUserWidget::BeginDestroy()
+{
+	Super::BeginDestroy();
+
+	// If anyone ever calls BeginDestroy explicitly on a widget we need to immediately remove it from
+	// the the parent as it may be owned currently by a slate widget.  As long as it's the viewport we're
+	// fine.
+	RemoveFromParent();
+
+	// If it's not owned by the viewport we need to take more extensive measures.  If the GC widget still
+	// exists after this point we should just reset the widget, which will forcefully cause the SObjectWidget
+	// to lose access to this UObject.
+	TSharedPtr<SObjectWidget> SafeGCWidget = MyGCWidget.Pin();
+	if ( SafeGCWidget.IsValid() )
+	{
+		SafeGCWidget->ResetWidget();
+	}
+}
+
 void UUserWidget::PostEditImport()
 {
 	Super::PostEditImport();
diff --git a/Engine/Source/Runtime/UMG/Public/Blueprint/UserWidget.h b/Engine/Source/Runtime/UMG/Public/Blueprint/UserWidget.h
--- a/Engine/Source/Runtime/UMG/Public/Blueprint/UserWidget.h
+++ b/Engine/Source/Runtime/UMG/Public/Blueprint/UserWidget.h
@@ -117,6 +117,7 @@ public:
 	virtual class UWorld* GetWorld() const override;
 	virtual void PostEditImport() override;
 	virtual void PostDuplicate(bool bDuplicateForPIE) override;
+	virtual void BeginDestroy() override;
 	// End of UObject interface
 
 	void Initialize();
diff --git a/Engine/Source/Runtime/UMG/Public/Slate/SObjectWidget.h b/Engine/Source/Runtime/UMG/Public/Slate/SObjectWidget.h
--- a/Engine/Source/Runtime/UMG/Public/Slate/SObjectWidget.h
+++ b/Engine/Source/Runtime/UMG/Public/Slate/SObjectWidget.h
@@ -24,6 +24,8 @@ class SObjectWidget : public SCompoundWidget, public FGCObject
 
 	void Construct(const FArguments& InArgs, UUserWidget* InWidgetObject);
 
+	void ResetWidget();
+
 	// FGCObject interface
 	virtual void AddReferencedObjects(FReferenceCollector& Collector) override;
 	// End of FGCObject interface
