Fix for the SWidget in the UWidget not being valid when OnConstructed is called.

diff --git a/Engine/Source/Runtime/UMG/Private/Components/Widget.cpp b/Engine/Source/Runtime/UMG/Private/Components/Widget.cpp
--- a/Engine/Source/Runtime/UMG/Private/Components/Widget.cpp
+++ b/Engine/Source/Runtime/UMG/Private/Components/Widget.cpp
@@ -251,6 +251,8 @@ TSharedRef<SWidget> UWidget::TakeWidget()
 		SafeWidget = RebuildWidget();
 		MyWidget = SafeWidget;
 
+		OnWidgetRebuilt();
+
 		bNewlyCreated = true;
 	}
 	else
@@ -292,6 +294,10 @@ TSharedRef<SWidget> UWidget::TakeWidget()
 	}
 }
 
+void UWidget::OnWidgetRebuilt()
+{
+}
+
 TSharedPtr<SWidget> UWidget::GetCachedWidget() const
 {
 	if ( MyGCWidget.IsValid() )
diff --git a/Engine/Source/Runtime/UMG/Private/UserWidget.cpp b/Engine/Source/Runtime/UMG/Private/UserWidget.cpp
--- a/Engine/Source/Runtime/UMG/Private/UserWidget.cpp
+++ b/Engine/Source/Runtime/UMG/Private/UserWidget.cpp
@@ -452,13 +452,16 @@ TSharedRef<SWidget> UUserWidget::RebuildWidget()
 	// Add the first component to the root of the widget surface.
 	TSharedRef<SWidget> UserRootWidget = WidgetTree->RootWidget ? WidgetTree->RootWidget->TakeWidget() : TSharedRef<SWidget>(SNew(SSpacer));
 
-	if ( !IsDesignTime() )
+	return UserRootWidget;
+}
+
+void UUserWidget::OnWidgetRebuilt()
+{
+	if (!IsDesignTime())
 	{
 		// Notify the widget that it has been constructed.
 		Construct();
 	}
-
-	return UserRootWidget;
 }
 
 TSharedPtr<SWidget> UUserWidget::GetSlateWidgetFromName(const FName& Name) const
diff --git a/Engine/Source/Runtime/UMG/Public/Blueprint/UserWidget.h b/Engine/Source/Runtime/UMG/Public/Blueprint/UserWidget.h
--- a/Engine/Source/Runtime/UMG/Public/Blueprint/UserWidget.h
+++ b/Engine/Source/Runtime/UMG/Public/Blueprint/UserWidget.h
@@ -655,6 +655,7 @@ public:
 
 protected:
 	virtual TSharedRef<SWidget> RebuildWidget() override;
+	virtual void OnWidgetRebuilt() override;
 
 	FMargin GetFullScreenOffset() const;
 	FAnchors GetViewportAnchors() const;
diff --git a/Engine/Source/Runtime/UMG/Public/Components/Widget.h b/Engine/Source/Runtime/UMG/Public/Components/Widget.h
--- a/Engine/Source/Runtime/UMG/Public/Components/Widget.h
+++ b/Engine/Source/Runtime/UMG/Public/Components/Widget.h
@@ -351,6 +351,9 @@ protected:
 protected:
 	/** Function implemented by all subclasses of UWidget is called when the underlying SWidget needs to be constructed. */
 	virtual TSharedRef<SWidget> RebuildWidget();
+
+	/** Function called after the underlying SWidget is constructed. */
+	virtual void OnWidgetRebuilt();
 	
 	TSharedRef<SWidget> BuildDesignTimeWidget(TSharedRef<SWidget> WrapWidget);
 
