UE-6670 - It's now safe to call IsPressed on a button before it has been added to the viewport.

diff --git a/Engine/Source/Runtime/UMG/Private/Components/Button.cpp b/Engine/Source/Runtime/UMG/Private/Components/Button.cpp
--- a/Engine/Source/Runtime/UMG/Private/Components/Button.cpp
+++ b/Engine/Source/Runtime/UMG/Private/Components/Button.cpp
@@ -98,7 +98,12 @@ void UButton::SetBackgroundColor(FLinearColor Color)
 
 bool UButton::IsPressed() const
 {
-	return MyButton->IsPressed();
+	if ( MyButton.IsValid() )
+	{
+		return MyButton->IsPressed();
+	}
+
+	return false;
 }
 
 void UButton::PostLoad()
