UMG - Widgets now support widget tool tips.

diff --git a/Engine/Source/Runtime/UMG/Private/Binding/WidgetBinding.cpp b/Engine/Source/Runtime/UMG/Private/Binding/WidgetBinding.cpp
new file mode 100644
--- /dev/null
+++ b/Engine/Source/Runtime/UMG/Private/Binding/WidgetBinding.cpp
@@ -0,0 +1,48 @@
+// Copyright 1998-2015 Epic Games, Inc. All Rights Reserved.
+
+#include "UMGPrivatePCH.h"
+#include "WidgetBlueprintLibrary.h"
+#include "WidgetBinding.h"
+
+#define LOCTEXT_NAMESPACE "UMG"
+
+UWidgetBinding::UWidgetBinding()
+{
+}
+
+bool UWidgetBinding::IsSupportedDestination(UProperty* Property) const
+{
+	return IsSupportedSource(Property);
+}
+
+bool UWidgetBinding::IsSupportedSource(UProperty* Property) const
+{
+	if ( IsConcreteTypeCompatibleWithReflectedType<UObject*>(Property) )
+	{
+		if ( UObjectProperty* ObjectProperty = Cast<UObjectProperty>(Property) )
+		{
+			return ObjectProperty->PropertyClass->IsChildOf(UWidget::StaticClass());
+		}
+	}
+
+	return false;
+}
+
+UWidget* UWidgetBinding::GetValue() const
+{
+	if ( UObject* Source = SourceObject.Get() )
+	{
+		UObject* Value;
+		if ( SourcePath.GetValue<UObject*>(Source, Value) )
+		{
+			if ( UWidget* Widget = Cast<UWidget>(Value) )
+			{
+				return Widget;
+			}
+		}
+	}
+
+	return nullptr;
+}
+
+#undef LOCTEXT_NAMESPACE
diff --git a/Engine/Source/Runtime/UMG/Private/Components/Widget.cpp b/Engine/Source/Runtime/UMG/Private/Components/Widget.cpp
--- a/Engine/Source/Runtime/UMG/Private/Components/Widget.cpp
+++ b/Engine/Source/Runtime/UMG/Private/Components/Widget.cpp
@@ -4,12 +4,87 @@
 
 #include "ReflectionMetadata.h"
 
-#include "TextBinding.h"
-#include "FloatBinding.h"
-#include "BoolBinding.h"
-#include "BrushBinding.h"
-#include "ColorBinding.h"
-#include "Int32Binding.h"
+/**
+* Interface for tool tips.
+*/
+class FDelegateToolTip : public IToolTip
+{
+public:
+
+	/**
+	* Gets the widget that this tool tip represents.
+	*
+	* @return The tool tip widget.
+	*/
+	virtual TSharedRef<class SWidget> AsWidget()
+	{
+		return GetContentWidget();
+	}
+
+	/**
+	* Gets the tool tip's content widget.
+	*
+	* @return The content widget.
+	*/
+	virtual TSharedRef<SWidget> GetContentWidget()
+	{
+		if ( CachedToolTip.IsValid() )
+		{
+			return CachedToolTip.ToSharedRef();
+		}
+
+		UWidget* Widget = ToolTipWidgetDelegate.Execute();
+		if ( Widget )
+		{
+			CachedToolTip = Widget->TakeWidget();
+			return CachedToolTip.ToSharedRef();
+		}
+
+		return SNullWidget::NullWidget;
+	}
+
+	/**
+	* Sets the tool tip's content widget.
+	*
+	* @param InContentWidget The new content widget to set.
+	*/
+	virtual void SetContentWidget(const TSharedRef<SWidget>& InContentWidget)
+	{
+		CachedToolTip = InContentWidget;
+	}
+
+	/**
+	* Checks whether this tool tip has no content to display right now.
+	*
+	* @return true if the tool tip has no content to display, false otherwise.
+	*/
+	virtual bool IsEmpty() const
+	{
+		return !ToolTipWidgetDelegate.IsBound();
+	}
+
+	/**
+	* Checks whether this tool tip can be made interactive by the user (by holding Ctrl).
+	*
+	* @return true if it is an interactive tool tip, false otherwise.
+	*/
+	virtual bool IsInteractive() const
+	{
+		return false;
+	}
+
+	virtual void OnClosed() override
+	{
+		CachedToolTip.Reset();
+	}
+
+public:
+	UWidget::FGetWidget ToolTipWidgetDelegate;
+
+private:
+	TSharedPtr<SWidget> CachedToolTip;
+};
+
 
 /////////////////////////////////////////////////////
 // UWidget
@@ -148,6 +223,31 @@ void UWidget::SetToolTipText(const FText& InToolTipText)
 	}
 }
 
+void UWidget::SetToolTip(UWidget* InToolTipWidget)
+{
+	ToolTipWidget = InToolTipWidget;
+
+	TSharedPtr<SWidget> SafeWidget = GetCachedWidget();
+	if ( SafeWidget.IsValid() )
+	{
+		if ( ToolTipWidget )
+		{
+			TSharedRef<SToolTip> ToolTip = SNew(SToolTip)
+				.TextMargin(FMargin(0))
+				.BorderImage(nullptr)
+				[
+					ToolTipWidget->TakeWidget()
+				];
+
+			SafeWidget->SetToolTip(ToolTip);
+		}
+		else
+		{
+			SafeWidget->SetToolTip(TSharedPtr<IToolTip>());
+		}
+	}
+}
+
 bool UWidget::IsHovered() const
 {
 	TSharedPtr<SWidget> SafeWidget = GetCachedWidget();
@@ -479,7 +579,24 @@ void UWidget::SynchronizeProperties()
 	UpdateRenderTransform();
 	SafeWidget->SetRenderTransformPivot(RenderTransformPivot);
 
-	if ( !ToolTipText.IsEmpty() )
+	if ( ToolTipWidgetDelegate.IsBound() && !IsDesignTime() )
+	{
+		TSharedRef<FDelegateToolTip> ToolTip = MakeShareable(new FDelegateToolTip());
+		ToolTip->ToolTipWidgetDelegate = ToolTipWidgetDelegate;
+		SafeWidget->SetToolTip(ToolTip);
+	}
+	else if ( ToolTipWidget != nullptr )
+	{
+		TSharedRef<SToolTip> ToolTip = SNew(SToolTip)
+			.TextMargin(FMargin(0))
+			.BorderImage(nullptr)
+			[
+				ToolTipWidget->TakeWidget()
+			];
+
+		SafeWidget->SetToolTip(ToolTip);
+	}
+	else if ( !ToolTipText.IsEmpty() || ToolTipTextDelegate.IsBound() )
 	{
 		SafeWidget->SetToolTipText(OPTIONAL_BINDING(FText, ToolTipText));
 	}
@@ -497,6 +614,8 @@ void UWidget::SynchronizeProperties()
 	}
 
 #if WITH_EDITOR
+	// In editor builds we add metadata to the widget so that once hit with the widget reflector it can report
+	// where it comes from, what blueprint, what the name of the widget was...etc.
 	SafeWidget->AddMetadata<FReflectionMetaData>(MakeShareable(new FReflectionMetaData(GetFName(), GetClass(), WidgetGeneratedBy)));
 #endif
 }
diff --git a/Engine/Source/Runtime/UMG/Public/Binding/WidgetBinding.h b/Engine/Source/Runtime/UMG/Public/Binding/WidgetBinding.h
new file mode 100644
--- /dev/null
+++ b/Engine/Source/Runtime/UMG/Public/Binding/WidgetBinding.h
@@ -0,0 +1,21 @@
+﻿// Copyright 1998-2015 Epic Games, Inc. All Rights Reserved.
+
+#pragma once
+
+#include "WidgetBinding.generated.h"
+
+UCLASS()
+class UMG_API UWidgetBinding : public UPropertyBinding
+{
+	GENERATED_BODY()
+
+public:
+
+	UWidgetBinding();
+
+	virtual bool IsSupportedSource(UProperty* Property) const override;
+	virtual bool IsSupportedDestination(UProperty* Property) const override;
+
+	UFUNCTION()
+	UWidget* GetValue() const;
+};
diff --git a/Engine/Source/Runtime/UMG/Public/Components/MenuAnchor.h b/Engine/Source/Runtime/UMG/Public/Components/MenuAnchor.h
--- a/Engine/Source/Runtime/UMG/Public/Components/MenuAnchor.h
+++ b/Engine/Source/Runtime/UMG/Public/Components/MenuAnchor.h
@@ -25,7 +25,7 @@ public:
 
 	/** Called when the menu content is requested to allow a more customized handling over what to display */
 	UPROPERTY(EditDefaultsOnly, Category="Events")
-	FGetContent OnGetMenuContentEvent;
+	FGetWidget OnGetMenuContentEvent;
 	
 	/** The placement location of the summoned widget. */
 	UPROPERTY(EditDefaultsOnly, BlueprintReadOnly, Category="Menu")
diff --git a/Engine/Source/Runtime/UMG/Public/Components/Widget.h b/Engine/Source/Runtime/UMG/Public/Components/Widget.h
--- a/Engine/Source/Runtime/UMG/Public/Components/Widget.h
+++ b/Engine/Source/Runtime/UMG/Public/Components/Widget.h
@@ -45,7 +45,7 @@ public:
 	DECLARE_DYNAMIC_DELEGATE_RetVal(ESlateVisibility, FGetSlateVisibility);
 	DECLARE_DYNAMIC_DELEGATE_RetVal(EMouseCursor::Type, FGetMouseCursor);
 	DECLARE_DYNAMIC_DELEGATE_RetVal(ECheckBoxState, FGetCheckBoxState);
-	DECLARE_DYNAMIC_DELEGATE_RetVal(UWidget*, FGetContent);
+	DECLARE_DYNAMIC_DELEGATE_RetVal(UWidget*, FGetWidget);
 
 	DECLARE_DYNAMIC_DELEGATE_RetVal_OneParam(UWidget*, FGenerateWidgetForString, FString, Item);
 	DECLARE_DYNAMIC_DELEGATE_RetVal_OneParam(UWidget*, FGenerateWidgetForObject, UObject*, Item);
@@ -72,7 +72,7 @@ public:
 	UPanelSlot* Slot;
 
 	/** Sets whether this widget can be modified interactively by the user */
-	UPROPERTY(EditDefaultsOnly, BlueprintReadOnly, Category=Behavior)
+	UPROPERTY(EditDefaultsOnly, BlueprintReadOnly, Category="Behavior")
 	bool bIsEnabled;
 
 	/** A bindable delegate for bIsEnabled */
@@ -80,19 +80,27 @@ public:
 	FGetBool bIsEnabledDelegate;
 
 	/** Tooltip text to show when the user hovers over the widget with the mouse */
-	UPROPERTY(EditDefaultsOnly, BlueprintReadOnly, Category=Behavior)
+	UPROPERTY(EditDefaultsOnly, BlueprintReadOnly, Category="Behavior")
 	FText ToolTipText;
 
 	/** A bindable delegate for ToolTipText */
 	UPROPERTY()
 	FGetText ToolTipTextDelegate;
 
+	/** Tooltip widget to show when the user hovers over the widget with the mouse */
+	UPROPERTY(EditDefaultsOnly, Category="Behavior", AdvancedDisplay)
+	UWidget* ToolTipWidget;
+
+	/** A bindable delegate for ToolTipWidget */
+	UPROPERTY()
+	FGetWidget ToolTipWidgetDelegate;
+
 	/** The visibility of the widget */
 	UPROPERTY()
 	TEnumAsByte<ESlateVisibility> Visiblity_DEPRECATED;
 
 	/** The visibility of the widget */
-	UPROPERTY(EditDefaultsOnly, Category=Behavior)
+	UPROPERTY(EditDefaultsOnly, Category="Behavior")
 	ESlateVisibility Visibility;
 
 	/** A bindable delegate for Visibility */
@@ -100,7 +108,7 @@ public:
 	FGetSlateVisibility VisibilityDelegate;
 
 	/** The cursor to show when the mouse is over the widget */
-	UPROPERTY(EditDefaultsOnly, BlueprintReadOnly, Category=Behavior, AdvancedDisplay)
+	UPROPERTY(EditDefaultsOnly, BlueprintReadOnly, Category="Behavior", AdvancedDisplay)
 	TEnumAsByte<EMouseCursor::Type> Cursor;
 
 	/** A bindable delegate for Cursor */
@@ -176,6 +184,10 @@ public:
 	UFUNCTION(BlueprintCallable, Category="Widget")
 	void SetToolTipText(const FText& InToolTipText);
 
+	/** Sets a custom widget as the tooltip of the widget. */
+	UFUNCTION(BlueprintCallable, Category="Widget")
+	void SetToolTip(UWidget* Widget);
+
 	/** @return true if the widget is Visible, HitTestInvisible or SelfHitTestInvisible. */
 	UFUNCTION(BlueprintCallable, Category="Widget")
 	bool IsVisible() const;
@@ -188,7 +200,7 @@ public:
 	UFUNCTION(BlueprintCallable, Category="Widget")
 	void SetVisibility(ESlateVisibility InVisibility);
 
-	/** Gets if the button is currently being hovered by the mouse */
+	/** @return true if the widget is currently being hovered by a pointer device */
 	UFUNCTION(BlueprintCallable, Category="Widget")
 	bool IsHovered() const;
 
@@ -271,12 +283,6 @@ public:
 	 */
 	bool IsChildOf(UWidget* PossibleParent);
 
-	///**
-	// * Allows binding to properties of other objects at runtime.
-	// */
-	//UFUNCTION(BlueprintCallable, Category="Widget")
-	//bool BindProperty(const FName& DestinationProperty, UObject* SourceObject, const FName& SourceProperty);
-
 	/**  */
 	bool AddBinding(UDelegateProperty* DelegateProperty, UObject* SourceObject, const FDynamicPropertyPath& BindingPath);
 
