[Github PR 628] Rama - On explict create widgets that take a player controller, we now return null from CreateWidget calls (instead of crashing) if the player controller is not local.

diff --git a/Engine/Source/Runtime/UMG/Public/Blueprint/UserWidget.h b/Engine/Source/Runtime/UMG/Public/Blueprint/UserWidget.h
--- a/Engine/Source/Runtime/UMG/Public/Blueprint/UserWidget.h
+++ b/Engine/Source/Runtime/UMG/Public/Blueprint/UserWidget.h
@@ -685,10 +685,13 @@ T* CreateWidget(UWorld* World, UClass* UserWidgetClass)
 
 	// Assign the outer to the game instance if it exists, otherwise use the world
 	UObject* Outer = World->GetGameInstance() ? StaticCast<UObject*>(World->GetGameInstance()) : StaticCast<UObject*>(World);
-	ULocalPlayer* Player = World->GetFirstLocalPlayerFromController();
 	UUserWidget* NewWidget = ConstructObject<UUserWidget>(UserWidgetClass, Outer);
 
-	NewWidget->SetPlayerContext(FLocalPlayerContext(Player));
+	if ( ULocalPlayer* Player = World->GetFirstLocalPlayerFromController() )
+	{
+		NewWidget->SetPlayerContext(FLocalPlayerContext(Player));
+	}
+
 	NewWidget->Initialize();
 
 	return Cast<T>(NewWidget);
@@ -703,6 +706,12 @@ T* CreateWidget(APlayerController* OwningPlayer, UClass* UserWidgetClass)
 		return nullptr;
 	}
 
+	if ( !OwningPlayer->IsLocalPlayerController() )
+	{
+		//Don't create widgets for proxies of the Player Controller, only the actual local player.
+		return nullptr;
+	}
+
 	// Assign the outer to the game instance if it exists, otherwise use the player controller's world
 	UWorld* World = OwningPlayer->GetWorld();
 	UObject* Outer = World->GetGameInstance() ? StaticCast<UObject*>(World->GetGameInstance()) : StaticCast<UObject*>(World);
@@ -725,9 +734,9 @@ T* CreateWidget(UGameInstance* OwningGame, UClass* UserWidgetClass)
 
 	UUserWidget* NewWidget = ConstructObject<UUserWidget>(UserWidgetClass, OwningGame);
 
-	if ( APlayerController* PC = OwningGame->GetFirstLocalPlayerController() )
+	if ( ULocalPlayer* Player = OwningGame->GetFirstLocalPlayerController() )
 	{
-		NewWidget->SetPlayerContext(FLocalPlayerContext(PC));
+		NewWidget->SetPlayerContext(FLocalPlayerContext(Player));
 	}
 	
 	NewWidget->Initialize();
