[Github PR 665] Rama - Adding HasChild to all UMG Panel widgets.

diff --git a/Engine/Source/Runtime/UMG/Private/Components/PanelWidget.cpp b/Engine/Source/Runtime/UMG/Private/Components/PanelWidget.cpp
--- a/Engine/Source/Runtime/UMG/Private/Components/PanelWidget.cpp
+++ b/Engine/Source/Runtime/UMG/Private/Components/PanelWidget.cpp
@@ -63,6 +63,15 @@ int32 UPanelWidget::GetChildIndex(UWidget* Content) const
 	return -1;
 }
 
+bool UPanelWidget::HasChild(UWidget* Content) const
+{
+	if ( !Content )
+	{
+		return false;
+	}
+	return ( Content->GetParent() == this );
+}
+
 bool UPanelWidget::RemoveChildAt(int32 Index)
 {
 	if ( Index < 0 || Index >= Slots.Num() )
diff --git a/Engine/Source/Runtime/UMG/Public/Components/PanelWidget.h b/Engine/Source/Runtime/UMG/Public/Components/PanelWidget.h
--- a/Engine/Source/Runtime/UMG/Public/Components/PanelWidget.h
+++ b/Engine/Source/Runtime/UMG/Public/Components/PanelWidget.h
@@ -35,6 +35,10 @@ public:
 	UFUNCTION(BlueprintCallable, Category="Widget|Panel")
 	int32 GetChildIndex(UWidget* Content) const;
 
+	/** @return true if panel contains this widget */
+	UFUNCTION(BlueprintCallable, Category="Widget|Panel")
+	bool HasChild(UWidget* Content) const;
+
 	/** Removes a child by it's index. */
 	UFUNCTION(BlueprintCallable, Category="Widget|Panel")
 	bool RemoveChildAt(int32 Index);
