UMG - Fixing the build, also adding a way to set the owning player from blueprints of a user widget.

diff --git a/Engine/Source/Runtime/UMG/Private/UserWidget.cpp b/Engine/Source/Runtime/UMG/Private/UserWidget.cpp
--- a/Engine/Source/Runtime/UMG/Private/UserWidget.cpp
+++ b/Engine/Source/Runtime/UMG/Private/UserWidget.cpp
@@ -682,6 +682,14 @@ ULocalPlayer* UUserWidget::GetOwningLocalPlayer() const
 	return PC ? Cast<ULocalPlayer>(PC->Player) : nullptr;
 }
 
+void UUserWidget::SetOwningLocalPlayer(ULocalPlayer* LocalPlayer)
+{
+	if ( LocalPlayer )
+	{
+		PlayerContext = FLocalPlayerContext(LocalPlayer);
+	}
+}
+
 APlayerController* UUserWidget::GetOwningPlayer() const
 {
 	return PlayerContext.IsValid() ? PlayerContext.GetPlayerController() : nullptr;
diff --git a/Engine/Source/Runtime/UMG/Public/Blueprint/UserWidget.h b/Engine/Source/Runtime/UMG/Public/Blueprint/UserWidget.h
--- a/Engine/Source/Runtime/UMG/Public/Blueprint/UserWidget.h
+++ b/Engine/Source/Runtime/UMG/Public/Blueprint/UserWidget.h
@@ -194,6 +194,13 @@ public:
 	UFUNCTION(BlueprintCallable, BlueprintCosmetic, Category="Player")
 	class ULocalPlayer* GetOwningLocalPlayer() const;
 
+	/**
+	 * Sets the local player associated with this UI.
+	 * @param LocalPlayer The local player you want to be the conceptual owner of this UI.
+	 */
+	UFUNCTION(BlueprintCallable, BlueprintCosmetic, Category="Player")
+	void SetOwningLocalPlayer(ULocalPlayer* LocalPlayer);
+
 	/**
 	 * Gets the player controller associated with this UI.
 	 * @return The player controller that owns the UI.
@@ -734,7 +741,7 @@ T* CreateWidget(UGameInstance* OwningGame, UClass* UserWidgetClass)
 
 	UUserWidget* NewWidget = ConstructObject<UUserWidget>(UserWidgetClass, OwningGame);
 
-	if ( ULocalPlayer* Player = OwningGame->GetFirstLocalPlayerController() )
+	if ( ULocalPlayer* Player = OwningGame->GetFirstGamePlayer() )
 	{
 		NewWidget->SetPlayerContext(FLocalPlayerContext(Player));
 	}
