UMG - Performance fix on screen space widgets coming from components, we no longer attempt to recreate them every frame, and then remove them from the viewport.

diff --git a/Engine/Source/Runtime/UMG/Private/Components/WidgetComponent.cpp b/Engine/Source/Runtime/UMG/Private/Components/WidgetComponent.cpp
--- a/Engine/Source/Runtime/UMG/Private/Components/WidgetComponent.cpp
+++ b/Engine/Source/Runtime/UMG/Private/Components/WidgetComponent.cpp
@@ -523,30 +523,35 @@ void UWidgetComponent::TickComponent(float DeltaTime, enum ELevelTick TickType,
 
 			if ( TargetPlayer && PlayerController )
 			{
-				if ( !Widget->IsInViewport() )
+				if ( IsVisible() )
 				{
-					Widget->AddToViewport();
-				}
+					FVector WorldLocation = GetComponentLocation();
 
-				FVector WorldLocation = GetComponentLocation();
+					FVector2D ScreenLocation;
+					bool bProjected = PlayerController->ProjectWorldLocationToScreen(WorldLocation, ScreenLocation);
 
-				FVector2D ScreenLocation;
-				bool bProjected = PlayerController->ProjectWorldLocationToScreen(WorldLocation, ScreenLocation);
-				if ( bProjected )
-				{
-					Widget->SetVisibility(ESlateVisibility::SelfHitTestInvisible);
-				}
-				else
-				{
-					Widget->SetVisibility(ESlateVisibility::Hidden);
+					if ( bProjected )
+					{
+						Widget->SetVisibility(ESlateVisibility::SelfHitTestInvisible);
+					}
+					else
+					{
+						Widget->SetVisibility(ESlateVisibility::Hidden);
+					}
+
+					if ( !Widget->IsInViewport() )
+					{
+						Widget->AddToViewport();
+					}
+
+					Widget->SetDesiredSizeInViewport(DrawSize);
+					Widget->SetPositionInViewport(ScreenLocation);
+					Widget->SetAlignmentInViewport(Pivot);
 				}
-				Widget->SetPositionInViewport(ScreenLocation);
-				Widget->SetAlignmentInViewport(Pivot);
 			}
-
-			// If the component isn't visible, and the widget is on the viewport, remove it.
-			if ( !IsVisible() && Widget->IsInViewport() )
+			else if ( Widget->IsInViewport() )
 			{
+				// If the component isn't visible, and the widget is on the viewport, remove it.
 				Widget->RemoveFromParent();
 			}
 		}
