Mismatched function signatures

diff --git a/Engine/Source/Runtime/UMG/Private/Components/Viewport.cpp b/Engine/Source/Runtime/UMG/Private/Components/Viewport.cpp
--- a/Engine/Source/Runtime/UMG/Private/Components/Viewport.cpp
+++ b/Engine/Source/Runtime/UMG/Private/Components/Viewport.cpp
@@ -466,7 +466,7 @@ class SAutoRefreshViewport : public SViewport
 		SetViewportInterface(Viewport.ToSharedRef());
 	}
 
-	virtual void Tick( const FGeometry& AllottedGeometry, double InCurrentTime, float InDeltaTime ) override
+	virtual void Tick( const FGeometry& AllottedGeometry, const double InCurrentTime, const float InDeltaTime ) override
 	{
 		Viewport->Invalidate();
 		Viewport->InvalidateDisplay();
