UMG - Improving the slate widget host to allow content to be set at any time.  Also the marching ants no longer appears unless it's design time.

diff --git a/Engine/Source/Runtime/UMG/Private/Components/NativeWidgetHost.cpp b/Engine/Source/Runtime/UMG/Private/Components/NativeWidgetHost.cpp
--- a/Engine/Source/Runtime/UMG/Private/Components/NativeWidgetHost.cpp
+++ b/Engine/Source/Runtime/UMG/Private/Components/NativeWidgetHost.cpp
@@ -15,7 +15,14 @@ UNativeWidgetHost::UNativeWidgetHost(const FObjectInitializer& ObjectInitializer
 
 void UNativeWidgetHost::SetContent(TSharedRef<SWidget> InContent)
 {
-	MyWidget = InContent;
+	NativeWidget = InContent;
+
+	TSharedPtr<SWidget> StableMyWidget = MyWidget.Pin();
+	if ( StableMyWidget.IsValid() )
+	{
+		TSharedPtr<SBox> MyBox = StaticCastSharedPtr<SBox>(StableMyWidget);
+		MyBox->SetContent(InContent);
+	}
 }
 
 void UNativeWidgetHost::ReleaseSlateResources(bool bReleaseChildren)
@@ -23,25 +30,34 @@ void UNativeWidgetHost::ReleaseSlateResources(bool bReleaseChildren)
 	Super::ReleaseSlateResources(bReleaseChildren);
 
 	MyWidget.Reset();
+	NativeWidget.Reset();
 }
 
 TSharedRef<SWidget> UNativeWidgetHost::RebuildWidget()
 {
-	if ( MyWidget.IsValid() )
+	return SNew(SBox)
+		[
+			( NativeWidget.IsValid() ) ? NativeWidget.ToSharedRef() : GetDefaultContent()
+		];
+}
+
+TSharedRef<SWidget> UNativeWidgetHost::GetDefaultContent()
+{
+	if ( IsDesignTime() )
 	{
-		return MyWidget.ToSharedRef();
+		return SNew(SBorder)
+			.Visibility(EVisibility::HitTestInvisible)
+			.BorderImage(FUMGStyle::Get().GetBrush("MarchingAnts"))
+			.HAlign(HAlign_Center)
+			.VAlign(VAlign_Center)
+			[
+				SNew(STextBlock)
+				.Text(LOCTEXT("NativeWidgetHostText", "Slate Widget Host"))
+			];
 	}
 	else
 	{
-		return SNew(SBorder)
-		.Visibility(EVisibility::HitTestInvisible)
-		.BorderImage(FUMGStyle::Get().GetBrush("MarchingAnts"))
-		.HAlign(HAlign_Center)
-		.VAlign(VAlign_Center)
-		[
-			SNew(STextBlock)
-			.Text(LOCTEXT("NativeWidgetHostText", "Slate Widget Host"))
-		];
+		return SNullWidget::NullWidget;
 	}
 }
 
diff --git a/Engine/Source/Runtime/UMG/Public/Components/NativeWidgetHost.h b/Engine/Source/Runtime/UMG/Public/Components/NativeWidgetHost.h
--- a/Engine/Source/Runtime/UMG/Public/Components/NativeWidgetHost.h
+++ b/Engine/Source/Runtime/UMG/Public/Components/NativeWidgetHost.h
@@ -22,11 +22,13 @@ class UMG_API UNativeWidgetHost : public UWidget
 	virtual const FText GetPaletteCategory() override;
 #endif
 
-protected:
-	TSharedPtr<SWidget> MyWidget;
-
 protected:
 	// UWidget interface
 	virtual TSharedRef<SWidget> RebuildWidget() override;
 	// End of UWidget interface
+
+	TSharedRef<SWidget> GetDefaultContent();
+
+protected:
+	TSharedPtr<SWidget> NativeWidget;
 };
