UMG - Don't need to reset MyWidget on the NativeWidgetHost as it is already a weak pointer.

diff --git a/Engine/Source/Runtime/UMG/Private/Components/NativeWidgetHost.cpp b/Engine/Source/Runtime/UMG/Private/Components/NativeWidgetHost.cpp
--- a/Engine/Source/Runtime/UMG/Private/Components/NativeWidgetHost.cpp
+++ b/Engine/Source/Runtime/UMG/Private/Components/NativeWidgetHost.cpp
@@ -29,7 +29,6 @@ void UNativeWidgetHost::ReleaseSlateResources(bool bReleaseChildren)
 {
 	Super::ReleaseSlateResources(bReleaseChildren);
 
-	MyWidget.Reset();
 	NativeWidget.Reset();
 }
 
