UMG - Adding stat trackers for the native binding objects, currently commented out to prevent overhead.

diff --git a/Engine/Source/Runtime/UMG/Private/Binding/BoolBinding.cpp b/Engine/Source/Runtime/UMG/Private/Binding/BoolBinding.cpp
--- a/Engine/Source/Runtime/UMG/Private/Binding/BoolBinding.cpp
+++ b/Engine/Source/Runtime/UMG/Private/Binding/BoolBinding.cpp
@@ -21,6 +21,8 @@ bool UBoolBinding::IsSupportedDestination(UProperty* Property) const
 
 bool UBoolBinding::GetValue() const
 {
+	//SCOPE_CYCLE_COUNTER(STAT_UMGBinding);
+
 	if ( UObject* Source = SourceObject.Get() )
 	{
 		bool Value;
diff --git a/Engine/Source/Runtime/UMG/Private/Binding/BrushBinding.cpp b/Engine/Source/Runtime/UMG/Private/Binding/BrushBinding.cpp
--- a/Engine/Source/Runtime/UMG/Private/Binding/BrushBinding.cpp
+++ b/Engine/Source/Runtime/UMG/Private/Binding/BrushBinding.cpp
@@ -30,6 +30,8 @@ bool UBrushBinding::IsSupportedSource(UProperty* Property) const
 
 FSlateBrush UBrushBinding::GetValue() const
 {
+	//SCOPE_CYCLE_COUNTER(STAT_UMGBinding);
+
 	if ( UObject* Source = SourceObject.Get() )
 	{
 		if ( bConversion.Get(EConversion::None) == EConversion::None )
diff --git a/Engine/Source/Runtime/UMG/Private/Binding/CheckedStateBinding.cpp b/Engine/Source/Runtime/UMG/Private/Binding/CheckedStateBinding.cpp
--- a/Engine/Source/Runtime/UMG/Private/Binding/CheckedStateBinding.cpp
+++ b/Engine/Source/Runtime/UMG/Private/Binding/CheckedStateBinding.cpp
@@ -31,6 +31,8 @@ bool UCheckedStateBinding::IsSupportedDestination(UProperty* Property) const
 
 ECheckBoxState UCheckedStateBinding::GetValue() const
 {
+	//SCOPE_CYCLE_COUNTER(STAT_UMGBinding);
+
 	if ( UObject* Source = SourceObject.Get() )
 	{
 		if ( bConversion.Get(EConversion::None) == EConversion::None )
diff --git a/Engine/Source/Runtime/UMG/Private/Binding/ColorBinding.cpp b/Engine/Source/Runtime/UMG/Private/Binding/ColorBinding.cpp
--- a/Engine/Source/Runtime/UMG/Private/Binding/ColorBinding.cpp
+++ b/Engine/Source/Runtime/UMG/Private/Binding/ColorBinding.cpp
@@ -37,6 +37,8 @@ void UColorBinding::Bind(UProperty* Property, FScriptDelegate* Delegate)
 
 FSlateColor UColorBinding::GetSlateValue() const
 {
+	//SCOPE_CYCLE_COUNTER(STAT_UMGBinding);
+
 	if ( UObject* Source = SourceObject.Get() )
 	{
 		if ( !bNeedsConversion.Get(false) )
@@ -65,6 +67,8 @@ FSlateColor UColorBinding::GetSlateValue() const
 
 FLinearColor UColorBinding::GetLinearValue() const
 {
+	//SCOPE_CYCLE_COUNTER(STAT_UMGBinding);
+
 	if ( UObject* Source = SourceObject.Get() )
 	{
 		if ( !bNeedsConversion.Get(false) )
diff --git a/Engine/Source/Runtime/UMG/Private/Binding/FloatBinding.cpp b/Engine/Source/Runtime/UMG/Private/Binding/FloatBinding.cpp
--- a/Engine/Source/Runtime/UMG/Private/Binding/FloatBinding.cpp
+++ b/Engine/Source/Runtime/UMG/Private/Binding/FloatBinding.cpp
@@ -21,6 +21,8 @@ bool UFloatBinding::IsSupportedSource(UProperty* Property) const
 
 float UFloatBinding::GetValue() const
 {
+	//SCOPE_CYCLE_COUNTER(STAT_UMGBinding);
+
 	if ( UObject* Source = SourceObject.Get() )
 	{
 		float Value;
diff --git a/Engine/Source/Runtime/UMG/Private/Binding/Int32Binding.cpp b/Engine/Source/Runtime/UMG/Private/Binding/Int32Binding.cpp
--- a/Engine/Source/Runtime/UMG/Private/Binding/Int32Binding.cpp
+++ b/Engine/Source/Runtime/UMG/Private/Binding/Int32Binding.cpp
@@ -21,6 +21,8 @@ bool UInt32Binding::IsSupportedSource(UProperty* Property) const
 
 int32 UInt32Binding::GetValue() const
 {
+	//SCOPE_CYCLE_COUNTER(STAT_UMGBinding);
+
 	if ( UObject* Source = SourceObject.Get() )
 	{
 		int32 Value;
diff --git a/Engine/Source/Runtime/UMG/Private/Binding/MouseCursorBinding.cpp b/Engine/Source/Runtime/UMG/Private/Binding/MouseCursorBinding.cpp
--- a/Engine/Source/Runtime/UMG/Private/Binding/MouseCursorBinding.cpp
+++ b/Engine/Source/Runtime/UMG/Private/Binding/MouseCursorBinding.cpp
@@ -31,6 +31,8 @@ bool UMouseCursorBinding::IsSupportedDestination(UProperty* Property) const
 
 EMouseCursor::Type UMouseCursorBinding::GetValue() const
 {
+	//SCOPE_CYCLE_COUNTER(STAT_UMGBinding);
+
 	if ( UObject* Source = SourceObject.Get() )
 	{
 		uint8 Value;
diff --git a/Engine/Source/Runtime/UMG/Private/Binding/PropertyBinding.cpp b/Engine/Source/Runtime/UMG/Private/Binding/PropertyBinding.cpp
--- a/Engine/Source/Runtime/UMG/Private/Binding/PropertyBinding.cpp
+++ b/Engine/Source/Runtime/UMG/Private/Binding/PropertyBinding.cpp
@@ -5,6 +5,8 @@
 
 #define LOCTEXT_NAMESPACE "UMG"
 
+DEFINE_STAT(STAT_UMGBinding);
+
 UPropertyBinding::UPropertyBinding()
 {
 }
diff --git a/Engine/Source/Runtime/UMG/Private/Binding/TextBinding.cpp b/Engine/Source/Runtime/UMG/Private/Binding/TextBinding.cpp
--- a/Engine/Source/Runtime/UMG/Private/Binding/TextBinding.cpp
+++ b/Engine/Source/Runtime/UMG/Private/Binding/TextBinding.cpp
@@ -37,6 +37,8 @@ void UTextBinding::Bind(UProperty* Property, FScriptDelegate* Delegate)
 
 FText UTextBinding::GetTextValue() const
 {
+	//SCOPE_CYCLE_COUNTER(STAT_UMGBinding);
+
 	if ( UObject* Source = SourceObject.Get() )
 	{
 		if ( !bNeedsConversion.Get(false) )
@@ -65,6 +67,8 @@ FText UTextBinding::GetTextValue() const
 
 FString UTextBinding::GetStringValue() const
 {
+	//SCOPE_CYCLE_COUNTER(STAT_UMGBinding);
+
 	if ( UObject* Source = SourceObject.Get() )
 	{
 		if ( !bNeedsConversion.Get(false) )
diff --git a/Engine/Source/Runtime/UMG/Private/Binding/VisibilityBinding.cpp b/Engine/Source/Runtime/UMG/Private/Binding/VisibilityBinding.cpp
--- a/Engine/Source/Runtime/UMG/Private/Binding/VisibilityBinding.cpp
+++ b/Engine/Source/Runtime/UMG/Private/Binding/VisibilityBinding.cpp
@@ -31,6 +31,8 @@ bool UVisibilityBinding::IsSupportedDestination(UProperty* Property) const
 
 ESlateVisibility UVisibilityBinding::GetValue() const
 {
+	//SCOPE_CYCLE_COUNTER(STAT_UMGBinding);
+
 	if ( UObject* Source = SourceObject.Get() )
 	{
 		uint8 Value;
diff --git a/Engine/Source/Runtime/UMG/Private/Binding/WidgetBinding.cpp b/Engine/Source/Runtime/UMG/Private/Binding/WidgetBinding.cpp
--- a/Engine/Source/Runtime/UMG/Private/Binding/WidgetBinding.cpp
+++ b/Engine/Source/Runtime/UMG/Private/Binding/WidgetBinding.cpp
@@ -30,6 +30,8 @@ bool UWidgetBinding::IsSupportedSource(UProperty* Property) const
 
 UWidget* UWidgetBinding::GetValue() const
 {
+	//SCOPE_CYCLE_COUNTER(STAT_UMGBinding);
+
 	if ( UObject* Source = SourceObject.Get() )
 	{
 		UObject* Value;
diff --git a/Engine/Source/Runtime/UMG/Public/Binding/PropertyBinding.h b/Engine/Source/Runtime/UMG/Public/Binding/PropertyBinding.h
--- a/Engine/Source/Runtime/UMG/Public/Binding/PropertyBinding.h
+++ b/Engine/Source/Runtime/UMG/Public/Binding/PropertyBinding.h
@@ -6,6 +6,8 @@
 
 #include "PropertyBinding.generated.h"
 
+DECLARE_CYCLE_STAT_EXTERN(TEXT("UMG Binding"), STAT_UMGBinding, STATGROUP_Slate,);
+
 UCLASS()
 class UMG_API UPropertyBinding : public UObject
 {
