UMG/Web - Adding a new async node that can download images from the web. Jpeg/Png/Bmp are the only supported types for now.  You can now download an image from the web and place it into a image widget or on a border...etc.

diff --git a/Engine/Source/Runtime/UMG/Private/Blueprint/AsyncTaskDownloadImage.cpp b/Engine/Source/Runtime/UMG/Private/Blueprint/AsyncTaskDownloadImage.cpp
new file mode 100644
--- /dev/null
+++ b/Engine/Source/Runtime/UMG/Private/Blueprint/AsyncTaskDownloadImage.cpp
@@ -0,0 +1,121 @@
+// Copyright 1998-2015 Epic Games, Inc. All Rights Reserved.
+
+#include "UMGPrivatePCH.h"
+
+#include "ImageWrapper.h"
+#include "Http.h"
+
+#include "Blueprint/AsyncTaskDownloadImage.h"
+#include "TimerManager.h"
+
+//----------------------------------------------------------------------//
+// UAsyncTaskDownloadImage
+//----------------------------------------------------------------------//
+
+#if !UE_SERVER
+
+static void WriteRawToTexture(UTexture2D* NewTexture2D, const TArray<uint8>& RawData, bool bUseSRGB = true)
+{
+	int32 Height = NewTexture2D->GetSizeY();
+	int32 Width = NewTexture2D->GetSizeX();
+
+	// Fill in the base mip for the texture we created
+	uint8* MipData = (uint8*)NewTexture2D->PlatformData->Mips[0].BulkData.Lock(LOCK_READ_WRITE);
+	for ( int32 y=0; y < Height; y++ )
+	{
+		uint8* DestPtr = &MipData[( Height - 1 - y ) * Width * sizeof(FColor)];
+		const FColor* SrcPtr = &( (FColor*)( RawData.GetData() ) )[( Height - 1 - y ) * Width];
+		for ( int32 x=0; x < Width; x++ )
+		{
+			*DestPtr++ = SrcPtr->B;
+			*DestPtr++ = SrcPtr->G;
+			*DestPtr++ = SrcPtr->R;
+			*DestPtr++ = SrcPtr->A;
+			SrcPtr++;
+		}
+	}
+	NewTexture2D->PlatformData->Mips[0].BulkData.Unlock();
+
+	// Set options
+	NewTexture2D->SRGB = bUseSRGB;
+	NewTexture2D->CompressionNone = true;
+	NewTexture2D->MipGenSettings = TMGS_NoMipmaps;
+	NewTexture2D->CompressionSettings = TC_EditorIcon;
+
+	NewTexture2D->UpdateResource();
+}
+
+#endif
+
+UAsyncTaskDownloadImage::UAsyncTaskDownloadImage(const FObjectInitializer& ObjectInitializer)
+	: Super(ObjectInitializer)
+{
+	if ( HasAnyFlags(RF_ClassDefaultObject) == false )
+	{
+		AddToRoot();
+	}
+}
+
+UAsyncTaskDownloadImage* UAsyncTaskDownloadImage::DownloadImage(FString URL)
+{
+	UAsyncTaskDownloadImage* DownloadTask = NewObject<UAsyncTaskDownloadImage>();
+	DownloadTask->Start(URL);
+
+	return DownloadTask;
+}
+
+void UAsyncTaskDownloadImage::Start(FString URL)
+{
+#if !UE_SERVER
+	// Create the Http request and add to pending request list
+	TSharedRef<class IHttpRequest> HttpRequest = FHttpModule::Get().CreateRequest();
+
+	HttpRequest->OnProcessRequestComplete().BindUObject(this, &UAsyncTaskDownloadImage::HandleImageRequest);
+	HttpRequest->SetURL(URL);
+	HttpRequest->SetVerb(TEXT("GET"));
+	HttpRequest->ProcessRequest();
+#else
+	// On the server we don't execute fail or success we just don't fire the request.
+	RemoveFromRoot();
+#endif
+}
+
+void UAsyncTaskDownloadImage::HandleImageRequest(FHttpRequestPtr HttpRequest, FHttpResponsePtr HttpResponse, bool bSucceeded)
+{
+#if !UE_SERVER
+
+	RemoveFromRoot();
+
+	if ( bSucceeded && HttpResponse.IsValid() && HttpResponse->GetContentLength() > 0 )
+	{
+		IImageWrapperModule& ImageWrapperModule = FModuleManager::LoadModuleChecked<IImageWrapperModule>(FName("ImageWrapper"));
+		IImageWrapperPtr ImageWrappers[3] =
+		{
+			ImageWrapperModule.CreateImageWrapper(EImageFormat::PNG),
+			ImageWrapperModule.CreateImageWrapper(EImageFormat::JPEG),
+			ImageWrapperModule.CreateImageWrapper(EImageFormat::BMP),
+		};
+
+		for ( auto ImageWrapper : ImageWrappers )
+		{
+			if ( ImageWrapper.IsValid() && ImageWrapper->SetCompressed(HttpResponse->GetContent().GetData(), HttpResponse->GetContentLength()) )
+			{
+				const TArray<uint8>* RawData = NULL;
+				if ( ImageWrapper->GetRaw(ERGBFormat::BGRA, 8, RawData) )
+				{
+					if ( UTexture2D* Texture = UTexture2D::CreateTransient(ImageWrapper->GetWidth(), ImageWrapper->GetHeight()) )
+					{
+						WriteRawToTexture(Texture, *RawData);
+						
+						OnSuccess.Broadcast(Texture);
+						return;
+					}
+				}
+			}
+		}
+	}
+
+	OnFail.Broadcast(nullptr);
+
+#endif
+}
diff --git a/Engine/Source/Runtime/UMG/Public/Blueprint/AsyncTaskDownloadImage.h b/Engine/Source/Runtime/UMG/Public/Blueprint/AsyncTaskDownloadImage.h
new file mode 100644
--- /dev/null
+++ b/Engine/Source/Runtime/UMG/Public/Blueprint/AsyncTaskDownloadImage.h
@@ -0,0 +1,37 @@
+// Copyright 1998-2015 Epic Games, Inc. All Rights Reserved.
+
+#pragma once
+
+#include "IHttpRequest.h"
+#include "Kismet/BlueprintAsyncActionBase.h"
+
+#include "AsyncTaskDownloadImage.generated.h"
+
+DECLARE_DYNAMIC_MULTICAST_DELEGATE_OneParam(FDownloadImageDelegate, UTexture2D*, Texture);
+
+UCLASS()
+class UMG_API UAsyncTaskDownloadImage : public UBlueprintAsyncActionBase
+{
+	GENERATED_UCLASS_BODY()
+
+public:
+	UFUNCTION(BlueprintCallable, meta=( BlueprintInternalUseOnly="true" ))
+	static UAsyncTaskDownloadImage* DownloadImage(FString URL);
+
+public:
+
+	UPROPERTY(BlueprintAssignable)
+	FDownloadImageDelegate OnSuccess;
+
+	UPROPERTY(BlueprintAssignable)
+	FDownloadImageDelegate OnFail;
+
+public:
+
+	void Start(FString URL);
+
+private:
+
+	/** Handles image requests coming from the web */
+	void HandleImageRequest(FHttpRequestPtr HttpRequest, FHttpResponsePtr HttpResponse, bool bSucceeded);
+};
diff --git a/Engine/Source/Runtime/UMG/UMG.Build.cs b/Engine/Source/Runtime/UMG/UMG.Build.cs
--- a/Engine/Source/Runtime/UMG/UMG.Build.cs
+++ b/Engine/Source/Runtime/UMG/UMG.Build.cs
@@ -26,6 +26,18 @@ public class UMG : ModuleRules
 			}
 		);
 
+        PublicDependencyModuleNames.AddRange(
+            new string[] {
+				"HTTP",
+			}
+        );
+
+        PrivateIncludePathModuleNames.AddRange(
+            new string[] {
+				"ImageWrapper",
+			}
+        );
+
 		if (Target.Type != TargetRules.TargetType.Server)
 		{
 			PrivateDependencyModuleNames.AddRange(
@@ -33,6 +45,12 @@ public class UMG : ModuleRules
 					"SlateRHIRenderer",
 				}
 			);
+
+            DynamicallyLoadedModuleNames.AddRange(
+                new string[] {
+				    "ImageWrapper",
+			    }
+            );
 		};
 	}
 }
