UMG - Fixing the build, excluding some setters that are only possible for editor only data.

diff --git a/Engine/Source/Runtime/UMG/Private/Blueprint/AsyncTaskDownloadImage.cpp b/Engine/Source/Runtime/UMG/Private/Blueprint/AsyncTaskDownloadImage.cpp
--- a/Engine/Source/Runtime/UMG/Private/Blueprint/AsyncTaskDownloadImage.cpp
+++ b/Engine/Source/Runtime/UMG/Private/Blueprint/AsyncTaskDownloadImage.cpp
@@ -38,8 +38,10 @@ static void WriteRawToTexture(UTexture2D* NewTexture2D, const TArray<uint8>& Raw
 
 	// Set options
 	NewTexture2D->SRGB = bUseSRGB;
+#if WITH_EDITORONLY_DATA
 	NewTexture2D->CompressionNone = true;
 	NewTexture2D->MipGenSettings = TMGS_NoMipmaps;
+#endif
 	NewTexture2D->CompressionSettings = TC_EditorIcon;
 
 	NewTexture2D->UpdateResource();
