UMG - When adding widgets to the viewport, you can now as an advanced parameter choose the ZOrder.  Additionally  ZOrder has been added as an option to the WidgetComponent.  Widget Components now correctly take resolution quality/scale into account when positioning the widget on the screen.

diff --git a/Engine/Source/Runtime/UMG/Private/Components/WidgetComponent.cpp b/Engine/Source/Runtime/UMG/Private/Components/WidgetComponent.cpp
--- a/Engine/Source/Runtime/UMG/Private/Components/WidgetComponent.cpp
+++ b/Engine/Source/Runtime/UMG/Private/Components/WidgetComponent.cpp
@@ -8,6 +8,7 @@
 	#include "ISlate3DRenderer.h"
 #endif // !UE_SERVER
 #include "DynamicMeshBuilder.h"
+#include "Scalability.h"
 
 class SVirtualWindow : public SWindow
 {
@@ -533,6 +534,17 @@ void UWidgetComponent::TickComponent(float DeltaTime, enum ELevelTick TickType,
 					if ( bProjected )
 					{
 						Widget->SetVisibility(ESlateVisibility::SelfHitTestInvisible);
+						
+						// If the user has configured a resolution quality we need to multiply
+						// the pixels by the resolution quality to arrive at the true position in
+						// the viewport, as the rendered image will be stretched to fill whatever
+						// size the viewport is at.
+						Scalability::FQualityLevels ScalabilityQuality = Scalability::GetQualityLevels();
+						float QualityScale = ( ScalabilityQuality.ResolutionQuality / 100.0f );
+
+						Widget->SetDesiredSizeInViewport(DrawSize);
+						Widget->SetPositionInViewport(ScreenLocation / QualityScale);
+						Widget->SetAlignmentInViewport(Pivot);
 					}
 					else
 					{
@@ -541,12 +553,8 @@ void UWidgetComponent::TickComponent(float DeltaTime, enum ELevelTick TickType,
 
 					if ( !Widget->IsInViewport() )
 					{
-						Widget->AddToViewport();
+						Widget->AddToViewport(ZOrder);
 					}
-
-					Widget->SetDesiredSizeInViewport(DrawSize);
-					Widget->SetPositionInViewport(ScreenLocation);
-					Widget->SetAlignmentInViewport(Pivot);
 				}
 			}
 			else if ( Widget->IsInViewport() )
diff --git a/Engine/Source/Runtime/UMG/Private/UserWidget.cpp b/Engine/Source/Runtime/UMG/Private/UserWidget.cpp
--- a/Engine/Source/Runtime/UMG/Private/UserWidget.cpp
+++ b/Engine/Source/Runtime/UMG/Private/UserWidget.cpp
@@ -608,7 +608,7 @@ UWidget* UUserWidget::GetRootWidget() const
 	return nullptr;
 }
 
-void UUserWidget::AddToViewport()
+void UUserWidget::AddToViewport(int32 ZOrder)
 {
 	if ( !FullScreenWidget.IsValid() )
 	{
@@ -623,7 +623,7 @@ void UUserWidget::AddToViewport()
 		{
 			if ( UGameViewportClient* ViewportClient = World->GetGameViewport() )
 			{
-				ViewportClient->AddViewportWidgetContent(RootWidget, 10);
+				ViewportClient->AddViewportWidgetContent(RootWidget, 10 + ZOrder);
 			}
 		}
 	}
diff --git a/Engine/Source/Runtime/UMG/Public/Blueprint/UserWidget.h b/Engine/Source/Runtime/UMG/Public/Blueprint/UserWidget.h
--- a/Engine/Source/Runtime/UMG/Public/Blueprint/UserWidget.h
+++ b/Engine/Source/Runtime/UMG/Public/Blueprint/UserWidget.h
@@ -140,11 +140,13 @@ public:
 	virtual void SetIsDesignTime(bool bInDesignTime) override;
 
 	/**
-	 * Adds it to the game's viewport, defaults to filling the entire viewport area.
-	 * @param bModal If this dialog should steal keyboard/mouse focus and consume all input. Great for a fullscreen menu. Terrible for HUDs.
+	 * Adds it to the game's viewport and fills the entire screen, unless SetDesiredSizeInViewport is called
+	 * to explicitly set the size.
+	 *
+	 * @param ZOrder The higher the number the greater the priority when determining if it will above other widgets.
 	 */
-	UFUNCTION(BlueprintCallable, BlueprintCosmetic, Category="User Interface|Viewport")
-	void AddToViewport();
+	UFUNCTION(BlueprintCallable, BlueprintCosmetic, Category="User Interface|Viewport", meta=( AdvancedDisplay = "ZOrder" ))
+	void AddToViewport(int32 ZOrder = 0);
 
 	/**
 	 * Removes the widget from the viewport.
diff --git a/Engine/Source/Runtime/UMG/Public/Components/WidgetComponent.h b/Engine/Source/Runtime/UMG/Public/Components/WidgetComponent.h
--- a/Engine/Source/Runtime/UMG/Public/Components/WidgetComponent.h
+++ b/Engine/Source/Runtime/UMG/Public/Components/WidgetComponent.h
@@ -107,6 +107,10 @@ private:
 	/** The Alignment/Pivot point that the widget is placed at relative to the position. */
 	UPROPERTY(EditAnywhere, Category=UI)
 	FVector2D Pivot;
+
+	/** The Screen Space ZOrder that is used if the widget is displayed in screen space. */
+	UPROPERTY(EditDefaultsOnly, Category=UI)
+	int32 ZOrder;
 	
 	/** The maximum distance from which a player can interact with this widget */
 	UPROPERTY(EditAnywhere, Category=UI, meta=(ClampMin="0.0", UIMax="5000.0", ClampMax="100000.0"))
