UMG - SetKeyboardFocus is now correctly marked as non-const so it should now appear as an executable node in blueprints.

diff --git a/Engine/Source/Runtime/UMG/Private/Components/Widget.cpp b/Engine/Source/Runtime/UMG/Private/Components/Widget.cpp
--- a/Engine/Source/Runtime/UMG/Private/Components/Widget.cpp
+++ b/Engine/Source/Runtime/UMG/Private/Components/Widget.cpp
@@ -296,7 +296,7 @@ bool UWidget::HasMouseCapture() const
 	return false;
 }
 
-void UWidget::SetKeyboardFocus() const
+void UWidget::SetKeyboardFocus()
 {
 	TSharedPtr<SWidget> SafeWidget = GetCachedWidget();
 	if (SafeWidget.IsValid())
diff --git a/Engine/Source/Runtime/UMG/Public/Components/Widget.h b/Engine/Source/Runtime/UMG/Public/Components/Widget.h
--- a/Engine/Source/Runtime/UMG/Public/Components/Widget.h
+++ b/Engine/Source/Runtime/UMG/Public/Components/Widget.h
@@ -227,7 +227,7 @@ public:
 
 	/** Sets the focus to this widget. */
 	UFUNCTION(BlueprintCallable, Category="Widget")
-	void SetKeyboardFocus() const;
+	void SetKeyboardFocus();
 
 	/** Forces the underlying slate system to perform a pre-pass on the layout of the widget.  This is for advanced users. */
 	UFUNCTION(BlueprintCallable, Category="Widget")
