UMG - Doing some general cleanup on the widget navigation code.  Some code documentation, a little organizing.

diff --git a/Engine/Source/Runtime/SlateCore/Public/Input/NavigationReply.h b/Engine/Source/Runtime/SlateCore/Public/Input/NavigationReply.h
--- a/Engine/Source/Runtime/SlateCore/Public/Input/NavigationReply.h
+++ b/Engine/Source/Runtime/SlateCore/Public/Input/NavigationReply.h
@@ -6,14 +6,23 @@
 
 class SWidget;
 
-UENUM()
-enum class EUINavigationRule
+UENUM(BlueprintType)
+enum class EUINavigationRule : uint8
 {
+	/** Allow the movement to continue in that direction, seeking the next navigable widget automatically. */
 	Escape,
+	/** Move to a specific widget. */
 	Explicit,
+	/**
+	 * Wrap movement inside this container, causing the movement to cycle around from the opposite side, 
+	 * if the navigation attempt would have escaped.
+	 */
 	Wrap,
+	/** Stops movement in this direction */
 	Stop,
+	/** Custom navigation handled by user code. */
 	Custom,
+	/** Invalid Rule */
 	Invalid
 };
 
diff --git a/Engine/Source/Runtime/UMG/Private/UWidgetNavigation.cpp b/Engine/Source/Runtime/UMG/Private/UWidgetNavigation.cpp
--- a/Engine/Source/Runtime/UMG/Private/UWidgetNavigation.cpp
+++ b/Engine/Source/Runtime/UMG/Private/UWidgetNavigation.cpp
@@ -9,3 +9,55 @@ UWidgetNavigation::UWidgetNavigation(const FObjectInitializer& ObjectInitializer
 	: Super(ObjectInitializer)
 {
 }
+
+#if WITH_EDITOR
+
+EUINavigationRule UWidgetNavigation::GetNavigationRule(EUINavigation Nav)
+{
+	switch ( Nav )
+	{
+	case EUINavigation::Up:
+		return Up.Rule;
+	case EUINavigation::Down:
+		return Down.Rule;
+	case EUINavigation::Left:
+		return Left.Rule;
+	case EUINavigation::Right:
+		return Right.Rule;
+	}
+	return EUINavigationRule::Escape;
+}
+
+#endif
+
+void UWidgetNavigation::UpdateMetaData(TSharedRef<FNavigationMetaData> MetaData)
+{
+	UpdateMetaDataEntry(MetaData, Up, EUINavigation::Up);
+	UpdateMetaDataEntry(MetaData, Down, EUINavigation::Down);
+	UpdateMetaDataEntry(MetaData, Left, EUINavigation::Left);
+	UpdateMetaDataEntry(MetaData, Right, EUINavigation::Right);
+}
+
+bool UWidgetNavigation::IsDefault() const
+{
+	return Up.Rule == EUINavigationRule::Escape &&
+		Down.Rule == EUINavigationRule::Escape &&
+		Left.Rule == EUINavigationRule::Escape &&
+		Right.Rule == EUINavigationRule::Escape;
+}
+
+void UWidgetNavigation::UpdateMetaDataEntry(TSharedRef<FNavigationMetaData> MetaData, const FWidgetNavigationData & NavData, EUINavigation Nav)
+{
+	switch ( NavData.Rule )
+	{
+	case EUINavigationRule::Escape:
+		MetaData->SetNavigationEscape(Nav);
+		break;
+	case EUINavigationRule::Stop:
+		MetaData->SetNavigationStop(Nav);
+		break;
+	case EUINavigationRule::Wrap:
+		MetaData->SetNavigationWrap(Nav);
+		break;
+	}
+}
diff --git a/Engine/Source/Runtime/UMG/Public/Blueprint/WidgetNavigation.h b/Engine/Source/Runtime/UMG/Public/Blueprint/WidgetNavigation.h
--- a/Engine/Source/Runtime/UMG/Public/Blueprint/WidgetNavigation.h
+++ b/Engine/Source/Runtime/UMG/Public/Blueprint/WidgetNavigation.h
@@ -17,7 +17,7 @@ struct UMG_API FWidgetNavigationData
 
 public:
 	UPROPERTY(EditAnywhere, BlueprintReadOnly, Category="Navigation")
-	TEnumAsByte<EUINavigationRule> Rule;
+	EUINavigationRule Rule;
 
 	UPROPERTY(EditAnywhere, BlueprintReadOnly, Category="Navigation")
 	FName WidgetToFocus;
@@ -48,53 +48,22 @@ public:
 	UPROPERTY(EditAnywhere, BlueprintReadOnly, Category="Navigation")
 	FWidgetNavigationData Right;
 
-	EUINavigationRule GetNavigationRule(EUINavigation Nav)
-	{
-		switch (Nav)
-		{
-		case EUINavigation::Up:
-			return Up.Rule;
-		case EUINavigation::Down:
-			return Down.Rule;
-		case EUINavigation::Left:
-			return Left.Rule;
-		case EUINavigation::Right:
-			return Right.Rule;
-		}
-		return EUINavigationRule::Escape;
-	}
-
-	void UpdateMetaData(TSharedRef<FNavigationMetaData> MetaData)
-	{
-		UpdateMetaDataEntry(MetaData, Up, EUINavigation::Up);
-		UpdateMetaDataEntry(MetaData, Down, EUINavigation::Down);
-		UpdateMetaDataEntry(MetaData, Left, EUINavigation::Left);
-		UpdateMetaDataEntry(MetaData, Right, EUINavigation::Right);
-	}
-
-	bool IsDefault() const
-	{
-		return Up.Rule == EUINavigationRule::Escape &&
-			Down.Rule == EUINavigationRule::Escape &&
-			Left.Rule == EUINavigationRule::Escape &&
-			Right.Rule == EUINavigationRule::Escape;
-	}
+public:
+
+#if WITH_EDITOR
+
+	/**  */
+	EUINavigationRule GetNavigationRule(EUINavigation Nav);
+
+#endif
+
+	/** Updates a slate metadata object to match this configured navigation ruleset. */
+	void UpdateMetaData(TSharedRef<FNavigationMetaData> MetaData);
+
+	/** @return true if the configured navigation object is the same as an un-customized navigation rule set. */
+	bool IsDefault() const;
 
 private:
 
-	void UpdateMetaDataEntry(TSharedRef<FNavigationMetaData> MetaData, const FWidgetNavigationData & NavData, EUINavigation Nav)
-	{
-		switch (NavData.Rule)
-		{
-			case EUINavigationRule::Escape:
-				MetaData->SetNavigationEscape(Nav);
-				break;
-			case EUINavigationRule::Stop:
-				MetaData->SetNavigationStop(Nav);
-				break;
-			case EUINavigationRule::Wrap:
-				MetaData->SetNavigationWrap(Nav);
-				break;
-		}
-	}
+	void UpdateMetaDataEntry(TSharedRef<FNavigationMetaData> MetaData, const FWidgetNavigationData & NavData, EUINavigation Nav);
 };
diff --git a/Engine/Source/Runtime/UMG/Public/Components/Widget.h b/Engine/Source/Runtime/UMG/Public/Components/Widget.h
--- a/Engine/Source/Runtime/UMG/Public/Components/Widget.h
+++ b/Engine/Source/Runtime/UMG/Public/Components/Widget.h
@@ -126,8 +126,12 @@ public:
 	UPROPERTY(EditDefaultsOnly, BlueprintReadOnly, Category="Render Transform", meta=( DisplayName="Pivot" ))
 	FVector2D RenderTransformPivot;
 
-	/**  */
-	UPROPERTY(Instanced, EditAnywhere, BlueprintReadOnly, Category="Navigation")
+	/**
+	 * The navigation object for this widget is optionally created if the user has configured custom
+	 * navigation rules for this widget in the widget designer.  Those rules determine how navigation transitions
+	 * can occur between widgets.
+	 */
+	UPROPERTY(Instanced, EditDefaultsOnly, BlueprintReadOnly, Category="Navigation")
 	class UWidgetNavigation* Navigation;
 
 #if WITH_EDITORONLY_DATA
