UE-7062 - UMG - Fixing many problems with the String Combobox.  Updating/Refreshing issues, crash bugs when clearing the list.

diff --git a/Engine/Source/Runtime/UMG/Private/Components/ComboBoxString.cpp b/Engine/Source/Runtime/UMG/Private/Components/ComboBoxString.cpp
--- a/Engine/Source/Runtime/UMG/Private/Components/ComboBoxString.cpp
+++ b/Engine/Source/Runtime/UMG/Private/Components/ComboBoxString.cpp
@@ -66,18 +66,25 @@ TSharedRef<SWidget> UComboBoxString::RebuildWidget()
 void UComboBoxString::AddOption(const FString& Option)
 {
 	Options.Add(MakeShareable(new FString(Option)));
+
+	RefreshOptions();
 }
 
 bool UComboBoxString::RemoveOption(const FString& Option)
 {
 	int32 OptionIndex = FindOptionIndex(Option);
-	if (Options[OptionIndex] == CurrentOptionPtr)
-	{
-		CurrentOptionPtr.Reset();
-	}
+
 	if ( OptionIndex != -1 )
 	{
+		if ( Options[OptionIndex] == CurrentOptionPtr )
+		{
+			ClearSelection();
+		}
+
 		Options.RemoveAt(OptionIndex);
+
+		RefreshOptions();
+
 		return true;
 	}
 
@@ -110,7 +117,37 @@ FString UComboBoxString::GetOptionAtIndex(int32 Index) const
 
 void UComboBoxString::ClearOptions()
 {
+	ClearSelection();
+
 	Options.Empty();
+
+	if ( MyComboBox.IsValid() )
+	{
+		MyComboBox->RefreshOptions();
+	}
+}
+
+void UComboBoxString::ClearSelection()
+{
+	CurrentOptionPtr.Reset();
+
+	if ( MyComboBox.IsValid() )
+	{
+		MyComboBox->ClearSelection();
+	}
+
+	if ( ComoboBoxContent.IsValid() )
+	{
+		ComoboBoxContent->SetContent(SNullWidget::NullWidget);
+	}
+}
+
+void UComboBoxString::RefreshOptions()
+{
+	if ( MyComboBox.IsValid() )
+	{
+		MyComboBox->RefreshOptions();
+	}
 }
 
 void UComboBoxString::SetSelectedOption(FString Option)
@@ -119,7 +156,16 @@ void UComboBoxString::SetSelectedOption(FString Option)
 	if (InitialIndex != -1)
 	{
 		CurrentOptionPtr = Options[InitialIndex];
-		ComoboBoxContent->SetContent(HandleGenerateWidget(CurrentOptionPtr));
+
+		if ( ComoboBoxContent.IsValid() )
+		{
+			MyComboBox->SetSelectedItem(CurrentOptionPtr);
+			ComoboBoxContent->SetContent(HandleGenerateWidget(CurrentOptionPtr));
+		}
+		else
+		{
+			SelectedOption = Option;
+		}
 	}
 }
 
@@ -132,12 +178,19 @@ FString UComboBoxString::GetSelectedOption() const
 	return FString();
 }
 
+int32 UComboBoxString::GetOptionCount() const
+{
+	return Options.Num();
+}
+
 TSharedRef<SWidget> UComboBoxString::HandleGenerateWidget(TSharedPtr<FString> Item) const
 {
+	FString StringItem = Item.IsValid() ? *Item : FString();
+
 	// Call the user's delegate to see if they want to generate a custom widget bound to the data source.
 	if ( !IsDesignTime() && OnGenerateWidgetEvent.IsBound() )
 	{
-		UWidget* Widget = OnGenerateWidgetEvent.Execute(*Item);
+		UWidget* Widget = OnGenerateWidgetEvent.Execute(StringItem);
 		if ( Widget != NULL )
 		{
 			return Widget->TakeWidget();
@@ -145,7 +198,7 @@ TSharedRef<SWidget> UComboBoxString::HandleGenerateWidget(TSharedPtr<FString> It
 	}
 
 	// If a row wasn't generated just create the default one, a simple text block of the item's name.
-	return SNew(STextBlock).Text(FText::FromString(*Item));
+	return SNew(STextBlock).Text(FText::FromString(StringItem));
 }
 
 void UComboBoxString::HandleSelectionChanged(TSharedPtr<FString> Item, ESelectInfo::Type SelectionType)
@@ -154,7 +207,7 @@ void UComboBoxString::HandleSelectionChanged(TSharedPtr<FString> Item, ESelectIn
 
 	if ( !IsDesignTime() )
 	{
-		OnSelectionChanged.Broadcast(*Item, SelectionType);
+		OnSelectionChanged.Broadcast(Item.IsValid() ? *Item : FString(), SelectionType);
 	}
 
 	// When the selection changes we always generate another widget to represent the content area of the comobox.
diff --git a/Engine/Source/Runtime/UMG/Public/Components/ComboBoxString.h b/Engine/Source/Runtime/UMG/Public/Components/ComboBoxString.h
--- a/Engine/Source/Runtime/UMG/Public/Components/ComboBoxString.h
+++ b/Engine/Source/Runtime/UMG/Public/Components/ComboBoxString.h
@@ -72,18 +72,32 @@ public:
 	UFUNCTION(BlueprintCallable, Category="ComboBox")
 	int32 FindOptionIndex(const FString& Option) const;
 
-	UFUNCTION(BlueprintCallable, Category = "ComboBox")
+	UFUNCTION(BlueprintCallable, Category="ComboBox")
 	FString GetOptionAtIndex(int32 Index) const;
 
 	UFUNCTION(BlueprintCallable, Category="ComboBox")
 	void ClearOptions();
 
-	UFUNCTION(BlueprintCallable, Category = "ComboBox")
+	UFUNCTION(BlueprintCallable, Category="ComboBox")
+	void ClearSelection();
+
+	/**
+	 * Refreshes the list of options.  If you added new ones, and want to update the list even if it's
+	 * currently being displayed use this.
+	 */
+	UFUNCTION(BlueprintCallable, Category="ComboBox")
+	void RefreshOptions();
+
+	UFUNCTION(BlueprintCallable, Category="ComboBox")
 	void SetSelectedOption(FString Option);
 
-	UFUNCTION(BlueprintCallable, Category = "ComboBox")
+	UFUNCTION(BlueprintCallable, Category="ComboBox")
 	FString GetSelectedOption() const;
 
+	/** @return The number of options */
+	UFUNCTION(BlueprintCallable, Category="ComboBox")
+	int32 GetOptionCount() const;
+
 	virtual void ReleaseSlateResources(bool bReleaseChildren) override;
 
 #if WITH_EDITOR
