UE-6978 - No longer duplicating the combobox entries on show/hide.  It now only loads the data once on load, the fields are also now private, as they are only useful at design time.

diff --git a/Engine/Source/Runtime/UMG/Private/Components/ComboBoxString.cpp b/Engine/Source/Runtime/UMG/Private/Components/ComboBoxString.cpp
--- a/Engine/Source/Runtime/UMG/Private/Components/ComboBoxString.cpp
+++ b/Engine/Source/Runtime/UMG/Private/Components/ComboBoxString.cpp
@@ -26,13 +26,19 @@ void UComboBoxString::ReleaseSlateResources(bool bReleaseChildren)
 	ComoboBoxContent.Reset();
 }
 
-TSharedRef<SWidget> UComboBoxString::RebuildWidget()
+void UComboBoxString::PostLoad()
 {
-	for ( FString& DefaultOptions : DefaultOptions )
+	Super::PostLoad();
+
+	// Initialize the set of options from the default set only once.
+	for ( const FString& DefaultOptions : DefaultOptions )
 	{
 		AddOption(DefaultOptions);
 	}
+}
 
+TSharedRef<SWidget> UComboBoxString::RebuildWidget()
+{
 	int32 InitialIndex = FindOptionIndex(SelectedOption);
 	if ( InitialIndex != -1 )
 	{
diff --git a/Engine/Source/Runtime/UMG/Public/Components/ComboBoxString.h b/Engine/Source/Runtime/UMG/Public/Components/ComboBoxString.h
--- a/Engine/Source/Runtime/UMG/Public/Components/ComboBoxString.h
+++ b/Engine/Source/Runtime/UMG/Public/Components/ComboBoxString.h
@@ -17,9 +17,9 @@ class UMG_API UComboBoxString : public UWidget
 	DECLARE_DYNAMIC_MULTICAST_DELEGATE_TwoParams(FOnSelectionChangedEvent, FString, SelectedItem, ESelectInfo::Type, SelectionType);
 	DECLARE_DYNAMIC_MULTICAST_DELEGATE(FOnOpeningEvent);
 
-public:
+private:
 
-	/** The list of items to be displayed on the combobox. */
+	/** The default list of items to be displayed on the combobox. */
 	UPROPERTY(EditDefaultsOnly, Category=Content)
 	TArray<FString> DefaultOptions;
 
@@ -100,6 +100,10 @@ public:
 
 	virtual void ReleaseSlateResources(bool bReleaseChildren) override;
 
+	// Begin UObject interface
+	virtual void PostLoad() override;
+	// End of UObject interface
+
 #if WITH_EDITOR
 	virtual const FSlateBrush* GetEditorIcon() override;
 	virtual const FText GetPaletteCategory() override;
