[Github PR 718] bozaro - adding GInitRunaway at the start of each widget tick to give users more leway in runaway detection.

diff --git a/Engine/Source/Runtime/UMG/Private/UserWidget.cpp b/Engine/Source/Runtime/UMG/Private/UserWidget.cpp
--- a/Engine/Source/Runtime/UMG/Private/UserWidget.cpp
+++ b/Engine/Source/Runtime/UMG/Private/UserWidget.cpp
@@ -555,6 +555,8 @@ void UUserWidget::SetContentForSlot(FName SlotName, UWidget* Content)
 
 void UUserWidget::NativeTick(const FGeometry& MyGeometry, float InDeltaTime )
 {
+	GInitRunaway();
+
 	// Update active movie scenes
 	for( UUMGSequencePlayer* Player : ActiveSequencePlayers )
 	{
