UE-7234 - UMG - Now releasing the slate resources of a slot after removing the slot from a panel.  This fixes long standing weird bugs users were experiencing as a crash on exit after adding and removing multiple items a panel, either through drag/drop or other means.

diff --git a/Engine/Source/Runtime/UMG/Private/Components/PanelWidget.cpp b/Engine/Source/Runtime/UMG/Private/Components/PanelWidget.cpp
--- a/Engine/Source/Runtime/UMG/Private/Components/PanelWidget.cpp
+++ b/Engine/Source/Runtime/UMG/Private/Components/PanelWidget.cpp
@@ -85,11 +85,16 @@ bool UPanelWidget::RemoveChildAt(int32 Index)
 		Slot->Content->Slot = nullptr;
 	}
 
-	Slot->Parent = nullptr;
 	Slots.RemoveAt(Index);
 
 	OnSlotRemoved(Slot);
 
+	Slot->Parent = nullptr;
+	Slot->Content = nullptr;
+
+	const bool bReleaseChildren = true;
+	Slot->ReleaseSlateResources(bReleaseChildren);
+
 	return true;
 }
 
