UMG - No longer allowing the construction of abstract classes for widgets.  If the class has the abstract flag then we just return null.

diff --git a/Engine/Source/Runtime/UMG/Public/Blueprint/UserWidget.h b/Engine/Source/Runtime/UMG/Public/Blueprint/UserWidget.h
--- a/Engine/Source/Runtime/UMG/Public/Blueprint/UserWidget.h
+++ b/Engine/Source/Runtime/UMG/Public/Blueprint/UserWidget.h
@@ -686,7 +686,7 @@ private:
 template< class T >
 T* CreateWidget(UWorld* World, UClass* UserWidgetClass)
 {
-	if ( World == nullptr || !UserWidgetClass->IsChildOf(UUserWidget::StaticClass()) )
+	if ( World == nullptr || !UserWidgetClass->IsChildOf(UUserWidget::StaticClass()) || UserWidgetClass->HasAnyClassFlags(CLASS_Abstract) )
 	{
 		// TODO UMG Error?
 		return nullptr;
@@ -709,7 +709,7 @@ T* CreateWidget(UWorld* World, UClass* UserWidgetClass)
 template< class T >
 T* CreateWidget(APlayerController* OwningPlayer, UClass* UserWidgetClass)
 {
-	if ( OwningPlayer == nullptr || !UserWidgetClass->IsChildOf(UUserWidget::StaticClass()) )
+	if ( OwningPlayer == nullptr || !UserWidgetClass->IsChildOf(UUserWidget::StaticClass()) || UserWidgetClass->HasAnyClassFlags(CLASS_Abstract) )
 	{
 		// TODO UMG Error?
 		return nullptr;
@@ -735,7 +735,7 @@ T* CreateWidget(APlayerController* OwningPlayer, UClass* UserWidgetClass)
 template< class T >
 T* CreateWidget(UGameInstance* OwningGame, UClass* UserWidgetClass)
 {
-	if ( OwningGame == nullptr || !UserWidgetClass->IsChildOf(UUserWidget::StaticClass()) )
+	if ( OwningGame == nullptr || !UserWidgetClass->IsChildOf(UUserWidget::StaticClass()) || UserWidgetClass->HasAnyClassFlags(CLASS_Abstract) )
 	{
 		// TODO UMG Error?
 		return nullptr;
