UE-7337 UMG - UserWidgets now follow their outer chain looking for an outer that know what world they are in.  This is nessesary if you want to know what world a widget is in for player controller queries...etc from a nested UUserWidget that's defined statically.

diff --git a/Engine/Source/Runtime/UMG/Private/UserWidget.cpp b/Engine/Source/Runtime/UMG/Private/UserWidget.cpp
--- a/Engine/Source/Runtime/UMG/Private/UserWidget.cpp
+++ b/Engine/Source/Runtime/UMG/Private/UserWidget.cpp
@@ -158,7 +158,14 @@ void UUserWidget::PostInitProperties()
 
 UWorld* UUserWidget::GetWorld() const
 {
-	// Use the Player Context's world.
+	if ( HasAllFlags(RF_ClassDefaultObject) )
+	{
+		// If we are a CDO, we must return nullptr instead of calling Outer->GetWorld() to fool UObject::ImplementsGetWorld.
+		return nullptr;
+	}
+
+	// Use the Player Context's world, if a specific player context is given, otherwise fall back to
+	// following the outer chain.
 	if ( PlayerContext.IsValid() )
 	{
 		if ( UWorld* World = PlayerContext.GetWorld() )
@@ -167,18 +174,19 @@ UWorld* UUserWidget::GetWorld() const
 		}
 	}
 
-	// If the current player context doesn't have a world or isn't valid, return the game instance's world.
-	if ( UGameInstance* GameInstance = Cast<UGameInstance>(GetOuter()) )
+	// Could be a GameInstance, could be World, could also be a WidgetTree, so we're just going to follow
+	// the outer chain to find the world we're in.
+	UObject* Outer = GetOuter();
+
+	while ( Outer )
 	{
-		if ( UWorld* World = GameInstance->GetWorld() )
+		UWorld* World = Outer->GetWorld();
+		if ( World )
 		{
 			return World;
 		}
-	}
 
-	if ( UWorld* World = Cast<UWorld>(GetOuter()) )
-	{
-		return World;
+		Outer = Outer->GetOuter();
 	}
 
 	return nullptr;
