UMG - The default value for text in a text property is now empty text, the "Text Block" value is now a modification after construction from the palette.

diff --git a/Engine/Source/Runtime/UMG/Private/Components/TextBlock.cpp b/Engine/Source/Runtime/UMG/Private/Components/TextBlock.cpp
--- a/Engine/Source/Runtime/UMG/Private/Components/TextBlock.cpp
+++ b/Engine/Source/Runtime/UMG/Private/Components/TextBlock.cpp
@@ -12,7 +12,6 @@ UTextBlock::UTextBlock(const FObjectInitializer& ObjectInitializer)
 {
 	bIsVariable = false;
 
-	Text = LOCTEXT("TextBlockDefaultValue", "Text Block");
 	ShadowOffset = FVector2D(1.0f, 1.0f);
 	ColorAndOpacity = FLinearColor::White;
 	ShadowColorAndOpacity = FLinearColor::Transparent;
@@ -173,6 +172,11 @@ const FText UTextBlock::GetPaletteCategory()
 	return LOCTEXT("Common", "Common");
 }
 
+void UTextBlock::OnCreationFromPalette()
+{
+	Text = LOCTEXT("TextBlockDefaultValue", "Text Block");
+}
+
 #endif
 
 /////////////////////////////////////////////////////
diff --git a/Engine/Source/Runtime/UMG/Public/Components/TextBlock.h b/Engine/Source/Runtime/UMG/Public/Components/TextBlock.h
--- a/Engine/Source/Runtime/UMG/Public/Components/TextBlock.h
+++ b/Engine/Source/Runtime/UMG/Public/Components/TextBlock.h
@@ -135,6 +135,7 @@ public:
 	// UWidget interface
 	virtual const FSlateBrush* GetEditorIcon() override;
 	virtual const FText GetPaletteCategory() override;
+	virtual void OnCreationFromPalette() override;
 	// End UWidget interface
 
 	virtual FString GetLabelMetadata() const override;
