UMG - General code cleanup in Widget.h

diff --git a/Engine/Source/Runtime/UMG/Public/Components/Widget.h b/Engine/Source/Runtime/UMG/Public/Components/Widget.h
--- a/Engine/Source/Runtime/UMG/Public/Components/Widget.h
+++ b/Engine/Source/Runtime/UMG/Public/Components/Widget.h
@@ -13,16 +13,28 @@ class UPanelSlot;
 /**
  * Helper macro for binding to a delegate or using the constant value when constructing the underlying SWidget
  */
-#define OPTIONAL_BINDING(ReturnType, MemberName) ( MemberName ## Delegate.IsBound() && !IsDesignTime() ) ? TAttribute< ReturnType >::Create(MemberName ## Delegate.GetUObject(), MemberName ## Delegate.GetFunctionName()) : TAttribute< ReturnType >(MemberName)
+#define OPTIONAL_BINDING(ReturnType, MemberName)				\
+	( MemberName ## Delegate.IsBound() && !IsDesignTime() )		\
+	?															\
+		TAttribute< ReturnType >::Create(MemberName ## Delegate.GetUObject(), MemberName ## Delegate.GetFunctionName()) \
+	:															\
+		TAttribute< ReturnType >(MemberName)
 
 /**
  * Helper macro for binding to a delegate or using the constant value when constructing the underlying SWidget,
  * also allows a conversion function to be provided to convert between the SWidget value and the value exposed to UMG.
  */
-#define OPTIONAL_BINDING_CONVERT(ReturnType, MemberName, ConvertedType, ConversionFunction) ( MemberName ## Delegate.IsBound() && !IsDesignTime() ) ? TAttribute< ConvertedType >::Create(TAttribute< ConvertedType >::FGetter::CreateUObject(this, &ThisClass::ConversionFunction, TAttribute< ReturnType >::Create(MemberName ## Delegate.GetUObject(), MemberName ## Delegate.GetFunctionName()))) : ConversionFunction(TAttribute< ReturnType >(MemberName))
+#define OPTIONAL_BINDING_CONVERT(ReturnType, MemberName, ConvertedType, ConversionFunction) \
+		( MemberName ## Delegate.IsBound() && !IsDesignTime() )								\
+		?																					\
+			TAttribute< ConvertedType >::Create(TAttribute< ConvertedType >::FGetter::CreateUObject(this, &ThisClass::ConversionFunction, TAttribute< ReturnType >::Create(MemberName ## Delegate.GetUObject(), MemberName ## Delegate.GetFunctionName()))) \
+		:																					\
+			ConversionFunction(TAttribute< ReturnType >(MemberName))
+
+
 
 /**
- * This is the base class for all wrapped Slate controls that are exposed to UMG.
+ * This is the base class for all wrapped Slate controls that are exposed to UObjects.
  */
 UCLASS(Abstract, BlueprintType)
 class UMG_API UWidget : public UVisual
@@ -31,14 +43,13 @@ class UMG_API UWidget : public UVisual
 
 public:
 
-	// Common Bindings
+	// Common Bindings - If you add any new common binding, you must provide a UPropertyBinder for it.
+	//                   all primitive binding in UMG goes through native binding evaluators to prevent
+	//                   thunking through the VM.
 	DECLARE_DYNAMIC_DELEGATE_RetVal(bool, FGetBool);
 	DECLARE_DYNAMIC_DELEGATE_RetVal(float, FGetFloat);
 	DECLARE_DYNAMIC_DELEGATE_RetVal(int32, FGetInt32);
 	DECLARE_DYNAMIC_DELEGATE_RetVal(FText, FGetText);
-	//DECLARE_DYNAMIC_DELEGATE_RetVal(FVector2D, FGetVector2D);
-	//DECLARE_DYNAMIC_DELEGATE_RetVal(FVector, FGetVector);
-	//DECLARE_DYNAMIC_DELEGATE_RetVal(FMargin, FGetMargin);
 	DECLARE_DYNAMIC_DELEGATE_RetVal(FSlateColor, FGetSlateColor);
 	DECLARE_DYNAMIC_DELEGATE_RetVal(FLinearColor, FGetLinearColor);
 	DECLARE_DYNAMIC_DELEGATE_RetVal(FSlateBrush, FGetSlateBrush);
@@ -47,6 +58,7 @@ public:
 	DECLARE_DYNAMIC_DELEGATE_RetVal(ECheckBoxState, FGetCheckBoxState);
 	DECLARE_DYNAMIC_DELEGATE_RetVal(UWidget*, FGetWidget);
 
+	// Events
 	DECLARE_DYNAMIC_DELEGATE_RetVal_OneParam(UWidget*, FGenerateWidgetForString, FString, Item);
 	DECLARE_DYNAMIC_DELEGATE_RetVal_OneParam(UWidget*, FGenerateWidgetForObject, UObject*, Item);
 
