UE-7439 - UMG - When the text changes or the checkbox checked state changes, we now mirror the change in the UObject's copy of the data so that if the slate object goes away and is then recreated, we can recall the state for the slate object if the user is not persisting the data themselves.

diff --git a/Engine/Source/Runtime/UMG/Private/Components/CheckBox.cpp b/Engine/Source/Runtime/UMG/Private/Components/CheckBox.cpp
--- a/Engine/Source/Runtime/UMG/Private/Components/CheckBox.cpp
+++ b/Engine/Source/Runtime/UMG/Private/Components/CheckBox.cpp
@@ -121,17 +121,11 @@ void UCheckBox::SetCheckedState(ECheckBoxState InCheckedState)
 
 void UCheckBox::SlateOnCheckStateChangedCallback(ECheckBoxState NewState)
 {
+	CheckedState = NewState;
+
 	//@TODO: Choosing to treat Undetermined as Checked
 	const bool bWantsToBeChecked = NewState != ECheckBoxState::Unchecked;
-
-	if (OnCheckStateChanged.IsBound())
-	{
-		OnCheckStateChanged.Broadcast(bWantsToBeChecked);
-	}
-	else
-	{
-		CheckedState = NewState;
-	}
+	OnCheckStateChanged.Broadcast(bWantsToBeChecked);
 }
 
 void UCheckBox::PostLoad()
diff --git a/Engine/Source/Runtime/UMG/Private/Components/EditableTextBox.cpp b/Engine/Source/Runtime/UMG/Private/Components/EditableTextBox.cpp
--- a/Engine/Source/Runtime/UMG/Private/Components/EditableTextBox.cpp
+++ b/Engine/Source/Runtime/UMG/Private/Components/EditableTextBox.cpp
@@ -112,14 +112,16 @@ void UEditableTextBox::ClearError()
 	}
 }
 
-void UEditableTextBox::HandleOnTextChanged(const FText& Text)
+void UEditableTextBox::HandleOnTextChanged(const FText& InText)
 {
-	OnTextChanged.Broadcast(Text);
+	Text = InText;
+	OnTextChanged.Broadcast(InText);
 }
 
-void UEditableTextBox::HandleOnTextCommitted(const FText& Text, ETextCommit::Type CommitMethod)
+void UEditableTextBox::HandleOnTextCommitted(const FText& InText, ETextCommit::Type CommitMethod)
 {
-	OnTextCommitted.Broadcast(Text, CommitMethod);
+	Text = InText;
+	OnTextCommitted.Broadcast(InText, CommitMethod);
 }
 
 void UEditableTextBox::PostLoad()
