UMG - No longer routing paint and tick requests while GIsRoutingPostLoad is true, this can occur in editor builds when uncooked data needs to be streamed in.  When this is occuring if slate ticks blueprint code it can cause a crash due to checking for this flag, so we're just going to ignore tick and paint requests when this is happening.

diff --git a/Engine/Source/Runtime/UMG/Private/Slate/SObjectWidget.cpp b/Engine/Source/Runtime/UMG/Private/Slate/SObjectWidget.cpp
--- a/Engine/Source/Runtime/UMG/Private/Slate/SObjectWidget.cpp
+++ b/Engine/Source/Runtime/UMG/Private/Slate/SObjectWidget.cpp
@@ -47,6 +47,16 @@ void SObjectWidget::AddReferencedObjects(FReferenceCollector& Collector)
 
 void SObjectWidget::Tick( const FGeometry& AllottedGeometry, const double InCurrentTime, const float InDeltaTime )
 {
+#if WITH_EDITOR
+	if ( GIsRoutingPostLoad )
+	{
+		// In editor builds streamed in data can cause PostLoad to be called on objects, when this is happening
+		// Slate Tick can and will Occur due to a Slow Task dialog being launched.  In order to prevent UMG ticking
+		// when this is true, we ignore Slate ticks when GIsRoutingPostLoad is true in editor builds.
+		return;
+	}
+#endif
+
 	if ( WidgetObject && !WidgetObject->IsDesignTime() )
 	{
 		return WidgetObject->NativeTick(AllottedGeometry, InDeltaTime);
@@ -55,6 +65,16 @@ void SObjectWidget::Tick( const FGeometry& AllottedGeometry, const double InCurr
 
 int32 SObjectWidget::OnPaint(const FPaintArgs& Args, const FGeometry& AllottedGeometry, const FSlateRect& MyClippingRect, FSlateWindowElementList& OutDrawElements, int32 LayerId, const FWidgetStyle& InWidgetStyle, bool bParentEnabled) const
 {
+#if WITH_EDITOR
+	if ( GIsRoutingPostLoad )
+	{
+		// In editor builds streamed in data can cause PostLoad to be called on objects, when this is happening
+		// Slate painting can and will Occur due to a Slow Task dialog being launched.  In order to prevent UMG painting
+		// when this is true, we ignore Slate painting when GIsRoutingPostLoad is true in editor builds.
+		return LayerId;
+	}
+#endif
+
 	int32 MaxLayer = SCompoundWidget::OnPaint(Args, AllottedGeometry, MyClippingRect, OutDrawElements, LayerId, InWidgetStyle, bParentEnabled);
 
 	if ( WidgetObject && !WidgetObject->IsDesignTime() )
