UMG - Adding ExpandableArea as an experimental widget, demonstraits how you can use the named slot interface on a non-user widget to provide multiple slots, but not N slots like a normal panel.

diff --git a/Engine/Source/Runtime/UMG/Private/Components/ExpandableArea.cpp b/Engine/Source/Runtime/UMG/Private/Components/ExpandableArea.cpp
new file mode 100644
--- /dev/null
+++ b/Engine/Source/Runtime/UMG/Private/Components/ExpandableArea.cpp
@@ -0,0 +1,153 @@
+// Copyright 1998-2015 Epic Games, Inc. All Rights Reserved.
+
+#include "UMGPrivatePCH.h"
+
+#include "SExpandableArea.h"
+
+#define LOCTEXT_NAMESPACE "UMG"
+
+/////////////////////////////////////////////////////
+// UExpandableArea
+
+static const FName HeaderName(TEXT("Header"));
+static const FName BodyName(TEXT("Body"));
+
+UExpandableArea::UExpandableArea(const FObjectInitializer& ObjectInitializer)
+	: Super(ObjectInitializer)
+	, bIsExpanded(false)
+{
+	bIsVariable = false;
+}
+
+bool UExpandableArea::GetIsExpanded() const
+{
+	if ( MyExpandableArea.IsValid() )
+	{
+		return MyExpandableArea->IsExpanded();
+	}
+
+	return bIsExpanded;
+}
+
+void UExpandableArea::SetIsExpanded(bool IsExpanded)
+{
+	bIsExpanded = IsExpanded;
+	if ( MyExpandableArea.IsValid() )
+	{
+		MyExpandableArea->SetExpanded(IsExpanded);
+	}
+}
+
+void UExpandableArea::ReleaseSlateResources(bool bReleaseChildren)
+{
+	Super::ReleaseSlateResources(bReleaseChildren);
+
+	MyExpandableArea.Reset();
+}
+
+void UExpandableArea::GetSlotNames(TArray<FName>& SlotNames) const
+{
+	SlotNames.Add(HeaderName);
+	SlotNames.Add(BodyName);
+}
+
+UWidget* UExpandableArea::GetContentForSlot(FName SlotName) const
+{
+	if ( SlotName == HeaderName )
+	{
+		return HeaderContent;
+	}
+	else if ( SlotName == BodyName )
+	{
+		return BodyContent;
+	}
+
+	return nullptr;
+}
+
+void UExpandableArea::SetContentForSlot(FName SlotName, UWidget* Content)
+{
+	if ( SlotName == HeaderName )
+	{
+		HeaderContent = Content;
+	}
+	else if ( SlotName == BodyName )
+	{
+		BodyContent = Content;
+	}
+}
+
+TSharedRef<SWidget> UExpandableArea::RebuildWidget()
+{
+	TSharedRef<SWidget> HeaderWidget = HeaderContent ? HeaderContent->TakeWidget() : SNullWidget::NullWidget;
+	TSharedRef<SWidget> BodyWidget = BodyContent ? BodyContent->TakeWidget() : SNullWidget::NullWidget;
+
+	MyExpandableArea = SNew(SExpandableArea)
+		.MaxHeight(MaxHeight)
+		.Padding(AreaPadding)
+		.OnAreaExpansionChanged(BIND_UOBJECT_DELEGATE(FOnBooleanValueChanged, SlateExpansionChanged))
+		.HeaderContent()
+		[
+			HeaderWidget
+		]
+		.BodyContent()
+		[
+			BodyWidget
+		];
+
+	return MyExpandableArea.ToSharedRef();
+}
+
+void UExpandableArea::SynchronizeProperties()
+{
+	Super::SynchronizeProperties();
+
+	MyExpandableArea->SetExpanded(bIsExpanded);
+}
+
+void UExpandableArea::SlateExpansionChanged(bool NewState)
+{
+	bIsExpanded = NewState;
+
+	if ( OnExpansionChanged.IsBound() )
+	{
+		OnExpansionChanged.Broadcast(NewState);
+	}
+}
+
+#if WITH_EDITOR
+
+const FSlateBrush* UExpandableArea::GetEditorIcon()
+{
+	return FUMGStyle::Get().GetBrush("Widget.Spacer");
+}
+
+const FText UExpandableArea::GetPaletteCategory()
+{
+	return LOCTEXT("Extra", "Extra");
+}
+
+void UExpandableArea::OnDescendantSelected(UWidget* DescendantWidget)
+{
+	// Temporarily sets the active child to the selected child to make
+	// dragging and dropping easier in the editor.
+	UWidget* SelectedChild = UWidget::FindChildContainingDescendant(BodyContent, DescendantWidget);
+	if ( SelectedChild )
+	{
+		MyExpandableArea->SetExpanded(true);
+	}
+}
+
+void UExpandableArea::OnDescendantDeselected(UWidget* DescendantWidget)
+{
+	if ( MyExpandableArea.IsValid() )
+	{
+		MyExpandableArea->SetExpanded(bIsExpanded);
+	}
+}
+
+#endif
+
+/////////////////////////////////////////////////////
+
+#undef LOCTEXT_NAMESPACE
diff --git a/Engine/Source/Runtime/UMG/Private/Components/Widget.cpp b/Engine/Source/Runtime/UMG/Private/Components/Widget.cpp
--- a/Engine/Source/Runtime/UMG/Private/Components/Widget.cpp
+++ b/Engine/Source/Runtime/UMG/Private/Components/Widget.cpp
@@ -749,6 +749,11 @@ void UWidget::GatherAllChildren(UWidget* Root, TSet<UWidget*>& Children)
 
 UWidget* UWidget::FindChildContainingDescendant(UWidget* Root, UWidget* Descendant)
 {
+	if ( Root == nullptr )
+	{
+		return nullptr;
+	}
+
 	UWidget* Parent = Descendant->GetParent();
 
 	while ( Parent != nullptr )
diff --git a/Engine/Source/Runtime/UMG/Public/Components/ExpandableArea.h b/Engine/Source/Runtime/UMG/Public/Components/ExpandableArea.h
new file mode 100644
--- /dev/null
+++ b/Engine/Source/Runtime/UMG/Public/Components/ExpandableArea.h
@@ -0,0 +1,80 @@
+﻿// Copyright 1998-2015 Epic Games, Inc. All Rights Reserved.
+
+#pragma once
+
+#include "ExpandableArea.generated.h"
+
+DECLARE_DYNAMIC_MULTICAST_DELEGATE_OneParam(FOnExpandableAreaExpansionChanged, bool, bIsExpanded);
+
+/**
+ * 
+ */
+UCLASS(Experimental)
+class UMG_API UExpandableArea : public UWidget, public INamedSlotInterface
+{
+	GENERATED_UCLASS_BODY()
+
+public:
+	/**  */
+	UPROPERTY(EditDefaultsOnly, BlueprintReadOnly, Category="Expansion")
+	bool bIsExpanded;
+
+	/** The maximum height of the area */
+	UPROPERTY(EditDefaultsOnly, BlueprintReadOnly, Category="Expansion")
+	float MaxHeight;
+	
+	/**  */
+	UPROPERTY(EditDefaultsOnly, BlueprintReadOnly, Category="Expansion")
+	FMargin AreaPadding;
+
+	/** A bindable delegate for the IsChecked. */
+	UPROPERTY(BlueprintAssignable, Category="ExpandableArea|Event")
+	FOnExpandableAreaExpansionChanged OnExpansionChanged;
+
+public:
+
+	UFUNCTION(BlueprintCallable, Category="Expansion")
+	bool GetIsExpanded() const;
+
+	UFUNCTION(BlueprintCallable, Category="Expansion")
+	void SetIsExpanded(bool IsExpanded);
+	
+	// Begin INamedSlotInterface
+	virtual void GetSlotNames(TArray<FName>& SlotNames) const;
+	virtual UWidget* GetContentForSlot(FName SlotName) const;
+	virtual void SetContentForSlot(FName SlotName, UWidget* Content);
+	// End INamedSlotInterface
+
+public:
+	
+	// UWidget interface
+	virtual void SynchronizeProperties() override;
+	// End of UWidget interface
+
+	// UVisual interface
+	virtual void ReleaseSlateResources(bool bReleaseChildren) override;
+	// End of UVisual interface
+
+#if WITH_EDITOR
+	virtual const FSlateBrush* GetEditorIcon() override;
+	virtual const FText GetPaletteCategory() override;
+	virtual void OnDescendantSelected(UWidget* DescendantWidget) override;
+	virtual void OnDescendantDeselected(UWidget* DescendantWidget) override;
+#endif
+
+protected:
+	// UWidget interface
+	virtual TSharedRef<SWidget> RebuildWidget() override;
+	// End of UWidget interface
+
+	void SlateExpansionChanged(bool NewState);
+
+protected:
+	UPROPERTY()
+	UWidget* HeaderContent;
+
+	UPROPERTY()
+	UWidget* BodyContent;
+
+	TSharedPtr<SExpandableArea> MyExpandableArea;
+};
diff --git a/Engine/Source/Runtime/UMG/Public/UMG.h b/Engine/Source/Runtime/UMG/Public/UMG.h
--- a/Engine/Source/Runtime/UMG/Public/UMG.h
+++ b/Engine/Source/Runtime/UMG/Public/UMG.h
@@ -102,6 +102,8 @@
 #include "ScaleBoxSlot.h"
 #include "ScaleBox.h"
 
+#include "ExpandableArea.h"
+
 #include "Spacer.h"
 #include "MenuAnchor.h"
 #include "ScrollBar.h"
