UE-7626 - UMG - Improvements to Drag/Drop.  Now allows you to choose a pivot, and offset.  The drag data can no longer be GCed while drag is being performed.  The drag pivot defaults to Center,Center.  In order to reduce drag/drop popping on begin drag there is now a small interpolation that occurs to move the object to the chosen pivot/offset.  The DefaultDragVisual is now correctly marked as ReadOnly, it was previously set to read/write.  This was inaccurate, as it couldn't actually be replaced after drag began.

diff --git a/Engine/Source/Runtime/UMG/Private/DragDropOperation.cpp b/Engine/Source/Runtime/UMG/Private/DragDropOperation.cpp
--- a/Engine/Source/Runtime/UMG/Private/DragDropOperation.cpp
+++ b/Engine/Source/Runtime/UMG/Private/DragDropOperation.cpp
@@ -8,6 +8,7 @@
 UDragDropOperation::UDragDropOperation(const FObjectInitializer& ObjectInitializer)
 	: Super(ObjectInitializer)
 {
+	Pivot = EDragPivot::CenterCenter;
 }
 
 void UDragDropOperation::Drop_Implementation(const FPointerEvent& PointerEvent)
diff --git a/Engine/Source/Runtime/UMG/Private/Slate/UMGDragDropOp.cpp b/Engine/Source/Runtime/UMG/Private/Slate/UMGDragDropOp.cpp
--- a/Engine/Source/Runtime/UMG/Private/Slate/UMGDragDropOp.cpp
+++ b/Engine/Source/Runtime/UMG/Private/Slate/UMGDragDropOp.cpp
@@ -6,23 +6,33 @@
 
 //////////////////////////////////////////////////////////////////////////
 // FUMGDragDropOp
+
+FUMGDragDropOp::FUMGDragDropOp()
+{
+	StartTime = FSlateApplicationBase::Get().GetCurrentTime();
+}
+
+void FUMGDragDropOp::AddReferencedObjects(FReferenceCollector& Collector)
+{
+	Collector.AddReferencedObject(DragOperation);
+}
+
 void FUMGDragDropOp::OnDrop( bool bDropWasHandled, const FPointerEvent& MouseEvent )
 {
-	UDragDropOperation* OperationObject = DragOperation.Get();
-	if ( OperationObject )
+	if ( DragOperation )
 	{
 		if ( bDropWasHandled )
 		{
-			OperationObject->Drop(MouseEvent);
+			DragOperation->Drop(MouseEvent);
 		}
 		else
 		{
 			if ( SourceUserWidget.IsValid() )
 			{
-				SourceUserWidget->OnDragCancelled(FDragDropEvent(MouseEvent, AsShared()), OperationObject);
+				SourceUserWidget->OnDragCancelled(FDragDropEvent(MouseEvent, AsShared()), DragOperation);
 			}
 
-			OperationObject->DragCancelled(MouseEvent);
+			DragOperation->DragCancelled(MouseEvent);
 		}
 	}
 	
@@ -31,15 +41,67 @@ void FUMGDragDropOp::OnDrop( bool bDropWasHandled, const FPointerEvent& MouseEve
 
 void FUMGDragDropOp::OnDragged( const class FDragDropEvent& DragDropEvent )
 {
-	UDragDropOperation* OperationObject = DragOperation.Get();
-	if ( OperationObject )
+	if ( DragOperation )
+	{
+		DragOperation->Dragged(DragDropEvent);
+	}
+
+	FVector2D CachedDesiredSize = DecoratorWidget->GetDesiredSize();
+
+	FVector2D Position = DragDropEvent.GetScreenSpacePosition();
+	Position += CachedDesiredSize * DragOperation->Offset;
+
+	switch ( DragOperation->Pivot )
 	{
-		OperationObject->Dragged(DragDropEvent);
+	case EDragPivot::MouseDown:
+		Position += MouseDownOffset;
+		break;
+
+	case EDragPivot::TopLeft:
+		// Position is already Top Left.
+		break;
+	case EDragPivot::TopCenter:
+		Position -= CachedDesiredSize * FVector2D(0.5f, 0);
+		break;
+	case EDragPivot::TopRight:
+		Position -= CachedDesiredSize * FVector2D(1, 0);
+		break;
+
+	case EDragPivot::CenterLeft:
+		Position -= CachedDesiredSize * FVector2D(0, 0.5f);
+		break;
+	case EDragPivot::CenterCenter:
+		Position -= CachedDesiredSize * FVector2D(0.5f, 0.5f);
+		break;
+	case EDragPivot::CenterRight:
+		Position -= CachedDesiredSize * FVector2D(1.0f, 0.5f);
+		break;
+
+	case EDragPivot::BottomLeft:
+		Position -= CachedDesiredSize * FVector2D(0, 1);
+		break;
+	case EDragPivot::BottomCenter:
+		Position -= CachedDesiredSize * FVector2D(0.5f, 1);
+		break;
+	case EDragPivot::BottomRight:
+		Position -= CachedDesiredSize * FVector2D(1, 1);
+		break;
 	}
 
-	FVector2D Position = DragDropEvent.GetScreenSpacePosition() + Offset;
+	const double AnimationTime = 0.150;
+
+	double DeltaTime = FSlateApplicationBase::Get().GetCurrentTime() - StartTime;
 
-	CursorDecoratorWindow->MoveWindowTo(Position);
+	if ( DeltaTime < AnimationTime )
+	{
+		float T = DeltaTime / AnimationTime;
+		FVector2D LerpPosition = ( Position - StartingScreenPos ) * T;
+		CursorDecoratorWindow->MoveWindowTo(StartingScreenPos + LerpPosition);
+	}
+	else
+	{
+		CursorDecoratorWindow->MoveWindowTo(Position);
+	}
 }
 
 TSharedPtr<SWidget> FUMGDragDropOp::GetDefaultDecorator() const
@@ -49,8 +111,10 @@ TSharedPtr<SWidget> FUMGDragDropOp::GetDefaultDecorator() const
 
 TSharedRef<FUMGDragDropOp> FUMGDragDropOp::New(UDragDropOperation* InOperation, const FVector2D &CursorPosition, const FVector2D &ScreenPositionOfDragee, TSharedPtr<SObjectWidget> SourceUserWidget)
 {
-	TSharedRef<FUMGDragDropOp> Operation = MakeShareable(new FUMGDragDropOp);
-	Operation->Offset = ScreenPositionOfDragee - CursorPosition;
+	check(InOperation);
+
+	TSharedRef<FUMGDragDropOp> Operation = MakeShareable(new FUMGDragDropOp());
+	Operation->MouseDownOffset = ScreenPositionOfDragee - CursorPosition;
 	Operation->StartingScreenPos = ScreenPositionOfDragee;
 	Operation->SourceUserWidget = SourceUserWidget;
 
diff --git a/Engine/Source/Runtime/UMG/Public/Blueprint/DragDropOperation.h b/Engine/Source/Runtime/UMG/Public/Blueprint/DragDropOperation.h
--- a/Engine/Source/Runtime/UMG/Public/Blueprint/DragDropOperation.h
+++ b/Engine/Source/Runtime/UMG/Public/Blueprint/DragDropOperation.h
@@ -7,6 +7,26 @@
 DECLARE_DYNAMIC_MULTICAST_DELEGATE_OneParam(FOnDragDropMulticast, UDragDropOperation*, Operation);
 
 
+/**
+ * Controls where the drag widget visual will appear when dragged relative to the pointer performing
+ * the drag operation.
+ */
+UENUM(BlueprintType)
+enum class EDragPivot : uint8
+{
+	/**  */
+	MouseDown,
+	TopLeft,
+	TopCenter,
+	TopRight,
+	CenterLeft,
+	CenterCenter,
+	CenterRight,
+	BottomLeft,
+	BottomCenter,
+	BottomRight
+};
+
 /**
  * This class is the base drag drop operation for UMG, extend it to add additional data and add new functionality.
  */
@@ -16,18 +36,35 @@ class UMG_API UDragDropOperation : public UObject
 	GENERATED_UCLASS_BODY()
 	
 public:
-	/**  */
+	/** A simple string tag you can optionally use to provide extra metadata about the operation. */
 	UPROPERTY(EditAnywhere, BlueprintReadWrite, Category="Drag and Drop", meta=( ExposeOnSpawn="true" ))
 	FString Tag;
 	
-	/**  */
+	/**
+	 * The payload of the drag operation.  This can be any UObject that you want to pass along as dragged data.  If you 
+	 * were building an inventory screen this would be the UObject representing the item being moved to another slot.
+	 */
 	UPROPERTY(EditAnywhere, BlueprintReadWrite, Category="Drag and Drop", meta=( ExposeOnSpawn="true" ))
 	UObject* Payload;
 	
-	/**  */
-	UPROPERTY(EditAnywhere, BlueprintReadWrite, Category="Drag and Drop", meta=( ExposeOnSpawn="true" ))
+	/**
+	 * The Drag Visual is the widget to display when dragging the item.  Normally people create a new widget to represent the 
+	 * temporary drag.
+	 */
+	UPROPERTY(EditAnywhere, BlueprintReadOnly, Category="Drag and Drop", meta=( ExposeOnSpawn="true", DisplayName="Drag Visual" ))
 	class UWidget* DefaultDragVisual;
 
+	/**
+	 * Controls where the drag widget visual will appear when dragged relative to the pointer performing
+	 * the drag operation.
+	 */
+	UPROPERTY(EditAnywhere, BlueprintReadWrite, Category="Drag and Drop", meta=( ExposeOnSpawn="true" ))
+	EDragPivot Pivot;
+
+	/** A percentage offset (-1..+1) from the Pivot location, the percentage is of the desired size of the dragged visual. */
+	UPROPERTY(EditAnywhere, BlueprintReadWrite, Category="Drag and Drop", AdvancedDisplay, meta=( ExposeOnSpawn="true" ))
+	FVector2D Offset;
+
 	/**  */
 	UPROPERTY(BlueprintAssignable)
 	FOnDragDropMulticast OnDrop;
diff --git a/Engine/Source/Runtime/UMG/Public/Slate/UMGDragDropOp.h b/Engine/Source/Runtime/UMG/Public/Slate/UMGDragDropOp.h
--- a/Engine/Source/Runtime/UMG/Public/Slate/UMGDragDropOp.h
+++ b/Engine/Source/Runtime/UMG/Public/Slate/UMGDragDropOp.h
@@ -8,24 +8,29 @@
  * This is the drag/drop class used for UMG, all UMG drag drop operations utilize this operation.
  * It supports moving a UObject payload and using a UWidget decorator.
  */
-class FUMGDragDropOp : public FDragDropOperation
+class FUMGDragDropOp : public FDragDropOperation, public FGCObject
 {
 public:
 	DRAG_DROP_OPERATOR_TYPE(FUMGDragDropOp, FDragDropOperation)
 	
 	static TSharedRef<FUMGDragDropOp> New(UDragDropOperation* Operation, const FVector2D &CursorPosition, const FVector2D &ScreenPositionOfNode, TSharedPtr<SObjectWidget> SourceUserWidget);
 
+	FUMGDragDropOp();
+
+	virtual void AddReferencedObjects(FReferenceCollector& Collector) override;
+
 	virtual void OnDrop( bool bDropWasHandled, const FPointerEvent& MouseEvent ) override;
 
 	virtual void OnDragged( const class FDragDropEvent& DragDropEvent ) override;
 
 	virtual TSharedPtr<SWidget> GetDefaultDecorator() const override;
 
-	UDragDropOperation* GetOperation() const { return DragOperation.Get(); }
+	UDragDropOperation* GetOperation() const { return DragOperation; }
 
 private:
 
-	TWeakObjectPtr<UDragDropOperation> DragOperation;
+	// Raw pointer to the drag operation, kept alive by AddReferencedObjects.
+	UDragDropOperation* DragOperation;
 
 	/** Source User Widget */
 	TSharedPtr<SObjectWidget> SourceUserWidget;
@@ -34,8 +39,11 @@ private:
 	TSharedPtr<SWidget> DecoratorWidget;
 
 	/** The offset to use when dragging the object so that it says the same distance away from the mouse. */
-	FVector2D Offset;
+	FVector2D MouseDownOffset;
 
 	/** The starting screen location where the drag operation started. */
 	FVector2D StartingScreenPos;
+
+	/** Allows smooth interpolation of the dragged visual over a few frames. */
+	double StartTime;
 };
