UMG - Don't count collapsed items towards the desired size of a canvas panel.

diff --git a/Engine/Source/Runtime/UMG/Private/Slate/SConstraintCanvas.cpp b/Engine/Source/Runtime/UMG/Private/Slate/SConstraintCanvas.cpp
--- a/Engine/Source/Runtime/UMG/Private/Slate/SConstraintCanvas.cpp
+++ b/Engine/Source/Runtime/UMG/Private/Slate/SConstraintCanvas.cpp
@@ -176,23 +176,28 @@ FVector2D SConstraintCanvas::ComputeDesiredSize( float ) const
 	for ( int32 ChildIndex = 0; ChildIndex < Children.Num(); ++ChildIndex )
 	{
 		const SConstraintCanvas::FSlot& CurChild = Children[ChildIndex];
+		const TSharedRef<SWidget>& Widget = CurChild.GetWidget();
 
-		const FMargin Offset = CurChild.OffsetAttr.Get();
-		const FVector2D Alignment = CurChild.AlignmentAttr.Get();
-		const FAnchors Anchors = CurChild.AnchorsAttr.Get();
+		// As long as the widgets are not collapsed, they should contribute to the desired size.
+		if ( Widget->GetVisibility() != EVisibility::Collapsed )
+		{
+			const FMargin Offset = CurChild.OffsetAttr.Get();
+			const FVector2D Alignment = CurChild.AlignmentAttr.Get();
+			const FAnchors Anchors = CurChild.AnchorsAttr.Get();
 
-		const FVector2D SlotSize = FVector2D(Offset.Right, Offset.Bottom);
-		const FVector2D WidgetDesiredSize = CurChild.GetWidget()->GetDesiredSize();
+			const FVector2D SlotSize = FVector2D(Offset.Right, Offset.Bottom);
+			const FVector2D WidgetDesiredSize = Widget->GetDesiredSize();
 
-		const bool AutoSize = CurChild.AutoSizeAttr.Get();
+			const bool AutoSize = CurChild.AutoSizeAttr.Get();
 
-		const FVector2D Size = AutoSize ? WidgetDesiredSize : SlotSize;
+			const FVector2D Size = AutoSize ? WidgetDesiredSize : SlotSize;
 
-		const bool bIsDockedHorizontally = ( Anchors.Minimum.X == Anchors.Maximum.X ) && ( Anchors.Minimum.X == 0 || Anchors.Minimum.X == 1 );
-		const bool bIsDockedVertically = ( Anchors.Minimum.Y == Anchors.Maximum.Y ) && ( Anchors.Minimum.Y == 0 || Anchors.Minimum.Y == 1 );
+			const bool bIsDockedHorizontally = ( Anchors.Minimum.X == Anchors.Maximum.X ) && ( Anchors.Minimum.X == 0 || Anchors.Minimum.X == 1 );
+			const bool bIsDockedVertically = ( Anchors.Minimum.Y == Anchors.Maximum.Y ) && ( Anchors.Minimum.Y == 0 || Anchors.Minimum.Y == 1 );
 
-		DesiredSize.X = FMath::Max(DesiredSize.X, Size.X + ( bIsDockedHorizontally ? FMath::Abs(Offset.Left) : 0.0f ));
-		DesiredSize.Y = FMath::Max(DesiredSize.Y, Size.Y + ( bIsDockedVertically ? FMath::Abs(Offset.Top) : 0.0f ));
+			DesiredSize.X = FMath::Max(DesiredSize.X, Size.X + ( bIsDockedHorizontally ? FMath::Abs(Offset.Left) : 0.0f ));
+			DesiredSize.Y = FMath::Max(DesiredSize.Y, Size.Y + ( bIsDockedVertically ? FMath::Abs(Offset.Top) : 0.0f ));
+		}
 	}
 
 	return DesiredSize;
