UE-8046 - UMG - Throbbers are now capped at 25 Pieces.

diff --git a/Engine/Source/Runtime/UMG/Private/Components/CircularThrobber.cpp b/Engine/Source/Runtime/UMG/Private/Components/CircularThrobber.cpp
--- a/Engine/Source/Runtime/UMG/Private/Components/CircularThrobber.cpp
+++ b/Engine/Source/Runtime/UMG/Private/Components/CircularThrobber.cpp
@@ -33,7 +33,7 @@ TSharedRef<SWidget> UCircularThrobber::RebuildWidget()
 
 	MyCircularThrobber = SNew(SCircularThrobber)
 		.PieceImage(&Image)
-		.NumPieces(NumberOfPieces)
+		.NumPieces(FMath::Clamp(NumberOfPieces, 1, 25))
 		.Period(Period)
 		.Radius(Radius);
 
@@ -44,7 +44,7 @@ void UCircularThrobber::SynchronizeProperties()
 {
 	Super::SynchronizeProperties();
 
-	MyCircularThrobber->SetNumPieces(NumberOfPieces);
+	MyCircularThrobber->SetNumPieces(FMath::Clamp(NumberOfPieces, 1, 25));
 	MyCircularThrobber->SetPeriod(Period);
 	MyCircularThrobber->SetRadius(Radius);
 }
@@ -54,7 +54,7 @@ void UCircularThrobber::SetNumberOfPieces(int32 InNumberOfPieces)
 	NumberOfPieces = InNumberOfPieces;
 	if (MyCircularThrobber.IsValid())
 	{
-		MyCircularThrobber->SetNumPieces(InNumberOfPieces);
+		MyCircularThrobber->SetNumPieces(FMath::Clamp(NumberOfPieces, 1, 25));
 	}
 }
 
diff --git a/Engine/Source/Runtime/UMG/Private/Components/Throbber.cpp b/Engine/Source/Runtime/UMG/Private/Components/Throbber.cpp
--- a/Engine/Source/Runtime/UMG/Private/Components/Throbber.cpp
+++ b/Engine/Source/Runtime/UMG/Private/Components/Throbber.cpp
@@ -34,7 +34,7 @@ TSharedRef<SWidget> UThrobber::RebuildWidget()
 
 	MyThrobber = SNew(SThrobber)
 		.PieceImage(&Image)
-		.NumPieces(NumberOfPieces)
+		.NumPieces(FMath::Clamp(NumberOfPieces, 1, 25))
 		.Animate(GetAnimation());
 
 	return MyThrobber.ToSharedRef();
@@ -44,7 +44,7 @@ void UThrobber::SynchronizeProperties()
 {
 	Super::SynchronizeProperties();
 
-	MyThrobber->SetNumPieces(NumberOfPieces);
+	MyThrobber->SetNumPieces(FMath::Clamp(NumberOfPieces, 1, 25));
 	MyThrobber->SetAnimate(GetAnimation());
 }
 
@@ -62,7 +62,7 @@ void UThrobber::SetNumberOfPieces(int32 InNumberOfPieces)
 	NumberOfPieces = InNumberOfPieces;
 	if (MyThrobber.IsValid())
 	{
-		MyThrobber->SetNumPieces(InNumberOfPieces);
+		MyThrobber->SetNumPieces(FMath::Clamp(NumberOfPieces, 1, 25));
 	}
 }
 
diff --git a/Engine/Source/Runtime/UMG/Public/Components/CircularThrobber.h b/Engine/Source/Runtime/UMG/Public/Components/CircularThrobber.h
--- a/Engine/Source/Runtime/UMG/Public/Components/CircularThrobber.h
+++ b/Engine/Source/Runtime/UMG/Public/Components/CircularThrobber.h
@@ -20,7 +20,7 @@ class UMG_API UCircularThrobber : public UWidget
 public:
 
 	/** How many pieces there are */
-	UPROPERTY(EditDefaultsOnly, BlueprintReadOnly, Category=Appearance, meta=( ClampMin = "1", UIMin = "1", UIMax = "8" ))
+	UPROPERTY(EditDefaultsOnly, BlueprintReadOnly, Category=Appearance, meta=( ClampMin="1", ClampMax="25", UIMin="1", UIMax="25" ))
 	int32 NumberOfPieces;
 
 	/** The amount of time for a full circle (in seconds) */
diff --git a/Engine/Source/Runtime/UMG/Public/Components/Throbber.h b/Engine/Source/Runtime/UMG/Public/Components/Throbber.h
--- a/Engine/Source/Runtime/UMG/Public/Components/Throbber.h
+++ b/Engine/Source/Runtime/UMG/Public/Components/Throbber.h
@@ -18,7 +18,7 @@ class UMG_API UThrobber : public UWidget
 public:
 	
 	/** How many pieces there are */
-	UPROPERTY(EditDefaultsOnly, BlueprintReadOnly, Category=Appearance, meta=( ClampMin="1", UIMin="1", UIMax="8" ))
+	UPROPERTY(EditDefaultsOnly, BlueprintReadOnly, Category=Appearance, meta=( ClampMin="1", ClampMax="25", UIMin="1", UIMax="25" ))
 	int32 NumberOfPieces;
 
 	/** Should the pieces animate horizontally? */
