[INTEGRATE] Change 2422763 by [REDACTED_PERSON][REDACTED_HANDLE] on 2015/01/28 15:58:29

	UMG - Exposing the ability to configure MaxInteractableDistance and DrawSize of the WidgetComponent to Blueprints and C++ at runtime.

diff --git a/Engine/Source/Runtime/UMG/Private/Components/WidgetComponent.cpp b/Engine/Source/Runtime/UMG/Private/Components/WidgetComponent.cpp
--- a/Engine/Source/Runtime/UMG/Private/Components/WidgetComponent.cpp
+++ b/Engine/Source/Runtime/UMG/Private/Components/WidgetComponent.cpp
@@ -819,6 +819,26 @@ TSharedPtr<SWidget> UWidgetComponent::GetSlateWidget() const
 	return SlateWidget;
 }
 
+FVector2D UWidgetComponent::GetDrawSize() const
+{
+	return DrawSize;
+}
+
+void UWidgetComponent::SetDrawSize(FVector2D Size)
+{
+	DrawSize = FIntPoint((int32)Size.X, (int32)Size.Y);
+}
+
+float UWidgetComponent::GetMaxInteractionDistance() const
+{
+	return MaxInteractionDistance;
+}
+
+void UWidgetComponent::SetMaxInteractionDistance(float Distance)
+{
+	MaxInteractionDistance = Distance;
+}
+
 void UWidgetComponent::PostLoad()
 {
 	Super::PostLoad();
diff --git a/Engine/Source/Runtime/UMG/Public/Components/WidgetComponent.h b/Engine/Source/Runtime/UMG/Public/Components/WidgetComponent.h
--- a/Engine/Source/Runtime/UMG/Public/Components/WidgetComponent.h
+++ b/Engine/Source/Runtime/UMG/Public/Components/WidgetComponent.h
@@ -72,7 +72,7 @@ public:
 	TArray<FWidgetAndPointer> GetHitWidgetPath( const FHitResult& HitResult, bool bIgnoreEnabledStatus );
 
 	/** @return The render target to which the user widget is rendered */
-	UTextureRenderTarget2D* GetRenderTarget() const { return RenderTarget; };
+	UTextureRenderTarget2D* GetRenderTarget() const { return RenderTarget; }
 
 	/** @return The dynamic material instance used to render the user widget */
 	UMaterialInstanceDynamic* GetMaterialInstance() const { return MaterialInstance; }
@@ -81,10 +81,20 @@ public:
 	TSharedPtr<SWidget> GetSlateWidget() const;
 
 	/** @return The draw size of the quad in the world */
-	const FIntPoint& GetDrawSize() const { return DrawSize; }
+	UFUNCTION(BlueprintCallable, Category=UI)
+	FVector2D GetDrawSize() const;
+
+	/** Sets the draw size of the quad in the world */
+	UFUNCTION(BlueprintCallable, Category=UI)
+	void SetDrawSize(FVector2D Size);
 
 	/** @return The max distance from which a player can interact with this widget */
-	float GetMaxInteractionDistance() const { return MaxInteractionDistance; }
+	UFUNCTION(BlueprintCallable, Category=UI)
+	float GetMaxInteractionDistance() const;
+
+	/** Sets the max distance from which a player can interact with this widget */
+	UFUNCTION(BlueprintCallable, Category=UI)
+	void SetMaxInteractionDistance(float Distance);
 
 	/** @return True if the component is opaque */
 	bool IsOpaque() const { return bIsOpaque; }
