UMG: Fixed deprecated code

diff --git a/Engine/Source/Runtime/UMG/Public/Blueprint/UserWidget.h b/Engine/Source/Runtime/UMG/Public/Blueprint/UserWidget.h
--- a/Engine/Source/Runtime/UMG/Public/Blueprint/UserWidget.h
+++ b/Engine/Source/Runtime/UMG/Public/Blueprint/UserWidget.h
@@ -741,7 +741,7 @@ T* CreateWidget(UGameInstance* OwningGame, UClass* UserWidgetClass)
 		return nullptr;
 	}
 
-	UUserWidget* NewWidget = ConstructObject<UUserWidget>(UserWidgetClass, OwningGame);
+	UUserWidget* NewWidget = NewObject<UUserWidget>(OwningGame, UserWidgetClass);
 
 	if ( ULocalPlayer* Player = OwningGame->GetFirstGamePlayer() )
 	{
