UMG - When placing items projected to the screen from a widget component, always round to the nearest whole number to prevent rounding error when laying out the item causing odd pixel snapping for child widgets.

diff --git a/Engine/Source/Runtime/UMG/Private/Components/WidgetComponent.cpp b/Engine/Source/Runtime/UMG/Private/Components/WidgetComponent.cpp
--- a/Engine/Source/Runtime/UMG/Private/Components/WidgetComponent.cpp
+++ b/Engine/Source/Runtime/UMG/Private/Components/WidgetComponent.cpp
@@ -542,8 +542,12 @@ void UWidgetComponent::TickComponent(float DeltaTime, enum ELevelTick TickType,
 						Scalability::FQualityLevels ScalabilityQuality = Scalability::GetQualityLevels();
 						float QualityScale = ( ScalabilityQuality.ResolutionQuality / 100.0f );
 
+						FVector2D FinalScreenLocation = ScreenLocation / QualityScale;
+						FinalScreenLocation.X = FMath::RoundToInt(FinalScreenLocation.X);
+						FinalScreenLocation.Y = FMath::RoundToInt(FinalScreenLocation.Y);
+
 						Widget->SetDesiredSizeInViewport(DrawSize);
-						Widget->SetPositionInViewport(ScreenLocation / QualityScale);
+						Widget->SetPositionInViewport(FinalScreenLocation);
 						Widget->SetAlignmentInViewport(Pivot);
 					}
 					else
