Fix for first crash encountered when dropping a button down in the UMG blueprint editor. Regression introuduced in preivous revision.

diff --git a/Engine/Source/Runtime/UMG/Private/Components/PanelWidget.cpp b/Engine/Source/Runtime/UMG/Private/Components/PanelWidget.cpp
--- a/Engine/Source/Runtime/UMG/Private/Components/PanelWidget.cpp
+++ b/Engine/Source/Runtime/UMG/Private/Components/PanelWidget.cpp
@@ -112,7 +112,7 @@ UPanelSlot* UPanelWidget::AddChild(UWidget* Content)
 
 	Content->RemoveFromParent();
 
-	UPanelSlot* Slot = NewObject<UPanelSlot>(this);
+	UPanelSlot* Slot = NewObject<UPanelSlot>(this, GetSlotClass());
 	Slot->SetFlags(RF_Transactional);
 	Slot->Content = Content;
 	Slot->Parent = this;
