[INTEGRATE] Change 2431842 by [REDACTED_PERSON][REDACTED_HANDLE] on 2015/02/04 11:44:03

	UMG - Fixing the scrollview widget to properly use the Bar style.

diff --git a/Engine/Source/Runtime/UMG/Private/Components/ScrollBox.cpp b/Engine/Source/Runtime/UMG/Private/Components/ScrollBox.cpp
--- a/Engine/Source/Runtime/UMG/Private/Components/ScrollBox.cpp
+++ b/Engine/Source/Runtime/UMG/Private/Components/ScrollBox.cpp
@@ -63,6 +63,7 @@ TSharedRef<SWidget> UScrollBox::RebuildWidget()
 {
 	MyScrollBox = SNew(SScrollBox)
 		.Style(&WidgetStyle)
+		.ScrollBarStyle(&WidgetBarStyle)
 		.Orientation(Orientation);
 
 	for ( UPanelSlot* Slot : Slots )
