[INTEGRATE] Change 2433439 by [REDACTED_PERSON][REDACTED_HANDLE] on 2015/02/05 08:31:50

	UMG - Exporting the widget component from the UMG module.

diff --git a/Engine/Source/Runtime/UMG/Public/Components/WidgetComponent.h b/Engine/Source/Runtime/UMG/Public/Components/WidgetComponent.h
--- a/Engine/Source/Runtime/UMG/Public/Components/WidgetComponent.h
+++ b/Engine/Source/Runtime/UMG/Public/Components/WidgetComponent.h
@@ -17,8 +17,8 @@ enum class EWidgetSpace : uint8
 * Beware! This feature is experimental and may be substantially changed or removed in future releases.
 * A 3D instance of a Widget Blueprint that can be interacted with in the world.
 */
-UCLASS(ClassGroup=Experimental, hidecategories=(Object,Activation,"Components|Activation",Sockets,Base,Lighting,LOD,Mesh), editinlinenew, meta=(BlueprintSpawnableComponent,  DevelopmentStatus=Experimental),MinimalAPI)
-class UWidgetComponent : public UPrimitiveComponent
+UCLASS(ClassGroup=Experimental, hidecategories=(Object,Activation,"Components|Activation",Sockets,Base,Lighting,LOD,Mesh), editinlinenew, meta=(BlueprintSpawnableComponent,  DevelopmentStatus=Experimental) )
+class UMG_API UWidgetComponent : public UPrimitiveComponent
 {
 	GENERATED_UCLASS_BODY()
 
