[INTEGRATE] Change 2439846 by [REDACTED_PERSON][REDACTED_HANDLE] on 2015/02/10 11:02:05

	UMG - Adding error messages to be logged to PIE when constructing user widgets fail.

diff --git a/Engine/Source/Runtime/UMG/Public/Blueprint/UserWidget.h b/Engine/Source/Runtime/UMG/Public/Blueprint/UserWidget.h
--- a/Engine/Source/Runtime/UMG/Public/Blueprint/UserWidget.h
+++ b/Engine/Source/Runtime/UMG/Public/Blueprint/UserWidget.h
@@ -683,12 +683,26 @@ private:
 	bool bInitialized;
 };
 
+#define LOCTEXT_NAMESPACE "UMG"
+
 template< class T >
 T* CreateWidget(UWorld* World, UClass* UserWidgetClass)
 {
-	if ( World == nullptr || !UserWidgetClass->IsChildOf(UUserWidget::StaticClass()) || UserWidgetClass->HasAnyClassFlags(CLASS_Abstract) )
+	if ( World == nullptr )
 	{
-		// TODO UMG Error?
+		FMessageLog("PIE").Error(LOCTEXT("WorldNull", "Unable to create a widget outered to a null world."));
+		return nullptr;
+	}
+
+	if ( !UserWidgetClass->IsChildOf(UUserWidget::StaticClass()) )
+	{
+		FMessageLog("PIE").Error(LOCTEXT("NotUserWidget", "CreateWidget can only be used on UUserWidget children."));
+		return nullptr;
+	}
+
+	if ( UserWidgetClass->HasAnyClassFlags(CLASS_Abstract | CLASS_NewerVersionExists | CLASS_Deprecated) )
+	{
+		FMessageLog("PIE").Error(LOCTEXT("NotValidClass", "Abstract, Deprecated or Replaced classes are not allowed to be used to construct a user widget."));
 		return nullptr;
 	}
 
@@ -709,15 +723,27 @@ T* CreateWidget(UWorld* World, UClass* UserWidgetClass)
 template< class T >
 T* CreateWidget(APlayerController* OwningPlayer, UClass* UserWidgetClass)
 {
-	if ( OwningPlayer == nullptr || !UserWidgetClass->IsChildOf(UUserWidget::StaticClass()) || UserWidgetClass->HasAnyClassFlags(CLASS_Abstract) )
+	if ( OwningPlayer == nullptr )
+	{
+		FMessageLog("PIE").Error(LOCTEXT("PlayerNull", "Unable to create a widget outered to a null player controller."));
+		return nullptr;
+	}
+
+	if ( !UserWidgetClass->IsChildOf(UUserWidget::StaticClass()) )
 	{
-		// TODO UMG Error?
+		FMessageLog("PIE").Error(LOCTEXT("NotUserWidget", "CreateWidget can only be used on UUserWidget children."));
+		return nullptr;
+	}
+
+	if ( UserWidgetClass->HasAnyClassFlags(CLASS_Abstract | CLASS_NewerVersionExists | CLASS_Deprecated) )
+	{
+		FMessageLog("PIE").Error(LOCTEXT("NotValidClass", "Abstract, Deprecated or Replaced classes are not allowed to be used to construct a user widget."));
 		return nullptr;
 	}
 
 	if ( !OwningPlayer->IsLocalPlayerController() )
 	{
-		//Don't create widgets for proxies of the Player Controller, only the actual local player.
+		FMessageLog("PIE").Error(LOCTEXT("NotLocalPlayer", "Only Local Player Controllers can be assigned to widgets."));
 		return nullptr;
 	}
 
@@ -735,9 +761,21 @@ T* CreateWidget(APlayerController* OwningPlayer, UClass* UserWidgetClass)
 template< class T >
 T* CreateWidget(UGameInstance* OwningGame, UClass* UserWidgetClass)
 {
-	if ( OwningGame == nullptr || !UserWidgetClass->IsChildOf(UUserWidget::StaticClass()) || UserWidgetClass->HasAnyClassFlags(CLASS_Abstract) )
+	if ( OwningGame == nullptr )
 	{
-		// TODO UMG Error?
+		FMessageLog("PIE").Error(LOCTEXT("GameNull", "Unable to create a widget outered to a null game instance."));
+		return nullptr;
+	}
+
+	if ( !UserWidgetClass->IsChildOf(UUserWidget::StaticClass()) )
+	{
+		FMessageLog("PIE").Error(LOCTEXT("NotUserWidget", "CreateWidget can only be used on UUserWidget children."));
+		return nullptr;
+	}
+
+	if ( UserWidgetClass->HasAnyClassFlags(CLASS_Abstract | CLASS_NewerVersionExists | CLASS_Deprecated) )
+	{
+		FMessageLog("PIE").Error(LOCTEXT("NotValidClass", "Abstract, Deprecated or Replaced classes are not allowed to be used to construct a user widget."));
 		return nullptr;
 	}
 
@@ -752,3 +790,5 @@ T* CreateWidget(UGameInstance* OwningGame, UClass* UserWidgetClass)
 
 	return Cast<T>(NewWidget);
 }
+
+#undef LOCTEXT_NAMESPACE
