[INTEGRATE] Change 2439996 by [REDACTED_PERSON][REDACTED_HANDLE] on 2015/02/10 12:52:50

	Build - Fixing a build issue not including a needed header.

diff --git a/Engine/Source/Runtime/UMG/Public/Blueprint/UserWidget.h b/Engine/Source/Runtime/UMG/Public/Blueprint/UserWidget.h
--- a/Engine/Source/Runtime/UMG/Public/Blueprint/UserWidget.h
+++ b/Engine/Source/Runtime/UMG/Public/Blueprint/UserWidget.h
@@ -9,6 +9,8 @@
 #include "NamedSlotInterface.h"
 #include "Slate/Anchors.h"
 #include "Engine/LocalPlayer.h"
+#include "Logging/MessageLog.h"
+
 #include "UserWidget.generated.h"
 
 static FGeometry NullGeometry;
@@ -652,10 +654,14 @@ public:
 	UPROPERTY()
 	FVector2D DesignTimeSize;
 
-	/** Stores the design time desired size of the user widget */
+	/** Designer flag used to tell when the user wants to provide their own custom size. */
 	UPROPERTY()
 	bool bUseDesignTimeSize;
 
+	/** Designer flag used to tell when to just use the desired size of the widget instead of an explicit size. */
+	UPROPERTY()
+	bool bUseDesiredSizeAtDesignTime;
+
 	/** The category this widget appears in the palette. */
 	UPROPERTY()
 	FText PaletteCategory;
