UMG - Preventing bindings and many events from widgets from being evaluated in the editor at the wrong time.

diff --git a/Engine/Source/Runtime/UMG/Private/Components/Widget.cpp b/Engine/Source/Runtime/UMG/Private/Components/Widget.cpp
--- a/Engine/Source/Runtime/UMG/Private/Components/Widget.cpp
+++ b/Engine/Source/Runtime/UMG/Private/Components/Widget.cpp
@@ -571,12 +571,12 @@ void UWidget::SynchronizeProperties()
 	if ( IsDesignTime() )
 	{
 		SafeWidget->SetEnabled(true);
-		SafeWidget->SetVisibility(TAttribute<EVisibility>::Create(TAttribute<EVisibility>::FGetter::CreateUObject(this, &UWidget::GetVisibilityInDesigner)));
+		SafeWidget->SetVisibility(BIND_UOBJECT_ATTRIBUTE(EVisibility, GetVisibilityInDesigner));
 	}
 	else
 #endif
 	{
-		SafeWidget->SetEnabled(OPTIONAL_BINDING(bool, bIsEnabled));
+		SafeWidget->SetEnabled(GAME_SAFE_OPTIONAL_BINDING( bool, bIsEnabled ));
 		SafeWidget->SetVisibility(OPTIONAL_BINDING_CONVERT(ESlateVisibility, Visibility, EVisibility, ConvertVisibility));
 	}
 
@@ -602,7 +602,7 @@ void UWidget::SynchronizeProperties()
 	}
 	else if ( !ToolTipText.IsEmpty() || ToolTipTextDelegate.IsBound() )
 	{
-		SafeWidget->SetToolTipText(OPTIONAL_BINDING(FText, ToolTipText));
+		SafeWidget->SetToolTipText(GAME_SAFE_OPTIONAL_BINDING(FText, ToolTipText));
 	}
 
 	if (Navigation != nullptr)
diff --git a/Engine/Source/Runtime/UMG/Private/Slate/SObjectWidget.cpp b/Engine/Source/Runtime/UMG/Private/Slate/SObjectWidget.cpp
--- a/Engine/Source/Runtime/UMG/Private/Slate/SObjectWidget.cpp
+++ b/Engine/Source/Runtime/UMG/Private/Slate/SObjectWidget.cpp
@@ -47,41 +47,21 @@ void SObjectWidget::AddReferencedObjects(FReferenceCollector& Collector)
 
 void SObjectWidget::Tick( const FGeometry& AllottedGeometry, const double InCurrentTime, const float InDeltaTime )
 {
-#if WITH_EDITOR
-	if ( GIsRoutingPostLoad )
+	if ( CanRouteEvent() )
 	{
-		// In editor builds streamed in data can cause PostLoad to be called on objects, when this is happening
-		// Slate Tick can and will Occur due to a Slow Task dialog being launched.  In order to prevent UMG ticking
-		// when this is true, we ignore Slate ticks when GIsRoutingPostLoad is true in editor builds.
-		return;
-	}
-#endif
-
-	if ( WidgetObject && !WidgetObject->IsDesignTime() )
-	{
-		return WidgetObject->NativeTick(AllottedGeometry, InDeltaTime);
+		WidgetObject->NativeTick(AllottedGeometry, InDeltaTime);
 	}
 }
 
 int32 SObjectWidget::OnPaint(const FPaintArgs& Args, const FGeometry& AllottedGeometry, const FSlateRect& MyClippingRect, FSlateWindowElementList& OutDrawElements, int32 LayerId, const FWidgetStyle& InWidgetStyle, bool bParentEnabled) const
 {
-#if WITH_EDITOR
-	if ( GIsRoutingPostLoad )
-	{
-		// In editor builds streamed in data can cause PostLoad to be called on objects, when this is happening
-		// Slate painting can and will Occur due to a Slow Task dialog being launched.  In order to prevent UMG painting
-		// when this is true, we ignore Slate painting when GIsRoutingPostLoad is true in editor builds.
-		return LayerId;
-	}
-#endif
-
 	int32 MaxLayer = SCompoundWidget::OnPaint(Args, AllottedGeometry, MyClippingRect, OutDrawElements, LayerId, InWidgetStyle, bParentEnabled);
 
-	if ( WidgetObject && !WidgetObject->IsDesignTime() )
+	if ( CanRouteEvent() )
 	{
 		FPaintContext Context(AllottedGeometry, MyClippingRect, OutDrawElements, LayerId, InWidgetStyle, bParentEnabled);
 		WidgetObject->OnPaint(Context);
-		
+
 		return FMath::Max(MaxLayer, Context.MaxLayer);
 	}
 	
@@ -90,7 +70,7 @@ int32 SObjectWidget::OnPaint(const FPaintArgs& Args, const FGeometry& AllottedGe
 
 bool SObjectWidget::SupportsKeyboardFocus() const
 {
-	if ( WidgetObject && !WidgetObject->IsDesignTime() )
+	if ( CanRouteEvent() )
 	{
 		return WidgetObject->bSupportsKeyboardFocus;
 	}
@@ -100,7 +80,7 @@ bool SObjectWidget::SupportsKeyboardFocus() const
 
 FReply SObjectWidget::OnFocusReceived(const FGeometry& MyGeometry, const FFocusEvent& InFocusEvent)
 {
-	if ( WidgetObject && !WidgetObject->IsDesignTime() )
+	if ( CanRouteEvent() )
 	{
 		return WidgetObject->OnFocusReceived(MyGeometry, InFocusEvent).NativeReply;
 	}
@@ -110,7 +90,7 @@ FReply SObjectWidget::OnFocusReceived(const FGeometry& MyGeometry, const FFocusE
 
 void SObjectWidget::OnFocusLost(const FFocusEvent& InFocusEvent)
 {
-	if ( WidgetObject && !WidgetObject->IsDesignTime() )
+	if ( CanRouteEvent() )
 	{
 		WidgetObject->OnFocusLost(InFocusEvent);
 	}
@@ -123,7 +103,7 @@ void SObjectWidget::OnFocusChanging(const FWeakWidgetPath& PreviousFocusPath, co
 
 FReply SObjectWidget::OnKeyChar(const FGeometry& MyGeometry, const FCharacterEvent& InCharacterEvent)
 {
-	if ( WidgetObject && !WidgetObject->IsDesignTime() )
+	if ( CanRouteEvent() )
 	{
 		return WidgetObject->OnKeyChar(MyGeometry, InCharacterEvent).NativeReply;
 	}
@@ -133,7 +113,7 @@ FReply SObjectWidget::OnKeyChar(const FGeometry& MyGeometry, const FCharacterEve
 
 FReply SObjectWidget::OnPreviewKeyDown(const FGeometry& MyGeometry, const FKeyEvent& InKeyEvent)
 {
-	if ( WidgetObject && !WidgetObject->IsDesignTime() )
+	if ( CanRouteEvent() )
 	{
 		return WidgetObject->OnPreviewKeyDown(MyGeometry, InKeyEvent).NativeReply;
 	}
@@ -143,7 +123,7 @@ FReply SObjectWidget::OnPreviewKeyDown(const FGeometry& MyGeometry, const FKeyEv
 
 FReply SObjectWidget::OnKeyDown(const FGeometry& MyGeometry, const FKeyEvent& InKeyEvent)
 {
-	if ( WidgetObject && !WidgetObject->IsDesignTime() )
+	if ( CanRouteEvent() )
 	{
 		return WidgetObject->OnKeyDown(MyGeometry, InKeyEvent).NativeReply;
 	}
@@ -153,7 +133,7 @@ FReply SObjectWidget::OnKeyDown(const FGeometry& MyGeometry, const FKeyEvent& In
 
 FReply SObjectWidget::OnKeyUp(const FGeometry& MyGeometry, const FKeyEvent& InKeyEvent)
 {
-	if ( WidgetObject && !WidgetObject->IsDesignTime() )
+	if ( CanRouteEvent() )
 	{
 		return WidgetObject->OnKeyUp(MyGeometry, InKeyEvent).NativeReply;
 	}
@@ -163,7 +143,7 @@ FReply SObjectWidget::OnKeyUp(const FGeometry& MyGeometry, const FKeyEvent& InKe
 
 FReply SObjectWidget::OnAnalogValueChanged(const FGeometry& MyGeometry, const FAnalogInputEvent& InAnalogInputEvent)
 {
-	if (WidgetObject && !WidgetObject->IsDesignTime())
+	if ( CanRouteEvent() )
 	{
 		return WidgetObject->OnAnalogValueChanged(MyGeometry, InAnalogInputEvent).NativeReply;
 	}
@@ -175,7 +155,7 @@ FReply SObjectWidget::OnMouseButtonDown(const FGeometry& MyGeometry, const FPoin
 {
 	SCompoundWidget::OnMouseButtonDown(MyGeometry, MouseEvent);
 
-	if ( WidgetObject && !WidgetObject->IsDesignTime() )
+	if ( CanRouteEvent() )
 	{
 		return WidgetObject->OnMouseButtonDown(MyGeometry, MouseEvent).NativeReply;
 	}
@@ -187,7 +167,7 @@ FReply SObjectWidget::OnPreviewMouseButtonDown(const FGeometry& MyGeometry, cons
 {
 	SCompoundWidget::OnPreviewMouseButtonDown(MyGeometry, MouseEvent);
 
-	if ( WidgetObject && !WidgetObject->IsDesignTime() )
+	if ( CanRouteEvent() )
 	{
 		return WidgetObject->OnPreviewMouseButtonDown(MyGeometry, MouseEvent).NativeReply;
 	}
@@ -199,7 +179,7 @@ FReply SObjectWidget::OnMouseButtonUp(const FGeometry& MyGeometry, const FPointe
 {
 	SCompoundWidget::OnMouseButtonUp(MyGeometry, MouseEvent);
 
-	if ( WidgetObject && !WidgetObject->IsDesignTime() )
+	if ( CanRouteEvent() )
 	{
 		return WidgetObject->OnMouseButtonUp(MyGeometry, MouseEvent).NativeReply;
 	}
@@ -209,7 +189,7 @@ FReply SObjectWidget::OnMouseButtonUp(const FGeometry& MyGeometry, const FPointe
 
 FReply SObjectWidget::OnMouseMove(const FGeometry& MyGeometry, const FPointerEvent& MouseEvent)
 {
-	if ( WidgetObject && !WidgetObject->IsDesignTime() )
+	if ( CanRouteEvent() )
 	{
 		return WidgetObject->OnMouseMove(MyGeometry, MouseEvent).NativeReply;
 	}
@@ -219,7 +199,7 @@ FReply SObjectWidget::OnMouseMove(const FGeometry& MyGeometry, const FPointerEve
 
 void SObjectWidget::OnMouseEnter(const FGeometry& MyGeometry, const FPointerEvent& MouseEvent)
 {
-	if ( WidgetObject && !WidgetObject->IsDesignTime() )
+	if ( CanRouteEvent() )
 	{
 		WidgetObject->OnMouseEnter(MyGeometry, MouseEvent);
 	}
@@ -227,7 +207,7 @@ void SObjectWidget::OnMouseEnter(const FGeometry& MyGeometry, const FPointerEven
 
 void SObjectWidget::OnMouseLeave(const FPointerEvent& MouseEvent)
 {
-	if ( WidgetObject && !WidgetObject->IsDesignTime() )
+	if ( CanRouteEvent() )
 	{
 		WidgetObject->OnMouseLeave(MouseEvent);
 	}
@@ -235,7 +215,7 @@ void SObjectWidget::OnMouseLeave(const FPointerEvent& MouseEvent)
 
 FReply SObjectWidget::OnMouseWheel(const FGeometry& MyGeometry, const FPointerEvent& MouseEvent)
 {
-	if ( WidgetObject && !WidgetObject->IsDesignTime() )
+	if ( CanRouteEvent() )
 	{
 		return WidgetObject->OnMouseWheel(MyGeometry, MouseEvent).NativeReply;
 	}
@@ -251,7 +231,7 @@ FCursorReply SObjectWidget::OnCursorQuery(const FGeometry& MyGeometry, const FPo
 
 FReply SObjectWidget::OnMouseButtonDoubleClick(const FGeometry& MyGeometry, const FPointerEvent& MouseEvent)
 {
-	if ( WidgetObject && !WidgetObject->IsDesignTime() )
+	if ( CanRouteEvent() )
 	{
 		return WidgetObject->OnMouseButtonDoubleClick(MyGeometry, MouseEvent).NativeReply;
 	}
@@ -261,15 +241,18 @@ FReply SObjectWidget::OnMouseButtonDoubleClick(const FGeometry& MyGeometry, cons
 
 FReply SObjectWidget::OnDragDetected(const FGeometry& MyGeometry, const FPointerEvent& PointerEvent)
 {
-	UDragDropOperation* Operation = nullptr;
-	WidgetObject->OnDragDetected(MyGeometry, PointerEvent, Operation);
-
-	if ( Operation )
+	if ( CanRouteEvent() )
 	{
-		FVector2D ScreenCursorPos = PointerEvent.GetScreenSpacePosition();
-		FVector2D ScreenDrageePosition = MyGeometry.AbsolutePosition;
+		UDragDropOperation* Operation = nullptr;
+		WidgetObject->OnDragDetected(MyGeometry, PointerEvent, Operation);
 
-		return FReply::Handled().BeginDragDrop(FUMGDragDropOp::New(Operation, ScreenCursorPos, ScreenDrageePosition, SharedThis(this)));
+		if ( Operation )
+		{
+			FVector2D ScreenCursorPos = PointerEvent.GetScreenSpacePosition();
+			FVector2D ScreenDrageePosition = MyGeometry.AbsolutePosition;
+
+			return FReply::Handled().BeginDragDrop(FUMGDragDropOp::New(Operation, ScreenCursorPos, ScreenDrageePosition, SharedThis(this)));
+		}
 	}
 
 	return FReply::Unhandled();
@@ -280,7 +263,10 @@ void SObjectWidget::OnDragEnter(const FGeometry& MyGeometry, const FDragDropEven
 	TSharedPtr<FUMGDragDropOp> NativeOp = DragDropEvent.GetOperationAs<FUMGDragDropOp>();
 	if ( NativeOp.IsValid() )
 	{
-		WidgetObject->OnDragEnter(MyGeometry, DragDropEvent, NativeOp->GetOperation());
+		if ( CanRouteEvent() )
+		{
+			WidgetObject->OnDragEnter(MyGeometry, DragDropEvent, NativeOp->GetOperation());
+		}
 	}
 }
 
@@ -289,7 +275,10 @@ void SObjectWidget::OnDragLeave(const FDragDropEvent& DragDropEvent)
 	TSharedPtr<FUMGDragDropOp> NativeOp = DragDropEvent.GetOperationAs<FUMGDragDropOp>();
 	if ( NativeOp.IsValid() )
 	{
-		WidgetObject->OnDragLeave(DragDropEvent, NativeOp->GetOperation());
+		if ( CanRouteEvent() )
+		{
+			WidgetObject->OnDragLeave(DragDropEvent, NativeOp->GetOperation());
+		}
 	}
 }
 
@@ -298,9 +287,12 @@ FReply SObjectWidget::OnDragOver(const FGeometry& MyGeometry, const FDragDropEve
 	TSharedPtr<FUMGDragDropOp> NativeOp = DragDropEvent.GetOperationAs<FUMGDragDropOp>();
 	if ( NativeOp.IsValid() )
 	{
-		if ( WidgetObject->OnDragOver(MyGeometry, DragDropEvent, NativeOp->GetOperation()) )
+		if ( CanRouteEvent() )
 		{
-			return FReply::Handled();
+			if ( WidgetObject->OnDragOver(MyGeometry, DragDropEvent, NativeOp->GetOperation()) )
+			{
+				return FReply::Handled();
+			}
 		}
 	}
 
@@ -312,9 +304,12 @@ FReply SObjectWidget::OnDrop(const FGeometry& MyGeometry, const FDragDropEvent&
 	TSharedPtr<FUMGDragDropOp> NativeOp = DragDropEvent.GetOperationAs<FUMGDragDropOp>();
 	if ( NativeOp.IsValid() )
 	{
-		if ( WidgetObject->OnDrop(MyGeometry, DragDropEvent, NativeOp->GetOperation()) )
+		if ( CanRouteEvent() )
 		{
-			return FReply::Handled();
+			if ( WidgetObject->OnDrop(MyGeometry, DragDropEvent, NativeOp->GetOperation()) )
+			{
+				return FReply::Handled();
+			}
 		}
 	}
 
@@ -326,13 +321,16 @@ void SObjectWidget::OnDragCancelled(const FDragDropEvent& DragDropEvent, UDragDr
 	TSharedPtr<FUMGDragDropOp> NativeOp = DragDropEvent.GetOperationAs<FUMGDragDropOp>();
 	if ( NativeOp.IsValid() )
 	{
-		WidgetObject->OnDragCancelled(DragDropEvent, NativeOp->GetOperation());
+		if ( CanRouteEvent() )
+		{
+			WidgetObject->OnDragCancelled(DragDropEvent, NativeOp->GetOperation());
+		}
 	}
 }
 
 FReply SObjectWidget::OnTouchGesture(const FGeometry& MyGeometry, const FPointerEvent& GestureEvent)
 {
-	if ( WidgetObject && !WidgetObject->IsDesignTime() )
+	if ( CanRouteEvent() )
 	{
 		return WidgetObject->OnTouchGesture(MyGeometry, GestureEvent).NativeReply;
 	}
@@ -342,7 +340,7 @@ FReply SObjectWidget::OnTouchGesture(const FGeometry& MyGeometry, const FPointer
 
 FReply SObjectWidget::OnTouchStarted(const FGeometry& MyGeometry, const FPointerEvent& InTouchEvent)
 {
-	if ( WidgetObject && !WidgetObject->IsDesignTime() )
+	if ( CanRouteEvent() )
 	{
 		return WidgetObject->OnTouchStarted(MyGeometry, InTouchEvent).NativeReply;
 	}
@@ -352,7 +350,7 @@ FReply SObjectWidget::OnTouchStarted(const FGeometry& MyGeometry, const FPointer
 
 FReply SObjectWidget::OnTouchMoved(const FGeometry& MyGeometry, const FPointerEvent& InTouchEvent)
 {
-	if ( WidgetObject && !WidgetObject->IsDesignTime() )
+	if ( CanRouteEvent() )
 	{
 		return WidgetObject->OnTouchMoved(MyGeometry, InTouchEvent).NativeReply;
 	}
@@ -362,7 +360,7 @@ FReply SObjectWidget::OnTouchMoved(const FGeometry& MyGeometry, const FPointerEv
 
 FReply SObjectWidget::OnTouchEnded(const FGeometry& MyGeometry, const FPointerEvent& InTouchEvent)
 {
-	if ( WidgetObject && !WidgetObject->IsDesignTime() )
+	if ( CanRouteEvent() )
 	{
 		return WidgetObject->OnTouchEnded(MyGeometry, InTouchEvent).NativeReply;
 	}
@@ -372,7 +370,7 @@ FReply SObjectWidget::OnTouchEnded(const FGeometry& MyGeometry, const FPointerEv
 
 FReply SObjectWidget::OnMotionDetected(const FGeometry& MyGeometry, const FMotionEvent& InMotionEvent)
 {
-	if ( WidgetObject && !WidgetObject->IsDesignTime() )
+	if ( CanRouteEvent() )
 	{
 		return WidgetObject->OnMotionDetected(MyGeometry, InMotionEvent).NativeReply;
 	}
diff --git a/Engine/Source/Runtime/UMG/Public/Components/Widget.h b/Engine/Source/Runtime/UMG/Public/Components/Widget.h
--- a/Engine/Source/Runtime/UMG/Public/Components/Widget.h
+++ b/Engine/Source/Runtime/UMG/Public/Components/Widget.h
@@ -20,6 +20,40 @@ class UPanelSlot;
 	:															\
 		TAttribute< ReturnType >(MemberName)
 
+#if WITH_EDITOR
+
+#define GAME_SAFE_OPTIONAL_BINDING(ReturnType, MemberName)			\
+	( MemberName ## Delegate.IsBound() && !IsDesignTime() )			\
+	?																\
+		BIND_UOBJECT_ATTRIBUTE(ReturnType, K2_Gate_ ## MemberName)	\
+	:																\
+		TAttribute< ReturnType >(MemberName)
+
+#define GAME_SAFE_BINDING_IMPLEMENTATION(ReturnType, MemberName)		\
+	ReturnType K2_Cache_ ## MemberName;									\
+	ReturnType K2_Gate_ ## MemberName()									\
+	{																	\
+		if (CanSafelyRouteEvent())										\
+		{																\
+			K2_Cache_ ## MemberName = TAttribute< ReturnType >::Create(MemberName ## Delegate.GetUObject(), MemberName ## Delegate.GetFunctionName()).Get(); \
+		}																\
+																		\
+		return K2_Cache_ ## MemberName;									\
+	}
+
+#else
+
+#define GAME_SAFE_OPTIONAL_BINDING(ReturnType, MemberName)		\
+	( MemberName ## Delegate.IsBound() && !IsDesignTime() )		\
+	?															\
+		TAttribute< ReturnType >::Create(MemberName ## Delegate.GetUObject(), MemberName ## Delegate.GetFunctionName()) \
+	:															\
+		TAttribute< ReturnType >(MemberName)
+
+#define GAME_SAFE_BINDING_IMPLEMENTATION(Type, MemberName)
+
+#endif
+
 /**
  * Helper macro for binding to a delegate or using the constant value when constructing the underlying SWidget,
  * also allows a conversion function to be provided to convert between the SWidget value and the value exposed to UMG.
@@ -284,8 +318,12 @@ public:
 	 */
 	virtual void SynchronizeProperties();
 
+#if WITH_EDITOR
 	/** Returns if the widget is currently being displayed in the designer, it may want to display different data. */
 	bool IsDesignTime() const;
+#else
+	FORCEINLINE bool IsDesignTime() const { return false; }
+#endif
 	
 	/** Sets that this widget is being designed */
 	virtual void SetIsDesignTime(bool bInDesignTime);
@@ -308,8 +346,18 @@ public:
 	virtual UWorld* GetWorld() const override;
 	virtual void PostLoad() override;
 	// End UObject
-	
+
 #if WITH_EDITOR
+	FORCEINLINE bool CanSafelyRouteEvent()
+	{
+		return !( IsDesignTime() || GIsRoutingPostLoad || GIntraFrameDebuggingGameThread );
+	}
+#else
+	FORCEINLINE bool CanSafelyRouteEvent() { return true; }
+#endif
+
+#if WITH_EDITOR
+
 	/** Is the label generated or provided by the user? */
 	bool IsGeneratedName() const;
 
@@ -416,4 +464,8 @@ protected:
 	TArray<class UPropertyBinding*> NativeBindings;
 
 	static TArray<TSubclassOf<UPropertyBinding>> BinderClasses;
+
+private:
+	GAME_SAFE_BINDING_IMPLEMENTATION(FText, ToolTipText)
+	GAME_SAFE_BINDING_IMPLEMENTATION(bool, bIsEnabled)
 };
diff --git a/Engine/Source/Runtime/UMG/Public/Slate/SObjectWidget.h b/Engine/Source/Runtime/UMG/Public/Slate/SObjectWidget.h
--- a/Engine/Source/Runtime/UMG/Public/Slate/SObjectWidget.h
+++ b/Engine/Source/Runtime/UMG/Public/Slate/SObjectWidget.h
@@ -2,7 +2,7 @@
 
 #pragma once
 
-class UUserWidget;
+#include "UserWidget.h"
 
 /**
  * The SObjectWidet allows UMG to insert an SWidget into the hierarchy that manages the lifetime of the
@@ -70,6 +70,14 @@ class SObjectWidget : public SCompoundWidget, public FGCObject
 	virtual FReply OnMotionDetected(const FGeometry& MyGeometry, const FMotionEvent& InMotionEvent) override;
 
 private:
+
 	/** The UWidget that created this SObjectWidget who needs to be kept alive. */
 	UUserWidget* WidgetObject;
+
+private:
+
+	FORCEINLINE bool CanRouteEvent() const
+	{
+		return WidgetObject && WidgetObject->CanSafelyRouteEvent();
+	}
 };
\ No newline at end of file
