Fixed a compilation error due to the UWidget::IsDesignTime method being defined multiple times in non-Editor builds. Was blocking my local build.

diff --git a/Engine/Source/Runtime/UMG/Private/Components/Widget.cpp b/Engine/Source/Runtime/UMG/Private/Components/Widget.cpp
--- a/Engine/Source/Runtime/UMG/Private/Components/Widget.cpp
+++ b/Engine/Source/Runtime/UMG/Private/Components/Widget.cpp
@@ -624,10 +624,12 @@ void UWidget::SynchronizeProperties()
 #endif
 }
 
+#if WITH_EDITOR
 bool UWidget::IsDesignTime() const
 {
 	return bDesignTime;
 }
+#endif
 
 void UWidget::SetIsDesignTime(bool bInDesignTime)
 {
