UMG - UE-9880 - The widget component now removes the widget from the parent when it's hidden.  Fixing the pivot point not allowing the widget to be hit tested correctly on the component.

diff --git a/Engine/Source/Runtime/UMG/Private/Components/WidgetComponent.cpp b/Engine/Source/Runtime/UMG/Private/Components/WidgetComponent.cpp
--- a/Engine/Source/Runtime/UMG/Private/Components/WidgetComponent.cpp
+++ b/Engine/Source/Runtime/UMG/Private/Components/WidgetComponent.cpp
@@ -54,6 +54,29 @@ public:
 	// FPrimitiveSceneProxy interface.
 	virtual void GetDynamicMeshElements(const TArray<const FSceneView*>& Views, const FSceneViewFamily& ViewFamily, uint32 VisibilityMap, FMeshElementCollector& Collector) const override
 	{
+#if WITH_EDITOR
+		const bool bWireframe = AllowDebugViewmodes() && ViewFamily.EngineShowFlags.Wireframe;
+
+		auto WireframeMaterialInstance = new FColoredMaterialRenderProxy(
+			GEngine->WireframeMaterial ? GEngine->WireframeMaterial->GetRenderProxy(IsSelected()) : nullptr,
+			FLinearColor(0, 0.5f, 1.f)
+			);
+
+		Collector.RegisterOneFrameMaterialProxy(WireframeMaterialInstance);
+
+		FMaterialRenderProxy* MaterialProxy = nullptr;
+		if ( bWireframe )
+		{
+			MaterialProxy = WireframeMaterialInstance;
+		}
+		else
+		{
+			MaterialProxy = MaterialInstance->GetRenderProxy(IsSelected());
+		}
+#else
+		FMaterialRenderProxy* MaterialProxy = MaterialInstance->GetRenderProxy(IsSelected());
+#endif
+
 		const FMatrix& LocalToWorld = GetLocalToWorld();
 
 		if( RenderTarget )
@@ -82,7 +105,7 @@ public:
 						MeshBuilder.AddTriangle(VertexIndices[0], VertexIndices[1], VertexIndices[2]);
 						MeshBuilder.AddTriangle(VertexIndices[0], VertexIndices[2], VertexIndices[3]);
 
-						MeshBuilder.GetMesh(LocalToWorld, MaterialInstance->GetRenderProxy(IsSelected()), SDPG_World, false, true, ViewIndex, Collector);
+						MeshBuilder.GetMesh(LocalToWorld, MaterialProxy, SDPG_World, false, true, ViewIndex, Collector);
 					}
 				}
 			}
@@ -103,7 +126,7 @@ public:
 		Result.bNormalTranslucencyRelevance = !bIsOpaque;
 		Result.bSeparateTranslucencyRelevance = !bIsOpaque;
 		Result.bDynamicRelevance = true;
-		Result.bShadowRelevance = IsShadowCast(View);;
+		Result.bShadowRelevance = IsShadowCast(View);
 		Result.bEditorPrimitiveRelevance = false;
 		return Result;
 	}
@@ -308,7 +331,9 @@ FBoxSphereBounds UWidgetComponent::CalcBounds(const FTransform & LocalToWorld) c
 {
 	if ( Space != EWidgetSpace::Screen )
 	{
-		const FVector Origin = FVector(DrawSize.X / 2.0f - ( DrawSize.X * Pivot.X ), DrawSize.Y / 2.0f - ( DrawSize.Y * Pivot.Y ), .5f);
+		const FVector Origin = FVector(
+			( DrawSize.X * 0.5f ) - ( DrawSize.X * Pivot.X ),
+			( DrawSize.Y * 0.5f ) - ( DrawSize.Y * Pivot.Y ), .5f);
 		const FVector BoxExtent = FVector(DrawSize.X / 2.0f, DrawSize.Y / 2.0f, 1.0f);
 
 		return FBoxSphereBounds(Origin, BoxExtent, DrawSize.Size() / 2.0f).TransformBy(LocalToWorld);
@@ -522,43 +547,40 @@ void UWidgetComponent::TickComponent(float DeltaTime, enum ELevelTick TickType,
 			ULocalPlayer* const TargetPlayer = GEngine->GetLocalPlayerFromControllerId(GetWorld(), 0);
 			APlayerController* PlayerController = TargetPlayer->PlayerController;
 
-			if ( TargetPlayer && PlayerController )
+			if ( TargetPlayer && PlayerController && IsVisible() )
 			{
-				if ( IsVisible() )
-				{
-					FVector WorldLocation = GetComponentLocation();
+				FVector WorldLocation = GetComponentLocation();
 
-					FVector2D ScreenLocation;
-					bool bProjected = PlayerController->ProjectWorldLocationToScreen(WorldLocation, ScreenLocation);
+				FVector2D ScreenLocation;
+				bool bProjected = PlayerController->ProjectWorldLocationToScreen(WorldLocation, ScreenLocation);
 
-					if ( bProjected )
-					{
-						Widget->SetVisibility(ESlateVisibility::SelfHitTestInvisible);
+				if ( bProjected )
+				{
+					Widget->SetVisibility(ESlateVisibility::SelfHitTestInvisible);
 						
-						// If the user has configured a resolution quality we need to multiply
-						// the pixels by the resolution quality to arrive at the true position in
-						// the viewport, as the rendered image will be stretched to fill whatever
-						// size the viewport is at.
-						Scalability::FQualityLevels ScalabilityQuality = Scalability::GetQualityLevels();
-						float QualityScale = ( ScalabilityQuality.ResolutionQuality / 100.0f );
-
-						FVector2D FinalScreenLocation = ScreenLocation / QualityScale;
-						FinalScreenLocation.X = FMath::RoundToInt(FinalScreenLocation.X);
-						FinalScreenLocation.Y = FMath::RoundToInt(FinalScreenLocation.Y);
-
-						Widget->SetDesiredSizeInViewport(DrawSize);
-						Widget->SetPositionInViewport(FinalScreenLocation);
-						Widget->SetAlignmentInViewport(Pivot);
-					}
-					else
-					{
-						Widget->SetVisibility(ESlateVisibility::Hidden);
-					}
+					// If the user has configured a resolution quality we need to multiply
+					// the pixels by the resolution quality to arrive at the true position in
+					// the viewport, as the rendered image will be stretched to fill whatever
+					// size the viewport is at.
+					Scalability::FQualityLevels ScalabilityQuality = Scalability::GetQualityLevels();
+					float QualityScale = ( ScalabilityQuality.ResolutionQuality / 100.0f );
+
+					FVector2D FinalScreenLocation = ScreenLocation / QualityScale;
+					FinalScreenLocation.X = FMath::RoundToInt(FinalScreenLocation.X);
+					FinalScreenLocation.Y = FMath::RoundToInt(FinalScreenLocation.Y);
+
+					Widget->SetDesiredSizeInViewport(DrawSize);
+					Widget->SetPositionInViewport(FinalScreenLocation);
+					Widget->SetAlignmentInViewport(Pivot);
+				}
+				else
+				{
+					Widget->SetVisibility(ESlateVisibility::Hidden);
+				}
 
-					if ( !Widget->IsInViewport() )
-					{
-						Widget->AddToViewport(ZOrder);
-					}
+				if ( !Widget->IsInViewport() )
+				{
+					Widget->AddToViewport(ZOrder);
 				}
 			}
 			else if ( Widget->IsInViewport() )
@@ -781,7 +803,9 @@ void UWidgetComponent::UpdateBodySetup( bool bDrawSizeChanged )
 
 		FKBoxElem* BoxElem = BodySetup->AggGeom.BoxElems.GetData();
 
-		const FVector Origin = FVector(DrawSize.X / 2.0f - ( DrawSize.X * Pivot.X ), DrawSize.Y / 2.0f - ( DrawSize.Y * Pivot.Y ), .5f);
+		const FVector Origin = FVector(
+			(DrawSize.X * 0.5f) - ( DrawSize.X * Pivot.X ),
+			(DrawSize.Y * 0.5f) - ( DrawSize.Y * Pivot.Y ), .5f);
 
 		BoxElem->SetTransform(FTransform::Identity);
 		BoxElem->Center = Origin;
@@ -798,7 +822,11 @@ UUserWidget* UWidgetComponent::GetUserWidgetObject() const
 
 TArray<FWidgetAndPointer> UWidgetComponent::GetHitWidgetPath( const FHitResult& HitResult, bool bIgnoreEnabledStatus  )
 {
+	// Find the hit location on the component
 	FVector2D LocalHitLocation = FVector2D( ComponentToWorld.InverseTransformPosition( HitResult.Location ) );
+	// Offset the position by the pivot to get the position in widget space.
+	LocalHitLocation.X += DrawSize.X * Pivot.X;
+	LocalHitLocation.Y += DrawSize.Y * Pivot.Y;
 
 	TSharedRef<FVirtualPointerPosition> VirtualMouseCoordinate = MakeShareable( new FVirtualPointerPosition );
 
diff --git a/Engine/Source/Runtime/UMG/UMG.Build.cs b/Engine/Source/Runtime/UMG/UMG.Build.cs
--- a/Engine/Source/Runtime/UMG/UMG.Build.cs
+++ b/Engine/Source/Runtime/UMG/UMG.Build.cs
@@ -20,6 +20,7 @@ public class UMG : ModuleRules
                 "InputCore",
 				"Slate",
 				"SlateCore",
+                "ShaderCore",
 				"RenderCore",
 				"MovieSceneCore",
 				"RHI",
