[INTEGRATE] Change 2444765 by [REDACTED_PERSON][REDACTED_HANDLE] on 2015/02/13 10:41:10

	UE-9755 - UMG - Fixing a regression in custom umg painting, the widgets max layer id was not passed into the custom painting code.  Could not reproduce the corruption reported though, that was likely fixed by the same change that caused the paint layer regression.

diff --git a/Engine/Source/Runtime/UMG/Private/Slate/SObjectWidget.cpp b/Engine/Source/Runtime/UMG/Private/Slate/SObjectWidget.cpp
--- a/Engine/Source/Runtime/UMG/Private/Slate/SObjectWidget.cpp
+++ b/Engine/Source/Runtime/UMG/Private/Slate/SObjectWidget.cpp
@@ -59,7 +59,7 @@ int32 SObjectWidget::OnPaint(const FPaintArgs& Args, const FGeometry& AllottedGe
 
 	if ( CanRouteEvent() )
 	{
-		FPaintContext Context(AllottedGeometry, MyClippingRect, OutDrawElements, LayerId, InWidgetStyle, bParentEnabled);
+		FPaintContext Context(AllottedGeometry, MyClippingRect, OutDrawElements, MaxLayer, InWidgetStyle, bParentEnabled);
 		WidgetObject->OnPaint(Context);
 
 		return FMath::Max(MaxLayer, Context.MaxLayer);
