Adding a DPI Scaler to all UMG drag drop operation widgets so that all drag visuals are scaled the same as any other widget would be inside the viewport.

diff --git a/Engine/Source/Runtime/UMG/Private/Slate/SObjectWidget.cpp b/Engine/Source/Runtime/UMG/Private/Slate/SObjectWidget.cpp
--- a/Engine/Source/Runtime/UMG/Private/Slate/SObjectWidget.cpp
+++ b/Engine/Source/Runtime/UMG/Private/Slate/SObjectWidget.cpp
@@ -251,7 +251,11 @@ FReply SObjectWidget::OnDragDetected(const FGeometry& MyGeometry, const FPointer
 			FVector2D ScreenCursorPos = PointerEvent.GetScreenSpacePosition();
 			FVector2D ScreenDrageePosition = MyGeometry.AbsolutePosition;
 
-			return FReply::Handled().BeginDragDrop(FUMGDragDropOp::New(Operation, ScreenCursorPos, ScreenDrageePosition, SharedThis(this)));
+			float DPIScale = UWidgetLayoutLibrary::GetViewportScale(WidgetObject);
+
+			TSharedRef<FUMGDragDropOp> DragDropOp = FUMGDragDropOp::New(Operation, ScreenCursorPos, ScreenDrageePosition, DPIScale, SharedThis(this));
+
+			return FReply::Handled().BeginDragDrop(DragDropOp);
 		}
 	}
 
diff --git a/Engine/Source/Runtime/UMG/Private/Slate/UMGDragDropOp.cpp b/Engine/Source/Runtime/UMG/Private/Slate/UMGDragDropOp.cpp
--- a/Engine/Source/Runtime/UMG/Private/Slate/UMGDragDropOp.cpp
+++ b/Engine/Source/Runtime/UMG/Private/Slate/UMGDragDropOp.cpp
@@ -2,6 +2,8 @@
 
 #include "UMGPrivatePCH.h"
 
+#include "Widgets/Layout/SDPIScaler.h"
+
 #include "UMGDragDropOp.h"
 
 //////////////////////////////////////////////////////////////////////////
@@ -115,7 +117,7 @@ TSharedPtr<SWidget> FUMGDragDropOp::GetDefaultDecorator() const
 	return DecoratorWidget;
 }
 
-TSharedRef<FUMGDragDropOp> FUMGDragDropOp::New(UDragDropOperation* InOperation, const FVector2D &CursorPosition, const FVector2D &ScreenPositionOfDragee, TSharedPtr<SObjectWidget> SourceUserWidget)
+TSharedRef<FUMGDragDropOp> FUMGDragDropOp::New(UDragDropOperation* InOperation, const FVector2D &CursorPosition, const FVector2D &ScreenPositionOfDragee, float DPIScale, TSharedPtr<SObjectWidget> SourceUserWidget)
 {
 	check(InOperation);
 
@@ -123,19 +125,26 @@ TSharedRef<FUMGDragDropOp> FUMGDragDropOp::New(UDragDropOperation* InOperation,
 	Operation->MouseDownOffset = ScreenPositionOfDragee - CursorPosition;
 	Operation->StartingScreenPos = ScreenPositionOfDragee;
 	Operation->SourceUserWidget = SourceUserWidget;
-
 	Operation->DragOperation = InOperation;
 
+	TSharedPtr<SWidget> DragVisual;
 	if ( InOperation->DefaultDragVisual == nullptr )
 	{
-		Operation->DecoratorWidget = SNew(STextBlock)
+		DragVisual = SNew(STextBlock)
 			.Text(FText::FromString(InOperation->Tag));
 	}
 	else
 	{
-		Operation->DecoratorWidget = InOperation->DefaultDragVisual->TakeWidget();
+		DragVisual = InOperation->DefaultDragVisual->TakeWidget();
 	}
 
+	Operation->DecoratorWidget =
+		SNew(SDPIScaler)
+		.DPIScale(DPIScale)
+		[
+			DragVisual.ToSharedRef()
+		];
+
 	Operation->DecoratorWidget->SlatePrepass();
 
 	Operation->Construct();
diff --git a/Engine/Source/Runtime/UMG/Public/Slate/UMGDragDropOp.h b/Engine/Source/Runtime/UMG/Public/Slate/UMGDragDropOp.h
--- a/Engine/Source/Runtime/UMG/Public/Slate/UMGDragDropOp.h
+++ b/Engine/Source/Runtime/UMG/Public/Slate/UMGDragDropOp.h
@@ -13,7 +13,7 @@ class FUMGDragDropOp : public FGameDragDropOperation, public FGCObject
 public:
 	DRAG_DROP_OPERATOR_TYPE(FUMGDragDropOp, FGameDragDropOperation)
 	
-	static TSharedRef<FUMGDragDropOp> New(UDragDropOperation* Operation, const FVector2D &CursorPosition, const FVector2D &ScreenPositionOfNode, TSharedPtr<SObjectWidget> SourceUserWidget);
+	static TSharedRef<FUMGDragDropOp> New(UDragDropOperation* Operation, const FVector2D &CursorPosition, const FVector2D &ScreenPositionOfNode, float DPIScale, TSharedPtr<SObjectWidget> SourceUserWidget);
 
 	FUMGDragDropOp();
 
