[Github PR 798] Allar - Adding a check to see if a widget is unreachable before trying to forward a slate event to it.

diff --git a/Engine/Source/Runtime/UMG/Public/Components/Widget.h b/Engine/Source/Runtime/UMG/Public/Components/Widget.h
--- a/Engine/Source/Runtime/UMG/Public/Components/Widget.h
+++ b/Engine/Source/Runtime/UMG/Public/Components/Widget.h
@@ -350,7 +350,7 @@ public:
 #if WITH_EDITOR
 	FORCEINLINE bool CanSafelyRouteEvent()
 	{
-		return !( IsDesignTime() || GIsRoutingPostLoad || GIntraFrameDebuggingGameThread );
+		return !( IsDesignTime() || GIsRoutingPostLoad || GIntraFrameDebuggingGameThread || HasAnyFlags(RF_Unreachable) );
 	}
 #else
 	FORCEINLINE bool CanSafelyRouteEvent() { return true; }
