[INTEGRATE] Change 2449054 by [REDACTED_PERSON][REDACTED_HANDLE] on 2015/02/17 15:01:24

	UE-9975 - UMG - Fixing non-native bindings.  The native binding code now makes sure that the number of parameters matches the native binders, which only accept a single parameter (the return), now the code properly fails over and performs a direct ufunction binding for any delegate it can't bind using the native binding.

diff --git a/Engine/Source/Runtime/UMG/Private/Components/Widget.cpp b/Engine/Source/Runtime/UMG/Private/Components/Widget.cpp
--- a/Engine/Source/Runtime/UMG/Private/Components/Widget.cpp
+++ b/Engine/Source/Runtime/UMG/Private/Components/Widget.cpp
@@ -815,17 +815,22 @@ static UPropertyBinding* GenerateBinder(UDelegateProperty* DelegateProperty, UOb
 	FScriptDelegate* ScriptDelegate = DelegateProperty->GetPropertyValuePtr_InContainer(Container);
 	if ( ScriptDelegate )
 	{
-		if ( UProperty* ReturnProperty = DelegateProperty->SignatureFunction->GetReturnProperty() )
+		// Only delegates that take no parameters have native binders.
+		UFunction* SignatureFunction = DelegateProperty->SignatureFunction;
+		if ( SignatureFunction->NumParms == 1 )
 		{
-			TSubclassOf<UPropertyBinding> BinderClass = UWidget::FindBinderClassForDestination(ReturnProperty);
-			if ( BinderClass != nullptr )
+			if ( UProperty* ReturnProperty = SignatureFunction->GetReturnProperty() )
 			{
-				UPropertyBinding* Binder = NewObject<UPropertyBinding>(Container, BinderClass);
-				Binder->SourceObject = SourceObject;
-				Binder->SourcePath = BindingPath;
-				Binder->Bind(ReturnProperty, ScriptDelegate);
-
-				return Binder;
+				TSubclassOf<UPropertyBinding> BinderClass = UWidget::FindBinderClassForDestination(ReturnProperty);
+				if ( BinderClass != nullptr )
+				{
+					UPropertyBinding* Binder = NewObject<UPropertyBinding>(Container, BinderClass);
+					Binder->SourceObject = SourceObject;
+					Binder->SourcePath = BindingPath;
+					Binder->Bind(ReturnProperty, ScriptDelegate);
+
+					return Binder;
+				}
 			}
 		}
 	}
diff --git a/Engine/Source/Runtime/UMG/Private/WidgetBlueprintGeneratedClass.cpp b/Engine/Source/Runtime/UMG/Private/WidgetBlueprintGeneratedClass.cpp
--- a/Engine/Source/Runtime/UMG/Private/WidgetBlueprintGeneratedClass.cpp
+++ b/Engine/Source/Runtime/UMG/Private/WidgetBlueprintGeneratedClass.cpp
@@ -92,6 +92,9 @@ void UWidgetBlueprintGeneratedClass::InitializeWidget(UUserWidget* UserWidget) c
 							bSourcePathBound = Widget->AddBinding(DelegateProperty, UserWidget, Binding.SourcePath);
 						}
 
+						// If no native binder is found then the only possibility is that the binding is for
+						// a delegate that doesn't match the known native binders available and so we
+						// fallback to just attempting to bind to the function directly.
 						if ( bSourcePathBound == false )
 						{
 							FScriptDelegate* ScriptDelegate = DelegateProperty->GetPropertyValuePtr_InContainer(Widget);
