Fix compile error

diff --git a/Engine/Source/Runtime/UMG/Public/Binding/DynamicPropertyPath.h b/Engine/Source/Runtime/UMG/Public/Binding/DynamicPropertyPath.h
--- a/Engine/Source/Runtime/UMG/Public/Binding/DynamicPropertyPath.h
+++ b/Engine/Source/Runtime/UMG/Public/Binding/DynamicPropertyPath.h
@@ -242,7 +242,7 @@ private:
 							if ( ArrayHelper.IsValidIndex(Segment.ArrayIndex) )
 							{
 								// Recursively call back into this function with the array element and container value
-								return GetValueRecursive(ArryayOfStructsProp->Struct, ArrayHelper.GetRawPtr(Segment.ArrayIndex), ArrayIndex, SegmentIndex + 1, OutValue, OutProperty);
+								return GetValueRecursive(ArrayOfStructsProp->Struct, ArrayHelper.GetRawPtr(Segment.ArrayIndex), ArrayIndex, SegmentIndex + 1, OutValue, OutProperty);
 							}
 						}
 						// if it's not an array of structs, maybe it's an array of classes
