Build Fix.

diff --git a/Engine/Source/Runtime/UMG/Private/Slate/SObjectWidget.cpp b/Engine/Source/Runtime/UMG/Private/Slate/SObjectWidget.cpp
--- a/Engine/Source/Runtime/UMG/Private/Slate/SObjectWidget.cpp
+++ b/Engine/Source/Runtime/UMG/Private/Slate/SObjectWidget.cpp
@@ -2,6 +2,7 @@
 
 #include "UMGPrivatePCH.h"
 
+#include "WidgetLayoutLibrary.h"
 #include "UMGDragDropOp.h"
 
 void SObjectWidget::Construct(const FArguments& InArgs, UUserWidget* InWidgetObject)
