UMG - 3D Widget Components now default to ZOrder -100 when added to the viewport.

diff --git a/Engine/Source/Runtime/UMG/Private/Components/WidgetComponent.cpp b/Engine/Source/Runtime/UMG/Private/Components/WidgetComponent.cpp
--- a/Engine/Source/Runtime/UMG/Private/Components/WidgetComponent.cpp
+++ b/Engine/Source/Runtime/UMG/Private/Components/WidgetComponent.cpp
@@ -311,6 +311,7 @@ UWidgetComponent::UWidgetComponent( const FObjectInitializer& PCIP )
 
 	Space = EWidgetSpace::World;
 	Pivot = FVector2D(0.5, 0.5);
+	ZOrder = -100;
 
 	PrimaryComponentTick.TickGroup = TG_PostUpdateWork;
 }
