UMG - ToolTipText on Widget is now marked as Multiline for the property grid.

diff --git a/Engine/Source/Runtime/UMG/Public/Components/Widget.h b/Engine/Source/Runtime/UMG/Public/Components/Widget.h
--- a/Engine/Source/Runtime/UMG/Public/Components/Widget.h
+++ b/Engine/Source/Runtime/UMG/Public/Components/Widget.h
@@ -126,7 +126,7 @@ public:
 	FGetBool bIsEnabledDelegate;
 
 	/** Tooltip text to show when the user hovers over the widget with the mouse */
-	UPROPERTY(EditDefaultsOnly, BlueprintReadOnly, Category="Behavior")
+	UPROPERTY(EditDefaultsOnly, BlueprintReadOnly, Category="Behavior", meta=(MultiLine=true))
 	FText ToolTipText;
 
 	/** A bindable delegate for ToolTipText */
