UMG - Allow user focus (by user index)  to be set on a widget.

diff --git a/Engine/Source/Runtime/UMG/Private/Components/Widget.cpp b/Engine/Source/Runtime/UMG/Private/Components/Widget.cpp
--- a/Engine/Source/Runtime/UMG/Private/Components/Widget.cpp
+++ b/Engine/Source/Runtime/UMG/Private/Components/Widget.cpp
@@ -305,6 +305,39 @@ void UWidget::SetKeyboardFocus()
 	}
 }
 
+bool UWidget::HasUserFocus(int32 UserIndex) const
+{
+	TSharedPtr<SWidget> SafeWidget = GetCachedWidget();
+	if ( SafeWidget.IsValid() )
+	{
+		TOptional<EFocusCause> FocusCause = SafeWidget->HasUserFocus(UserIndex);
+		return FocusCause.IsSet();
+	}
+
+	return false;
+}
+
+bool UWidget::HasAnyUserFocus() const
+{
+	TSharedPtr<SWidget> SafeWidget = GetCachedWidget();
+	if ( SafeWidget.IsValid() )
+	{
+		TOptional<EFocusCause> FocusCause = SafeWidget->HasAnyUserFocus();
+		return FocusCause.IsSet();
+	}
+
+	return false;
+}
+
+void UWidget::SetUserFocus(int32 UserIndex)
+{
+	TSharedPtr<SWidget> SafeWidget = GetCachedWidget();
+	if ( SafeWidget.IsValid() )
+	{
+		FSlateApplication::Get().SetUserFocus(UserIndex, SafeWidget);
+	}
+}
+
 void UWidget::ForceLayoutPrepass()
 {
 	TSharedPtr<SWidget> SafeWidget = GetCachedWidget();
diff --git a/Engine/Source/Runtime/UMG/Public/Components/Widget.h b/Engine/Source/Runtime/UMG/Public/Components/Widget.h
--- a/Engine/Source/Runtime/UMG/Public/Components/Widget.h
+++ b/Engine/Source/Runtime/UMG/Public/Components/Widget.h
@@ -279,6 +279,18 @@ public:
 	UFUNCTION(BlueprintCallable, Category="Widget")
 	void SetKeyboardFocus();
 
+	/** Gets the focus to this widget. */
+	UFUNCTION(BlueprintCallable, Category="Widget")
+	bool HasUserFocus(int32 UserIndex) const;
+
+	/** Gets the focus to this widget. */
+	UFUNCTION(BlueprintCallable, Category="Widget")
+	bool HasAnyUserFocus() const;
+	
+	/** Sets the focus to this widget. */
+	UFUNCTION(BlueprintCallable, Category="Widget")
+	void SetUserFocus(int32 UserIndex);
+
 	/** Forces the underlying slate system to perform a pre-pass on the layout of the widget.  This is for advanced users. */
 	UFUNCTION(BlueprintCallable, Category="Widget")
 	void ForceLayoutPrepass();
