UMG - Reorganizing the properties on the progress bar to make it a bit cleaner.

diff --git a/Engine/Source/Runtime/UMG/Public/Components/ProgressBar.h b/Engine/Source/Runtime/UMG/Public/Components/ProgressBar.h
--- a/Engine/Source/Runtime/UMG/Public/Components/ProgressBar.h
+++ b/Engine/Source/Runtime/UMG/Public/Components/ProgressBar.h
@@ -38,17 +38,17 @@ public:
 	UPROPERTY()
 	USlateBrushAsset* MarqueeImage_DEPRECATED;
 
+	/** Used to determine the fill position of the progress bar ranging 0..1 */
+	UPROPERTY(EditDefaultsOnly, BlueprintReadOnly, Category=Progress, meta=( UIMin = "0", UIMax = "1" ))
+	float Percent;
+
 	/** Defines if this progress bar fills Left to right or right to left */
-	UPROPERTY(EditDefaultsOnly, BlueprintReadOnly, Category=Appearance)
+	UPROPERTY(EditDefaultsOnly, BlueprintReadOnly, Category=Progress)
 	TEnumAsByte<EProgressBarFillType::Type> BarFillType;
 	
-	UPROPERTY(EditDefaultsOnly, BlueprintReadOnly, Category=Appearance)
+	UPROPERTY(EditDefaultsOnly, BlueprintReadOnly, Category=Progress)
 	bool bIsMarquee;
 
-	/** Used to determine the fill position of the progress bar ranging 0..1 */
-	UPROPERTY(EditDefaultsOnly, BlueprintReadOnly, Category=Appearance, meta=( UIMin = "0", UIMax = "1" ))
-	float Percent;
-
 	/** A bindable delegate to allow logic to drive the text of the widget */
 	UPROPERTY()
 	FGetFloat PercentDelegate;
@@ -64,15 +64,15 @@ public:
 public:
 	
 	/** Sets the current value of the ProgressBar. */
-	UFUNCTION(BlueprintCallable, Category="Progress Bar")
+	UFUNCTION(BlueprintCallable, Category="Progress")
 	void SetPercent(float InPercent);
 
 	/** Sets the fill color of the progress bar. */
-	UFUNCTION(BlueprintCallable, Category="Progress Bar")
+	UFUNCTION(BlueprintCallable, Category="Progress")
 	void SetFillColorAndOpacity(FLinearColor InColor);
 
 	/** Sets the progress bar to show as a marquee. */
-	UFUNCTION(BlueprintCallable, Category="Behavior")
+	UFUNCTION(BlueprintCallable, Category="Progress")
 	void SetIsMarquee(bool InbIsMarquee);
 
 	//TODO UMG Add Set BarFillType.
