UMG - Add a play mode option to the PlayAnimation node which allows playing in reverse and ping-pong modes and add an event for when an animation starts.

diff --git a/Engine/Source/Runtime/UMG/Private/Animation/UMGSequencePlayer.cpp b/Engine/Source/Runtime/UMG/Private/Animation/UMGSequencePlayer.cpp
--- a/Engine/Source/Runtime/UMG/Private/Animation/UMGSequencePlayer.cpp
+++ b/Engine/Source/Runtime/UMG/Private/Animation/UMGSequencePlayer.cpp
@@ -56,31 +56,51 @@ void UUMGSequencePlayer::Tick(float DeltaTime)
 	{
 		double LastTimePosition = TimeCursorPosition;
 
-		TimeCursorPosition += DeltaTime;
+		TimeCursorPosition += bIsPlayingForward ? DeltaTime : -DeltaTime;
 
-		if (RootMovieSceneInstance.IsValid())
+		float AnimationLength = TimeRange.GetUpperBoundValue();
+		if ( TimeCursorPosition < 0 )
 		{
-			RootMovieSceneInstance->Update(TimeCursorPosition, LastTimePosition, *this);
+			if ( PlayMode == EUMGSequencePlayMode::PingPong )
+			{
+				bIsPlayingForward = !bIsPlayingForward;
+				TimeCursorPosition = FMath::Abs(TimeCursorPosition);
+			}
+			else
+			{
+				TimeCursorPosition += AnimationLength;
+			}
+			NumLoopsCompleted++;
 		}
-
-		// Has a loop been completed?
-		if ( TimeCursorPosition >= TimeRange.GetUpperBoundValue() )
+		else if ( TimeCursorPosition >= AnimationLength )
 		{
+			if ( PlayMode == EUMGSequencePlayMode::PingPong )
+			{
+				bIsPlayingForward = !bIsPlayingForward;
+				TimeCursorPosition = AnimationLength - (TimeCursorPosition - AnimationLength);
+			}
+			else
+			{
+				TimeCursorPosition -= AnimationLength;
+			}
 			NumLoopsCompleted++;
-			TimeCursorPosition = 0.0;
 		}
 		
 		// Have we looped the requested number of times?
 		if (NumLoopsToPlay != 0 && NumLoopsCompleted >= NumLoopsToPlay)
 		{
 			PlayerStatus = EMovieScenePlayerStatus::Stopped;
-
 			OnSequenceFinishedPlayingEvent.Broadcast( *this );
 		}
+		
+		if (RootMovieSceneInstance.IsValid() && PlayerStatus != EMovieScenePlayerStatus::Stopped)
+		{
+			RootMovieSceneInstance->Update(TimeCursorPosition, LastTimePosition, *this);
+		}
 	}
 }
 
-void UUMGSequencePlayer::Play(float StartAtTime, int32 InNumLoopsToPlay)
+void UUMGSequencePlayer::Play(float StartAtTime, int32 InNumLoopsToPlay, EUMGSequencePlayMode::Type InPlayMode)
 {
 	UMovieScene* MovieScene = Animation->MovieScene;
 
@@ -89,8 +109,18 @@ void UUMGSequencePlayer::Play(float StartAtTime, int32 InNumLoopsToPlay)
 	RootMovieSceneInstance->RefreshInstance( *this );
 
 	TimeCursorPosition = StaticCast<double>(StartAtTime);
-	NumLoopsToPlay = InNumLoopsToPlay;
+	PlayMode = InPlayMode;
+	if ( PlayMode == EUMGSequencePlayMode::PingPong )
+	{
+		// When animating in ping-pong mode double the number of loops to play so that a loop is a complete forward/reverse cycle.
+		NumLoopsToPlay = 2 * InNumLoopsToPlay;
+	}
+	else
+	{
+		NumLoopsToPlay = InNumLoopsToPlay;
+	}
 	NumLoopsCompleted = 0;
+	bIsPlayingForward = InPlayMode != EUMGSequencePlayMode::Reverse;
 
 	PlayerStatus = EMovieScenePlayerStatus::Playing;
 }
diff --git a/Engine/Source/Runtime/UMG/Private/UserWidget.cpp b/Engine/Source/Runtime/UMG/Private/UserWidget.cpp
--- a/Engine/Source/Runtime/UMG/Private/UserWidget.cpp
+++ b/Engine/Source/Runtime/UMG/Private/UserWidget.cpp
@@ -371,12 +371,17 @@ FEventReply UUserWidget::OnMotionDetected_Implementation(FGeometry MyGeometry, F
 	return UWidgetBlueprintLibrary::Unhandled();
 }
 
+void UUserWidget::OnAnimationStarted_Implementation( const UWidgetAnimation* Animation )
+{
+
+}
+
 void UUserWidget::OnAnimationFinished_Implementation( const UWidgetAnimation* Animation )
 {
 	
 }
 
-void UUserWidget::PlayAnimation(const UWidgetAnimation* InAnimation, float StartAtTime, int32 NumberOfLoops)
+void UUserWidget::PlayAnimation( const UWidgetAnimation* InAnimation, float StartAtTime, int32 NumberOfLoops, EUMGSequencePlayMode::Type PlayMode)
 {
 	if( InAnimation )
 	{
@@ -392,12 +397,14 @@ void UUserWidget::PlayAnimation(const UWidgetAnimation* InAnimation, float Start
 
 			NewPlayer->InitSequencePlayer( *InAnimation, *this );
 
-			NewPlayer->Play( StartAtTime, NumberOfLoops );
+			NewPlayer->Play( StartAtTime, NumberOfLoops, PlayMode );
 		}
 		else
 		{
-			( *FoundPlayer )->Play( StartAtTime, NumberOfLoops );
+			( *FoundPlayer )->Play( StartAtTime, NumberOfLoops, PlayMode );
 		}
+
+		OnAnimationStarted( InAnimation );
 	}
 }
 
diff --git a/Engine/Source/Runtime/UMG/Public/Animation/UMGSequencePlayer.h b/Engine/Source/Runtime/UMG/Public/Animation/UMGSequencePlayer.h
--- a/Engine/Source/Runtime/UMG/Public/Animation/UMGSequencePlayer.h
+++ b/Engine/Source/Runtime/UMG/Public/Animation/UMGSequencePlayer.h
@@ -20,7 +20,7 @@ public:
 	void Tick( float DeltaTime );
 
 	/** Begins playing or restarts an animation */
-	void Play( float StartAtTime, int32 InNumLoopsToPlay );
+	void Play( float StartAtTime, int32 InNumLoopsToPlay, EUMGSequencePlayMode::Type InPlayMode);
 
 	/** Stops a running animation and resets time */
 	void Stop();
@@ -71,4 +71,10 @@ private:
 
 	/** The number of loops completed since the last call to Play() */
 	int32 NumLoopsCompleted;
+
+	/** The current playback mode. */
+	EUMGSequencePlayMode::Type PlayMode;
+
+	/** True if the animation is playing forward, otherwise false and it's playing in reverse. */
+	bool bIsPlayingForward;
 };
diff --git a/Engine/Source/Runtime/UMG/Public/Blueprint/UserWidget.h b/Engine/Source/Runtime/UMG/Public/Blueprint/UserWidget.h
--- a/Engine/Source/Runtime/UMG/Public/Blueprint/UserWidget.h
+++ b/Engine/Source/Runtime/UMG/Public/Blueprint/UserWidget.h
@@ -101,6 +101,21 @@ public:
 
 class UUMGSequencePlayer;
 
+/** Describes playback modes for UMG sequences. */
+UENUM()
+namespace EUMGSequencePlayMode
+{
+	enum Type
+	{
+		/** Animation plays and loops from the beginning to the end. */
+		Forward,
+		/** Animation plays and loops from the end to the beginning. */
+		Reverse,
+		/** Animation plays from the begging to the end and then from the end to beginning. */
+		PingPong,
+	};
+}
+
 //TODO UMG If you want to host a widget that's full screen there may need to be a SWindow equivalent that you spawn it into.
 
 DECLARE_DYNAMIC_MULTICAST_DELEGATE(FOnConstructEvent);
@@ -525,6 +540,14 @@ public:
 	UFUNCTION(BlueprintNativeEvent, BlueprintCosmetic, Category="Touch Input")
 	FEventReply OnMotionDetected(FGeometry MyGeometry, FMotionEvent InMotionEvent);
 
+	/**
+	 * Called when an animation is started.
+	 *
+	 * @param Animation the animation that started
+	 */
+	UFUNCTION( BlueprintNativeEvent, BlueprintCosmetic, Category = "Animation" )
+	void OnAnimationStarted( const UWidgetAnimation* Animation );
+
 	/**
 	 * Called when an animation has either played all the way through or is stopped
 	 *
@@ -558,9 +581,10 @@ public:
 	 * @param InAnimation The animation to play
 	 * @param StartAtTime The time in the animation from which to start playing. For looped animations, this will only affect the first playback of the animation.
 	 * @param NumLoopsToPlay The number of times to loop this animation (0 to loop indefinitely)
+	 * @param PlayMode Specifies the playback mode
 	 */
 	UFUNCTION(BlueprintCallable, BlueprintCosmetic, Category="User Interface|Animation")
-	void PlayAnimation(const UWidgetAnimation* InAnimation, float StartAtTime = 0.0f, int32 NumLoopsToPlay = 1);
+	void PlayAnimation(const UWidgetAnimation* InAnimation, float StartAtTime = 0.0f, int32 NumLoopsToPlay = 1, EUMGSequencePlayMode::Type PlayMode = EUMGSequencePlayMode::Forward);
 
 	/**
 	 * Stops an already running animation in this widget
