[UE-11299] - Updated TranslateMouseCoordinateFor3DChild to properly account for the location of the pivot on the widget component

diff --git a/Engine/Source/Runtime/UMG/Private/Components/WidgetComponent.cpp b/Engine/Source/Runtime/UMG/Private/Components/WidgetComponent.cpp
--- a/Engine/Source/Runtime/UMG/Private/Components/WidgetComponent.cpp
+++ b/Engine/Source/Runtime/UMG/Private/Components/WidgetComponent.cpp
@@ -241,7 +241,8 @@ public:
 							{
 								TSharedPtr<FVirtualPointerPosition> VirtualCursorPos = MakeShareable(new FVirtualPointerPosition);
 
-								FVector2D LocalHitLocation = FVector2D( WidgetComponent->ComponentToWorld.InverseTransformPosition( HitResult.Location ) );
+								FVector2D LocalHitLocation;
+								WidgetComponent->GetLocalHitLocation(HitResult, LocalHitLocation);
 
 								VirtualCursorPos->CurrentCursorPosition = LocalHitLocation;
 								VirtualCursorPos->LastCursorPosition = LocalHitLocation;
@@ -817,6 +818,16 @@ void UWidgetComponent::UpdateBodySetup( bool bDrawSizeChanged )
 	}	
 }
 
+void UWidgetComponent::GetLocalHitLocation(const FHitResult& HitResult, FVector2D& OutLocalHitLocation) const
+{
+	// Find the hit location on the component
+	OutLocalHitLocation = FVector2D(ComponentToWorld.InverseTransformPosition(HitResult.Location));
+
+	// Offset the position by the pivot to get the position in widget space.
+	OutLocalHitLocation.X += DrawSize.X * Pivot.X;
+	OutLocalHitLocation.Y += DrawSize.Y * Pivot.Y;
+}
+
 UUserWidget* UWidgetComponent::GetUserWidgetObject() const
 {
 	return Widget;
@@ -824,11 +835,8 @@ UUserWidget* UWidgetComponent::GetUserWidgetObject() const
 
 TArray<FWidgetAndPointer> UWidgetComponent::GetHitWidgetPath( const FHitResult& HitResult, bool bIgnoreEnabledStatus  )
 {
-	// Find the hit location on the component
-	FVector2D LocalHitLocation = FVector2D( ComponentToWorld.InverseTransformPosition( HitResult.Location ) );
-	// Offset the position by the pivot to get the position in widget space.
-	LocalHitLocation.X += DrawSize.X * Pivot.X;
-	LocalHitLocation.Y += DrawSize.Y * Pivot.Y;
+	FVector2D LocalHitLocation;
+	GetLocalHitLocation(HitResult, LocalHitLocation);
 
 	TSharedRef<FVirtualPointerPosition> VirtualMouseCoordinate = MakeShareable( new FVirtualPointerPosition );
 
diff --git a/Engine/Source/Runtime/UMG/Public/Components/WidgetComponent.h b/Engine/Source/Runtime/UMG/Public/Components/WidgetComponent.h
--- a/Engine/Source/Runtime/UMG/Public/Components/WidgetComponent.h
+++ b/Engine/Source/Runtime/UMG/Public/Components/WidgetComponent.h
@@ -61,6 +61,13 @@ public:
 	*/
 	void UpdateBodySetup( bool bDrawSizeChanged = false );
 
+	/**
+	 * Converts a world-space hit result to a hit location on the widget
+	 * @param HitResult The hit on this widget in the world
+	 * @param (Out) The transformed 2D hit location on the widget
+	 */
+	void GetLocalHitLocation(const FHitResult& HitResult, FVector2D& OutLocalHitLocation) const;
+
 	/** @return The class of the user widget displayed by this component */
 	TSubclassOf<UUserWidget> GetWidgetClass() const { return WidgetClass; }
 
@@ -103,7 +110,7 @@ public:
 	FVector2D GetPivot() const { return Pivot; }
 	
 private:
-	/** The class of User Widget to create and display an instance of */
+	/** The coordinate space in which to render the widget */
 	UPROPERTY(EditAnywhere, Category=UI)
 	EWidgetSpace Space;
 
