UMG - Fixing a flaw in the order Construct is called for a widget, it was being called too early, before the SObjectWidget was created, which can cause lots of strange problems like setting visibility on the inner widget instead of a wrapper if done in the Construct event.  You can now set user focus to a widget using their player controller.  Improving the ForceLayoutPrepass tooltip.

diff --git a/Engine/Source/Runtime/UMG/Private/Components/Widget.cpp b/Engine/Source/Runtime/UMG/Private/Components/Widget.cpp
--- a/Engine/Source/Runtime/UMG/Private/Components/Widget.cpp
+++ b/Engine/Source/Runtime/UMG/Private/Components/Widget.cpp
@@ -305,13 +305,26 @@ void UWidget::SetKeyboardFocus()
 	}
 }
 
-bool UWidget::HasUserFocus(int32 UserIndex) const
+bool UWidget::HasUserFocus(APlayerController* PlayerController) const
 {
+	if ( PlayerController == nullptr || !PlayerController->IsLocalPlayerController() )
+	{
+		return false;
+	}
+
 	TSharedPtr<SWidget> SafeWidget = GetCachedWidget();
 	if ( SafeWidget.IsValid() )
 	{
-		TOptional<EFocusCause> FocusCause = SafeWidget->HasUserFocus(UserIndex);
-		return FocusCause.IsSet();
+		FLocalPlayerContext Context(PlayerController);
+
+		if ( ULocalPlayer* LocalPlayer = Context.GetLocalPlayer() )
+		{
+			// HACK: We use the controller Id as the local player index for focusing widgets in Slate.
+			int32 UserIndex = LocalPlayer->GetControllerId();
+
+			TOptional<EFocusCause> FocusCause = SafeWidget->HasUserFocus(UserIndex);
+			return FocusCause.IsSet();
+		}
 	}
 
 	return false;
@@ -329,12 +342,25 @@ bool UWidget::HasAnyUserFocus() const
 	return false;
 }
 
-void UWidget::SetUserFocus(int32 UserIndex)
+void UWidget::SetUserFocus(APlayerController* PlayerController)
 {
+	if ( PlayerController == nullptr || !PlayerController->IsLocalPlayerController() )
+	{
+		return;
+	}
+
 	TSharedPtr<SWidget> SafeWidget = GetCachedWidget();
 	if ( SafeWidget.IsValid() )
 	{
-		FSlateApplication::Get().SetUserFocus(UserIndex, SafeWidget);
+		FLocalPlayerContext Context(PlayerController);
+
+		if ( ULocalPlayer* LocalPlayer = Context.GetLocalPlayer() )
+		{
+			// HACK: We use the controller Id as the local player index for focusing widgets in Slate.
+			int32 UserIndex = LocalPlayer->GetControllerId();
+
+			FSlateApplication::Get().SetUserFocus(UserIndex, SafeWidget);
+		}
 	}
 }
 
@@ -388,8 +414,6 @@ TSharedRef<SWidget> UWidget::TakeWidget()
 		SafeWidget = RebuildWidget();
 		MyWidget = SafeWidget;
 
-		OnWidgetRebuilt();
-
 		bNewlyCreated = true;
 	}
 	else
@@ -400,22 +424,27 @@ TSharedRef<SWidget> UWidget::TakeWidget()
 	// If it is a user widget wrap it in a SObjectWidget to keep the instance from being GC'ed
 	if ( IsA(UUserWidget::StaticClass()) )
 	{
+		TSharedPtr<SObjectWidget> SafeGCWidget = MyGCWidget.Pin();
+
 		// If the GC Widget is still valid we still exist in the slate hierarchy, so just return the GC Widget.
-		if ( MyGCWidget.IsValid() )
+		if ( SafeGCWidget.IsValid() )
 		{
-			return MyGCWidget.Pin().ToSharedRef();
+			return SafeGCWidget.ToSharedRef();
 		}
 		// Otherwise we need to recreate the wrapper widget
 		else
 		{
-			TSharedPtr<SObjectWidget> SafeGCWidget = SNew(SObjectWidget, Cast<UUserWidget>(this))
+			SafeGCWidget = SNew(SObjectWidget, Cast<UUserWidget>(this))
 				[
 					SafeWidget.ToSharedRef()
 				];
 
 			MyGCWidget = SafeGCWidget;
 
+			// Always synchronize properties of a user widget and call construct AFTER we've
+			// properly setup the GCWidget and synced all the properties.
 			SynchronizeProperties();
+			OnWidgetRebuilt();
 
 			return SafeGCWidget.ToSharedRef();
 		}
@@ -425,6 +454,7 @@ TSharedRef<SWidget> UWidget::TakeWidget()
 		if ( bNewlyCreated )
 		{
 			SynchronizeProperties();
+			OnWidgetRebuilt();
 		}
 
 		return SafeWidget.ToSharedRef();
diff --git a/Engine/Source/Runtime/UMG/Public/Components/Widget.h b/Engine/Source/Runtime/UMG/Public/Components/Widget.h
--- a/Engine/Source/Runtime/UMG/Public/Components/Widget.h
+++ b/Engine/Source/Runtime/UMG/Public/Components/Widget.h
@@ -281,17 +281,21 @@ public:
 
 	/** Gets the focus to this widget. */
 	UFUNCTION(BlueprintCallable, Category="Widget")
-	bool HasUserFocus(int32 UserIndex) const;
+	bool HasUserFocus(APlayerController* PlayerController) const;
 
 	/** Gets the focus to this widget. */
 	UFUNCTION(BlueprintCallable, Category="Widget")
 	bool HasAnyUserFocus() const;
 	
-	/** Sets the focus to this widget. */
+	/** Sets the focus to this widget for a specific user */
 	UFUNCTION(BlueprintCallable, Category="Widget")
-	void SetUserFocus(int32 UserIndex);
+	void SetUserFocus(APlayerController* PlayerController);
 
-	/** Forces the underlying slate system to perform a pre-pass on the layout of the widget.  This is for advanced users. */
+	/**
+	 * Forces a pre-pass.  A pre-pass caches the desired size of the widget hierarchy owned by this widget.  
+	 * One pre-pass is already happens for every widget before Tick occurs.  You only need to perform another 
+	 * pre-pass if you are adding child widgets this frame and want them to immediately be visible this frame.
+	 */
 	UFUNCTION(BlueprintCallable, Category="Widget")
 	void ForceLayoutPrepass();
 
