UMG - Menu Anchor - Several new functions on menu anchor, along with a delegate to be notified when the opened state of the menu anchor changes.

diff --git a/Engine/Source/Runtime/UMG/Private/Components/MenuAnchor.cpp b/Engine/Source/Runtime/UMG/Private/Components/MenuAnchor.cpp
--- a/Engine/Source/Runtime/UMG/Private/Components/MenuAnchor.cpp
+++ b/Engine/Source/Runtime/UMG/Private/Components/MenuAnchor.cpp
@@ -24,7 +24,8 @@ TSharedRef<SWidget> UMenuAnchor::RebuildWidget()
 {
 	MyMenuAnchor = SNew(SMenuAnchor)
 		.Placement(Placement)
-		.OnGetMenuContent(BIND_UOBJECT_DELEGATE(FOnGetContent, HandleGetMenuContent));
+		.OnGetMenuContent(BIND_UOBJECT_DELEGATE(FOnGetContent, HandleGetMenuContent))
+		.OnMenuOpenChanged(BIND_UOBJECT_DELEGATE(FOnIsOpenChanged, HandleMenuOpenChanged));
 
 	if ( GetChildrenCount() > 0 )
 	{
@@ -52,6 +53,11 @@ void UMenuAnchor::OnSlotRemoved(UPanelSlot* Slot)
 	}
 }
 
+void UMenuAnchor::HandleMenuOpenChanged(bool bIsOpen)
+{
+	OnMenuOpenChanged.Broadcast(bIsOpen);
+}
+
 TSharedRef<SWidget> UMenuAnchor::HandleGetMenuContent()
 {
 	TSharedPtr<SWidget> SlateMenuWidget;
@@ -115,6 +121,36 @@ bool UMenuAnchor::IsOpen() const
 	return false;
 }
 
+bool UMenuAnchor::ShouldOpenDueToClick() const
+{
+	if ( MyMenuAnchor.IsValid() )
+	{
+		return MyMenuAnchor->ShouldOpenDueToClick();
+	}
+
+	return false;
+}
+
+FVector2D UMenuAnchor::GetMenuPosition() const
+{
+	if ( MyMenuAnchor.IsValid() )
+	{
+		return MyMenuAnchor->GetMenuPosition();
+	}
+
+	return FVector2D(0, 0);
+}
+
+bool UMenuAnchor::HasOpenSubMenus() const
+{
+	if ( MyMenuAnchor.IsValid() )
+	{
+		return MyMenuAnchor->HasOpenSubMenus();
+	}
+
+	return false;
+}
+
 #if WITH_EDITOR
 
 const FSlateBrush* UMenuAnchor::GetEditorIcon()
diff --git a/Engine/Source/Runtime/UMG/Public/Components/MenuAnchor.h b/Engine/Source/Runtime/UMG/Public/Components/MenuAnchor.h
--- a/Engine/Source/Runtime/UMG/Public/Components/MenuAnchor.h
+++ b/Engine/Source/Runtime/UMG/Public/Components/MenuAnchor.h
@@ -1,12 +1,16 @@
-// Copyright 1998-2015 Epic Games, Inc. All Rights Reserved.
+﻿// Copyright 1998-2015 Epic Games, Inc. All Rights Reserved.
 
 #pragma once
 
 #include "MenuAnchor.generated.h"
 
+DECLARE_DYNAMIC_MULTICAST_DELEGATE_OneParam(FOnMenuOpenChangedEvent, bool, bIsOpen);
+
 /**
  * The Menu Anchor allows you to specify an location that a popup menu should be anchored to, 
  * and should be summoned from.
+ * ● Single Child
+ * ● Popup
  */
 UCLASS()
 class UMG_API UMenuAnchor : public UContentWidget
@@ -20,7 +24,7 @@ public:
 	 * If you want to customize the creation of the popup, you should bind a function to OnGetMenuContentEvent 
 	 * instead.
 	 */
-	UPROPERTY(EditDefaultsOnly, BlueprintReadOnly, Category="Menu")
+	UPROPERTY(EditDefaultsOnly, BlueprintReadOnly, Category="Menu Anchor")
 	TSubclassOf<class UUserWidget> MenuClass;
 
 	/** Called when the menu content is requested to allow a more customized handling over what to display */
@@ -28,9 +32,14 @@ public:
 	FGetWidget OnGetMenuContentEvent;
 	
 	/** The placement location of the summoned widget. */
-	UPROPERTY(EditDefaultsOnly, BlueprintReadOnly, Category="Menu")
+	UPROPERTY(EditDefaultsOnly, BlueprintReadOnly, Category="Menu Anchor")
 	TEnumAsByte<EMenuPlacement> Placement;
 
+public:
+	/** Called when the opened state of the menu changes */
+	UPROPERTY(BlueprintAssignable, Category="Menu Anchor|Event")
+	FOnMenuOpenChangedEvent OnMenuOpenChanged;
+
 public:
 	//TODO UMG Add Set MenuClass
 	//TODO UMG Add Set Placement
@@ -57,15 +66,21 @@ public:
 	UFUNCTION(BlueprintCallable, Category="Menu Anchor")
 	bool IsOpen() const;
 
-	///** @return true if we should open the menu due to a click. Sometimes we should not, if
-	//the same MouseDownEvent that just closed the menu is about to re-open it because it happens to land on the button.*/
-	//bool ShouldOpenDueToClick() const;
+	/**
+	 * @return true if we should open the menu due to a click. Sometimes we should not, if
+	 * the same MouseDownEvent that just closed the menu is about to re-open it because it 
+	 * happens to land on the button.
+	 */
+	UFUNCTION(BlueprintCallable, Category="Menu Anchor")
+	bool ShouldOpenDueToClick() const;
 
-	///** @return The current menu position */
-	//FVector2D GetMenuPosition() const;
+	/** @return The current menu position */
+	UFUNCTION(BlueprintCallable, Category="Menu Anchor")
+	FVector2D GetMenuPosition() const;
 
-	///** @return Whether this menu has open submenus */
-	//bool HasOpenSubMenus() const;
+	/** @return Whether this menu has open submenus */
+	UFUNCTION(BlueprintCallable, Category="Menu Anchor")
+	bool HasOpenSubMenus() const;
 
 	virtual void ReleaseSlateResources(bool bReleaseChildren) override;
 
@@ -87,6 +102,7 @@ protected:
 
 protected:
 	TSharedRef<SWidget> HandleGetMenuContent();
+	void HandleMenuOpenChanged(bool bIsOpen);
 
 protected:
 	TSharedPtr<SMenuAnchor> MyMenuAnchor;
