UMG - Adding a way to set the brush on image and border.

diff --git a/Engine/Source/Runtime/UMG/Private/Components/Border.cpp b/Engine/Source/Runtime/UMG/Private/Components/Border.cpp
--- a/Engine/Source/Runtime/UMG/Private/Components/Border.cpp
+++ b/Engine/Source/Runtime/UMG/Private/Components/Border.cpp
@@ -181,6 +181,16 @@ FReply UBorder::HandleMouseDoubleClick(const FGeometry& Geometry, const FPointer
 	return FReply::Unhandled();
 }
 
+void UBorder::SetBrush(FSlateBrush& Brush)
+{
+	Background = Brush;
+
+	if ( MyBorder.IsValid() )
+	{
+		MyBorder->SetBorderImage(&Background);
+	}
+}
+
 void UBorder::SetBrushFromAsset(USlateBrushAsset* Asset)
 {
 	Background = Asset ? Asset->Brush : FSlateBrush();
diff --git a/Engine/Source/Runtime/UMG/Private/Components/Image.cpp b/Engine/Source/Runtime/UMG/Private/Components/Image.cpp
--- a/Engine/Source/Runtime/UMG/Private/Components/Image.cpp
+++ b/Engine/Source/Runtime/UMG/Private/Components/Image.cpp
@@ -76,6 +76,16 @@ const FSlateBrush* UImage::ConvertImage(TAttribute<FSlateBrush> InImageAsset) co
 	return &Brush;
 }
 
+void UImage::SetBrush(FSlateBrush& InBrush)
+{
+	Brush = InBrush;
+
+	if ( MyImage.IsValid() )
+	{
+		MyImage->SetImage(&Brush);
+	}
+}
+
 void UImage::SetBrushFromAsset(USlateBrushAsset* Asset)
 {
 	Brush = Asset ? Asset->Brush : FSlateBrush();
diff --git a/Engine/Source/Runtime/UMG/Public/Components/Border.h b/Engine/Source/Runtime/UMG/Public/Components/Border.h
--- a/Engine/Source/Runtime/UMG/Public/Components/Border.h
+++ b/Engine/Source/Runtime/UMG/Public/Components/Border.h
@@ -11,6 +11,7 @@ class USlateBrushAsset;
  * to surround it with a background image and adjustable padding.
  *
  * ● Single Child
+ * ● Image
  */
 UCLASS()
 class UMG_API UBorder : public UContentWidget
@@ -92,6 +93,10 @@ public:
 	UFUNCTION(BlueprintCallable, Category="Appearance")
 	void SetBrushColor(FLinearColor InBrushColor);
 
+	/**  */
+	UFUNCTION(BlueprintCallable, Category="Appearance")
+	void SetBrush(FSlateBrush& InBrush);
+
 	/**  */
 	UFUNCTION(BlueprintCallable, Category="Appearance")
 	void SetBrushFromAsset(USlateBrushAsset* Asset);
diff --git a/Engine/Source/Runtime/UMG/Public/Components/Image.h b/Engine/Source/Runtime/UMG/Public/Components/Image.h
--- a/Engine/Source/Runtime/UMG/Public/Components/Image.h
+++ b/Engine/Source/Runtime/UMG/Public/Components/Image.h
@@ -53,6 +53,10 @@ public:
 	UFUNCTION(BlueprintCallable, Category="Appearance")
 	void SetOpacity(float InOpacity);
 
+	/**  */
+	UFUNCTION(BlueprintCallable, Category="Appearance")
+	void SetBrush(FSlateBrush& InBrush);
+
 	/**  */
 	UFUNCTION(BlueprintCallable, Category="Appearance")
 	void SetBrushFromAsset(USlateBrushAsset* Asset);
