UMG - Fixing a binding asserting in for checkbox checked state binding.

diff --git a/Engine/Source/Runtime/UMG/Private/Binding/CheckedStateBinding.cpp b/Engine/Source/Runtime/UMG/Private/Binding/CheckedStateBinding.cpp
--- a/Engine/Source/Runtime/UMG/Private/Binding/CheckedStateBinding.cpp
+++ b/Engine/Source/Runtime/UMG/Private/Binding/CheckedStateBinding.cpp
@@ -16,7 +16,7 @@ bool UCheckedStateBinding::IsSupportedSource(UProperty* Property) const
 
 bool UCheckedStateBinding::IsSupportedDestination(UProperty* Property) const
 {
-	static const FName CheckBoxStateEnum(TEXT("ESlateCheckBoxState"));
+	static const FName CheckBoxStateEnum(TEXT("ECheckBoxState"));
 
 	if ( UByteProperty* ByteProperty = Cast<UByteProperty>(Property) )
 	{
