UMG - Adding support for split screen to UMG and Slate.  There is now a framework at the game viewport client level that allows content to be placed into a canvas slot that is kept in the sub-rect for the player.  The new framework also provides a tooltip provider so that new windows are not created for games if a tool tip pops up.

diff --git a/Engine/Source/Editor/LevelEditor/Private/SLevelViewport.cpp b/Engine/Source/Editor/LevelEditor/Private/SLevelViewport.cpp
--- a/Engine/Source/Editor/LevelEditor/Private/SLevelViewport.cpp
+++ b/Engine/Source/Editor/LevelEditor/Private/SLevelViewport.cpp
@@ -29,8 +29,6 @@
 #include "EditorViewportCommands.h"
 #include "Runtime/Engine/Classes/Engine/UserInterfaceSettings.h"
 #include "Runtime/Engine/Classes/Engine/RendererSettings.h"
-#include "SScissorRectBox.h"
-#include "SDPIScaler.h"
 #include "SNotificationList.h"
 #include "NotificationManager.h"
 #include "SLevelViewportControlsPopup.h"
@@ -41,6 +39,7 @@
 #include "GameFramework/PlayerInput.h"
 #include "GameFramework/PlayerController.h"
 #include "SActorPilotViewportToolbar.h"
+#include "SGameLayerManager.h"
 
 static const FName LevelEditorName("LevelEditor");
 
@@ -181,13 +180,10 @@ void SLevelViewport::ConstructViewportOverlayContent()
 #if ALLOW_PLAY_IN_VIEWPORT_GAMEUI
 		ViewportOverlay->AddSlot( SlotIndex )
 		[
-			SNew( SScissorRectBox )
+			SAssignNew(GameLayerManager, SGameLayerManager)
+			.SceneViewport(this, &SLevelViewport::GetGameSceneViewport)
 			[
-				SNew(SDPIScaler)
-				.DPIScale(this, &SLevelViewport::GetGameViewportDPIScale)
-				[
-					PIEViewportOverlayWidget.ToSharedRef()
-				]
+				PIEViewportOverlayWidget.ToSharedRef()
 			]
 		];
 
@@ -386,14 +382,9 @@ void SLevelViewport::ConstructLevelEditorViewportClient( const FArguments& InArg
 	LevelViewportClient->SetViewModes(ViewportInstanceSettings.PerspViewModeIndex, ViewportInstanceSettings.OrthoViewModeIndex );
 }
 
-float SLevelViewport::GetGameViewportDPIScale() const
+const FSceneViewport* SLevelViewport::GetGameSceneViewport() const
 {
-	if ( HasPlayInEditorViewport() || LevelViewportClient->IsSimulateInEditorViewport() )
-	{
-		return GetDefault<UUserInterfaceSettings>(UUserInterfaceSettings::StaticClass())->GetDPIScaleBasedOnSize(ActiveViewport->GetSize());
-	}
-
-	return 1.0f;
+	return ActiveViewport.Get();
 }
 
 FReply SLevelViewport::OnKeyDown( const FGeometry& MyGeometry, const FKeyEvent& InKeyEvent )
@@ -3246,7 +3237,7 @@ bool SLevelViewport::IsViewportConfigurationSet(FName ConfigurationName) const
 	return false;
 }
 
-void SLevelViewport::StartPlayInEditorSession( UGameViewportClient* PlayClient, const bool bInSimulateInEditor )
+void SLevelViewport::StartPlayInEditorSession(UGameViewportClient* PlayClient, const bool bInSimulateInEditor)
 {
 	check( !HasPlayInEditorViewport() );
 
@@ -3276,6 +3267,7 @@ void SLevelViewport::StartPlayInEditorSession( UGameViewportClient* PlayClient,
 	ActiveViewport->OnPlayWorldViewportSwapped( *InactiveViewport );
 
 	PlayClient->SetViewportOverlayWidget( TSharedPtr<SWindow>(), PIEViewportOverlayWidget.ToSharedRef() );
+	PlayClient->SetGameLayerManager(GameLayerManager);
 
 	// Our viewport widget should start rendering the new viewport for the play in editor scene
 	ViewportWidget->SetViewportInterface( ActiveViewport.ToSharedRef() );
diff --git a/Engine/Source/Editor/LevelEditor/Public/SLevelViewport.h b/Engine/Source/Editor/LevelEditor/Public/SLevelViewport.h
--- a/Engine/Source/Editor/LevelEditor/Public/SLevelViewport.h
+++ b/Engine/Source/Editor/LevelEditor/Public/SLevelViewport.h
@@ -7,7 +7,8 @@
 #include "Editor/UnrealEd/Public/SEditorViewport.h"
 
 class SActorPreview;
-	
+class SGameLayerManager;
+
 /**
  * Encapsulates an SViewport and an SLevelViewportToolBar
  */
@@ -613,9 +614,9 @@ private:
 
 	/** Resets view flags when a new level is created or opened */
 	void ResetNewLevelViewFlags();
-private:
-	/** Returns the DPI scaler that should be used for game UI in the viewport */
-	float GetGameViewportDPIScale() const;
+
+	/** Gets the active scene viewport for the game */
+	const FSceneViewport* GetGameSceneViewport() const;
 
 private:
 	/** Tab which this viewport is located in */
@@ -627,6 +628,8 @@ private:
 	/** Viewport overlay widget exposed to game systems when running play-in-editor */
 	TSharedPtr<SOverlay> PIEViewportOverlayWidget;
 
+	TSharedPtr<SGameLayerManager> GameLayerManager;
+
 	/** Viewport horizontal box used internally for drawing actor previews on top of the level viewport */
 	TSharedPtr<SHorizontalBox> ActorPreviewHorizontalBox;
 
diff --git a/Engine/Source/Editor/UnrealEd/Classes/Editor/EditorEngine.h b/Engine/Source/Editor/UnrealEd/Classes/Editor/EditorEngine.h
--- a/Engine/Source/Editor/UnrealEd/Classes/Editor/EditorEngine.h
+++ b/Engine/Source/Editor/UnrealEd/Classes/Editor/EditorEngine.h
@@ -2487,8 +2487,8 @@ public:
 	void OnViewportCloseRequested(FViewport* InViewport);
 
 private:
-	/** Gets the DPI Scale for the game viewport in the editor. */
-	float GetGameViewportDPIScale(UGameViewportClient* ViewportClient) const;
+	/** Gets the scene viewport for a viewport client */
+	const FSceneViewport* GetGameSceneViewport(UGameViewportClient* ViewportClient) const;
 
 	/**
 	 * Non Online PIE creation flow, creates all instances of PIE at once when online isn't requested/required
diff --git a/Engine/Source/Editor/UnrealEd/Private/PlayLevel.cpp b/Engine/Source/Editor/UnrealEd/Private/PlayLevel.cpp
--- a/Engine/Source/Editor/UnrealEd/Private/PlayLevel.cpp
+++ b/Engine/Source/Editor/UnrealEd/Private/PlayLevel.cpp
@@ -36,7 +36,7 @@
 #include "SScissorRectBox.h"
 #include "Online.h"
 #include "SNotificationList.h"
-#include "SDPIScaler.h"
+#include "SGameLayerManager.h"
 #include "NotificationManager.h"
 #include "Engine/Selection.h"
 #include "TimerManager.h"
@@ -2742,9 +2742,14 @@ UGameInstance* UEditorEngine::CreatePIEGameInstance(int32 PIEInstance, bool bInS
 
 				FSlateApplication::Get().AddWindow( PieWindow );
 
-				TSharedPtr<SDPIScaler> PIEDPIScaler;
+				TSharedRef<SOverlay> ViewportOverlayWidgetRef = SNew(SOverlay);
+
+				TSharedRef<SGameLayerManager> GameLayerManagerRef = SNew(SGameLayerManager)
+					.SceneViewport_UObject(this, &UEditorEngine::GetGameSceneViewport, ViewportClient)
+					[
+						ViewportOverlayWidgetRef
+					];
 
-				TSharedRef<SOverlay> ViewportOverlayWidgetRef = SNew( SOverlay );
 				PieViewportWidget = 
 					SNew( SViewport )
 						.IsEnabled(FSlateApplication::Get().GetNormalExecutionAttribute())
@@ -2752,17 +2757,9 @@ UGameInstance* UEditorEngine::CreatePIEGameInstance(int32 PIEInstance, bool bInS
 						.RenderDirectlyToWindow( bRenderDirectlyToWindow )
 						.EnableStereoRendering( bEnableStereoRendering )
 						[
-							SNew(SScissorRectBox)
-							[
-								SAssignNew(PIEDPIScaler, SDPIScaler)
-								[
-									ViewportOverlayWidgetRef
-								]
-							]
+							GameLayerManagerRef
 						];
 
-				PIEDPIScaler->SetDPIScale(TAttribute<float>::Create(TAttribute<float>::FGetter::CreateUObject(this, &UEditorEngine::GetGameViewportDPIScale, ViewportClient)));
-
 				// Create a viewport widget for the game to render in.
 				PieWindow->SetContent( PieViewportWidget.ToSharedRef() );
 
@@ -2770,6 +2767,7 @@ UGameInstance* UEditorEngine::CreatePIEGameInstance(int32 PIEInstance, bool bInS
 				PieWindow->BringToFront();
 
 				ViewportClient->SetViewportOverlayWidget( PieWindow, ViewportOverlayWidgetRef );
+				ViewportClient->SetGameLayerManager(GameLayerManagerRef);
 
 				// Set up a notification when the window is closed so we can clean up PIE
 				{
@@ -2908,11 +2906,9 @@ void UEditorEngine::OnViewportCloseRequested(FViewport* InViewport)
 	RequestEndPlayMap();
 }
 
-float UEditorEngine::GetGameViewportDPIScale(UGameViewportClient* ViewportClient) const
+const FSceneViewport* UEditorEngine::GetGameSceneViewport(UGameViewportClient* ViewportClient) const
 {
-	FVector2D ViewportSize;
-	ViewportClient->GetViewportSize(ViewportSize);
-	return GetDefault<UUserInterfaceSettings>(UUserInterfaceSettings::StaticClass())->GetDPIScaleBasedOnSize(FIntPoint(ViewportSize.X, ViewportSize.Y));
+	return ViewportClient->GetGameViewport();
 }
 
 FViewport* UEditorEngine::GetActiveViewport()
diff --git a/Engine/Source/Runtime/Engine/Classes/Engine/GameEngine.h b/Engine/Source/Runtime/Engine/Classes/Engine/GameEngine.h
--- a/Engine/Source/Runtime/Engine/Classes/Engine/GameEngine.h
+++ b/Engine/Source/Runtime/Engine/Classes/Engine/GameEngine.h
@@ -143,7 +143,7 @@ public:
 
 protected:
 
-	float GetGameViewportDPIScale(UGameViewportClient* GameViewportClient) const;
+	const FSceneViewport* GetGameSceneViewport(UGameViewportClient* ViewportClient) const;
 
 private:
 
diff --git a/Engine/Source/Runtime/Engine/Classes/Engine/GameViewportClient.h b/Engine/Source/Runtime/Engine/Classes/Engine/GameViewportClient.h
--- a/Engine/Source/Runtime/Engine/Classes/Engine/GameViewportClient.h
+++ b/Engine/Source/Runtime/Engine/Classes/Engine/GameViewportClient.h
@@ -8,6 +8,7 @@
 #include "DebugDisplayProperty.h"
 #include "TitleSafeZone.h"
 #include "GameViewportDelegates.h"
+
 #include "GameViewportClient.generated.h"
 
 
@@ -18,6 +19,7 @@ class FSceneViewport;
 class SViewport;
 class SWindow;
 class SOverlay;
+class IGameLayerManager;
 class UCanvas;
 
 
@@ -207,6 +209,23 @@ public:
 	 */
 	virtual void RemoveViewportWidgetContent( TSharedRef<class SWidget> ViewportContent );
 
+	/**
+	 * Adds a widget to the Slate viewport's overlay (i.e for in game UI or tools) at the specified Z-order.
+	 * associates it with a specific player and keeps it in their sub-rect.
+	 * 
+	 * @param  Player The player to add the widget to.
+	 * @param  ViewportContent	The widget to add.  Must be valid.
+	 * @param  ZOrder  The Z-order index for this widget.  Larger values will cause the widget to appear on top of widgets with lower values.
+	 */
+	virtual void AddViewportWidgetForPlayer(ULocalPlayer* Player, TSharedRef<SWidget> ViewportContent, const int32 ZOrder);
+
+	/**
+	 * Removes a previously-added widget from the Slate viewport, in the player's section.
+	 *
+	 * @param	Player The player to remove the widget's viewport from.
+	 * @param	ViewportContent  The widget to remove.  Must be valid.
+	 */
+	virtual void RemoveViewportWidgetForPlayer(ULocalPlayer* Player, TSharedRef<SWidget> ViewportContent);
 
 	/**
 	 * This function removes all widgets from the viewport overlay
@@ -231,7 +250,6 @@ public:
 	 */
 	virtual bool RequiresUncapturedAxisInput() const;
 
-
 	/**
 	 * Set this GameViewportClient's viewport and viewport frame to the viewport specified
 	 * @param	InViewportFrame	The viewportframe to set
@@ -251,6 +269,12 @@ public:
 		ViewportOverlayWidget = InViewportOverlayWidget;
 	}
 
+	/** Assigns the viewport game layer manager for this viewport client.  Should only be called when first created. */
+	void SetGameLayerManager(TSharedPtr< IGameLayerManager > LayerManager)
+	{
+		GameLayerManagerPtr = LayerManager;
+	}
+
 	/** Returns access to this viewport's Slate window */
 	TSharedPtr< SWindow > GetWindow()
 	{
@@ -645,6 +669,9 @@ private:
 	/** Overlay widget that contains widgets to draw on top of the game viewport */
 	TWeakPtr<SOverlay> ViewportOverlayWidget;
 
+	/** The game layer manager allows management of widgets for different player areas of the screen. */
+	TWeakPtr<IGameLayerManager> GameLayerManagerPtr;
+
 	/** Current buffer visualization mode for this game viewport */
 	FName CurrentBufferVisualizationMode;
 
diff --git a/Engine/Source/Runtime/Engine/Private/GameEngine.cpp b/Engine/Source/Runtime/Engine/Private/GameEngine.cpp
--- a/Engine/Source/Runtime/Engine/Private/GameEngine.cpp
+++ b/Engine/Source/Runtime/Engine/Private/GameEngine.cpp
@@ -31,7 +31,7 @@
 #include "IHeadMountedDisplay.h"
 #include "RendererInterface.h"
 #include "HotReloadInterface.h"
-#include "SDPIScaler.h"
+#include "SGameLayerManager.h"
 #include "Components/SkyLightComponent.h"
 #include "Components/ReflectionCaptureComponent.h"
 #include "Engine/GameEngine.h"
@@ -90,6 +90,14 @@ void UGameEngine::CreateGameViewportWidget( UGameViewportClient* GameViewportCli
 {
 	bool bRenderDirectlyToWindow = !GEngine->MatineeScreenshotOptions.bStartWithMatineeCapture && GIsDumpingMovie == 0; 
 	TSharedRef<SOverlay> ViewportOverlayWidgetRef = SNew( SOverlay );
+
+	TSharedRef<SGameLayerManager> GameLayerManagerRef = SNew(SGameLayerManager)
+		.SceneViewport_UObject(this, &UGameEngine::GetGameSceneViewport, GameViewportClient)
+		.UseScissor(false)
+		[
+			ViewportOverlayWidgetRef
+		];
+
 	TSharedRef<SViewport> GameViewportWidgetRef = 
 		SNew( SViewport )
 			// Render directly to the window backbuffer unless capturing a movie or getting screenshots
@@ -99,16 +107,13 @@ void UGameEngine::CreateGameViewportWidget( UGameViewportClient* GameViewportCli
 			.EnableGammaCorrection(false)
 			.EnableStereoRendering(true)
 			[
-				SNew(SDPIScaler)
-				.DPIScale(TAttribute<float>::Create(TAttribute<float>::FGetter::CreateUObject(this, &UGameEngine::GetGameViewportDPIScale, GameViewportClient)))
-				[
-					ViewportOverlayWidgetRef
-				]
+				GameLayerManagerRef
 			];
 
 	GameViewportWidget = GameViewportWidgetRef;
 
 	GameViewportClient->SetViewportOverlayWidget( GameViewportWindow.Pin(), ViewportOverlayWidgetRef );
+	GameViewportClient->SetGameLayerManager(GameLayerManagerRef);
 }
 
 void UGameEngine::CreateGameViewport( UGameViewportClient* GameViewportClient )
@@ -154,11 +159,9 @@ void UGameEngine::CreateGameViewport( UGameViewportClient* GameViewportClient )
 	GameViewport->SetViewportFrame(ViewportFrame);
 }
 
-float UGameEngine::GetGameViewportDPIScale(UGameViewportClient* ViewportClient) const
+const FSceneViewport* UGameEngine::GetGameSceneViewport(UGameViewportClient* ViewportClient) const
 {
-	FVector2D ViewportSize;
-	ViewportClient->GetViewportSize(ViewportSize);
-	return GetDefault<UUserInterfaceSettings>(UUserInterfaceSettings::StaticClass())->GetDPIScaleBasedOnSize(FIntPoint(ViewportSize.X, ViewportSize.Y));
+	return ViewportClient->GetGameViewport();
 }
 
 void UGameEngine::ConditionallyOverrideSettings(int32& ResolutionX, int32& ResolutionY, EWindowMode::Type& WindowMode)
diff --git a/Engine/Source/Runtime/Engine/Private/GameViewportClient.cpp b/Engine/Source/Runtime/Engine/Private/GameViewportClient.cpp
--- a/Engine/Source/Runtime/Engine/Private/GameViewportClient.cpp
+++ b/Engine/Source/Runtime/Engine/Private/GameViewportClient.cpp
@@ -31,6 +31,7 @@
 #include "UserWidget.h"
 #include "GameFramework/GameUserSettings.h"
 #include "Runtime/Engine/Classes/Engine/UserInterfaceSettings.h"
+#include "SGameLayerManager.h"
 
 #define LOCTEXT_NAMESPACE "GameViewport"
 
@@ -264,6 +265,11 @@ void UGameViewportClient::Init(struct FWorldContext& WorldContext, UGameInstance
 					// will pass the audio device handle to them and audio will play on the correct audio device
 					WorldContext.AudioDeviceHandle = AudioDeviceHandle;
 				}
+				else
+				{
+					// Shut it down if we failed to initialize
+					AudioDeviceManager->ShutdownAudioDevice(AudioDeviceHandle);
+				}
 			}
 		}
 	}
@@ -1342,11 +1348,14 @@ void UGameViewportClient::Precache()
 TOptional<bool> UGameViewportClient::QueryShowFocus(const EFocusCause InFocusCause) const
 {
 	UUserInterfaceSettings* UISettings = GetMutableDefault<UUserInterfaceSettings>(UUserInterfaceSettings::StaticClass());
-	if ((UISettings->RenderFocusRule == ERenderFocusRule::NonPointer && InFocusCause == EFocusCause::Mouse) ||
+	
+	if ( UISettings->RenderFocusRule == ERenderFocusRule::Never ||
+		(UISettings->RenderFocusRule == ERenderFocusRule::NonPointer && InFocusCause == EFocusCause::Mouse) ||
 		(UISettings->RenderFocusRule == ERenderFocusRule::NavigationOnly && InFocusCause != EFocusCause::Navigation))
 	{
 		return false;
 	}
+
 	return true;
 }
 
@@ -1885,15 +1894,26 @@ void UGameViewportClient::DrawTransitionMessage(UCanvas* Canvas,const FString& M
 	Canvas->DrawItem( TextItem, 0.5f * (Canvas->ClipX - XL), 0.66f * Canvas->ClipY - YL * 0.5f );	
 }
 
-
 void UGameViewportClient::NotifyPlayerAdded( int32 PlayerIndex, ULocalPlayer* AddedPlayer )
 {
 	LayoutPlayers();
+
+	TSharedPtr< IGameLayerManager > GameLayerManager(GameLayerManagerPtr.Pin());
+	if ( GameLayerManager.IsValid() )
+	{
+		GameLayerManager->NotifyPlayerAdded(PlayerIndex, AddedPlayer);
+	}
 }
 
 void UGameViewportClient::NotifyPlayerRemoved( int32 PlayerIndex, ULocalPlayer* RemovedPlayer )
 {
 	LayoutPlayers();
+
+	TSharedPtr< IGameLayerManager > GameLayerManager(GameLayerManagerPtr.Pin());
+	if ( GameLayerManager.IsValid() )
+	{
+		GameLayerManager->NotifyPlayerRemoved(PlayerIndex, RemovedPlayer);
+	}
 }
 
 void UGameViewportClient::AddViewportWidgetContent( TSharedRef<SWidget> ViewportContent, const int32 ZOrder )
@@ -1918,6 +1938,24 @@ void UGameViewportClient::RemoveViewportWidgetContent( TSharedRef<SWidget> Viewp
 	}
 }
 
+void UGameViewportClient::AddViewportWidgetForPlayer(ULocalPlayer* Player, TSharedRef<SWidget> ViewportContent, const int32 ZOrder)
+{
+	TSharedPtr< IGameLayerManager > GameLayerManager(GameLayerManagerPtr.Pin());
+	if ( GameLayerManager.IsValid() )
+	{
+		GameLayerManager->AddWidgetForPlayer(Player, ViewportContent, ZOrder);
+	}
+}
+
+void UGameViewportClient::RemoveViewportWidgetForPlayer(ULocalPlayer* Player, TSharedRef<SWidget> ViewportContent)
+{
+	TSharedPtr< IGameLayerManager > GameLayerManager(GameLayerManagerPtr.Pin());
+	if ( GameLayerManager.IsValid() )
+	{
+		GameLayerManager->RemoveWidgetForPlayer(Player, ViewportContent);
+	}
+}
+
 void UGameViewportClient::RemoveAllViewportWidgets()
 {
 	CursorWidgets.Empty();
@@ -1925,10 +1963,13 @@ void UGameViewportClient::RemoveAllViewportWidgets()
 	TSharedPtr< SOverlay > PinnedViewportOverlayWidget( ViewportOverlayWidget.Pin() );
 	if( PinnedViewportOverlayWidget.IsValid() )
 	{
-		while(PinnedViewportOverlayWidget->GetNumWidgets())
-		{
-			PinnedViewportOverlayWidget->RemoveSlot();
-		}
+		PinnedViewportOverlayWidget->ClearChildren();
+	}
+
+	TSharedPtr< IGameLayerManager > GameLayerManager(GameLayerManagerPtr.Pin());
+	if ( GameLayerManager.IsValid() )
+	{
+		GameLayerManager->ClearWidgets();
 	}
 }
 
diff --git a/Engine/Source/Runtime/Engine/Private/Slate/SGameLayerManager.cpp b/Engine/Source/Runtime/Engine/Private/Slate/SGameLayerManager.cpp
new file mode 100644
--- /dev/null
+++ b/Engine/Source/Runtime/Engine/Private/Slate/SGameLayerManager.cpp
@@ -0,0 +1,260 @@
+// Copyright 1998-2015 Epic Games, Inc. All Rights Reserved.
+
+#include "EnginePrivate.h"
+
+#include "SGameLayerManager.h"
+
+#include "Engine/UserInterfaceSettings.h"
+
+#include "STooltipPresenter.h"
+#include "Widgets/Layout/SScissorRectBox.h"
+#include "Widgets/Layout/SDPIScaler.h"
+#include "Widgets/Layout/SPopup.h"
+
+/* SGameLayerManager interface
+ *****************************************************************************/
+
+SGameLayerManager::SGameLayerManager()
+{
+}
+
+void SGameLayerManager::Construct(const SGameLayerManager::FArguments& InArgs)
+{
+	SceneViewport = InArgs._SceneViewport;
+
+	TSharedRef<SDPIScaler> DPIScaler =
+		SNew(SDPIScaler)
+		.DPIScale(this, &SGameLayerManager::GetGameViewportDPIScale)
+		[
+			SNew(SOverlay)
+				
+			+ SOverlay::Slot()
+			[
+				SAssignNew(PlayerCanvas, SCanvas)
+			]
+
+			+ SOverlay::Slot()
+			[
+				InArgs._Content.Widget
+			]
+
+			+ SOverlay::Slot()
+			[
+				SNew(SPopup)
+				[
+					SAssignNew(TooltipPresenter, STooltipPresenter)
+				]
+			]
+		];
+
+	if ( InArgs._UseScissor )
+	{
+		ChildSlot
+		[
+			SNew(SScissorRectBox)
+			[
+				DPIScaler
+			]
+		];
+	}
+	else
+	{
+		ChildSlot
+		[
+			DPIScaler
+		];
+	}
+}
+
+void SGameLayerManager::NotifyPlayerAdded(int32 PlayerIndex, ULocalPlayer* AddedPlayer)
+{
+	UpdateLayout();
+}
+
+void SGameLayerManager::NotifyPlayerRemoved(int32 PlayerIndex, ULocalPlayer* RemovedPlayer)
+{
+	UpdateLayout();
+}
+
+void SGameLayerManager::AddWidgetForPlayer(ULocalPlayer* Player, TSharedRef<SWidget> ViewportContent, const int32 ZOrder)
+{
+	TSharedPtr<FPlayerLayer> PlayerLayer = FindOrCreatePlayerLayer(Player);
+	
+	// NOTE: Returns FSimpleSlot but we're ignoring here.  Could be used for alignment though.
+	PlayerLayer->Widget->AddSlot(ZOrder)
+	[
+		ViewportContent
+	];
+}
+
+void SGameLayerManager::RemoveWidgetForPlayer(ULocalPlayer* Player, TSharedRef<SWidget> ViewportContent)
+{
+	TSharedPtr<FPlayerLayer>* PlayerLayerPtr = PlayerLayers.Find(Player);
+	if ( PlayerLayerPtr )
+	{
+		TSharedPtr<FPlayerLayer> PlayerLayer = *PlayerLayerPtr;
+		PlayerLayer->Widget->RemoveSlot(ViewportContent);
+	}
+}
+
+void SGameLayerManager::ClearWidgetsForPlayer(ULocalPlayer* Player)
+{
+	TSharedPtr<FPlayerLayer>* PlayerLayerPtr = PlayerLayers.Find(Player);
+	if ( PlayerLayerPtr )
+	{
+		TSharedPtr<FPlayerLayer> PlayerLayer = *PlayerLayerPtr;
+		PlayerLayer->Widget->ClearChildren();
+	}
+}
+
+void SGameLayerManager::ClearWidgets()
+{
+	PlayerCanvas->ClearChildren();
+}
+
+void SGameLayerManager::Tick(const FGeometry& AllottedGeometry, const double InCurrentTime, const float InDeltaTime)
+{
+	CachedGeometry = AllottedGeometry;
+
+	UpdateLayout();
+}
+
+bool SGameLayerManager::OnVisualizeTooltip(const TSharedPtr<SWidget>& TooltipContent)
+{
+	TooltipPresenter->SetContent(TooltipContent.IsValid() ? TooltipContent.ToSharedRef() : SNullWidget::NullWidget);
+
+	return true;
+}
+
+float SGameLayerManager::GetGameViewportDPIScale() const
+{
+	if ( const FSceneViewport* Viewport = SceneViewport.Get() )
+	{
+		FIntPoint ViewportSize = Viewport->GetSize();
+		return GetDefault<UUserInterfaceSettings>(UUserInterfaceSettings::StaticClass())->GetDPIScaleBasedOnSize(ViewportSize);
+	}
+
+	return 1;
+}
+
+void SGameLayerManager::UpdateLayout()
+{
+	if ( const FSceneViewport* Viewport = SceneViewport.Get() )
+	{
+		if ( UWorld* World = Viewport->GetClient()->GetWorld() )
+		{
+			if ( World->IsGameWorld() == false )
+			{
+				PlayerLayers.Reset();
+				return;
+			}
+
+			if ( UGameViewportClient* ViewportClient = World->GetGameViewport() )
+			{
+				const TArray<ULocalPlayer*>& GamePlayers = GEngine->GetGamePlayers(World);
+
+				RemoveMissingPlayerLayers(GamePlayers);
+				AddOrUpdatePlayerLayers(CachedGeometry, ViewportClient, GamePlayers);
+			}
+		}
+	}
+}
+
+TSharedPtr<SGameLayerManager::FPlayerLayer> SGameLayerManager::FindOrCreatePlayerLayer(ULocalPlayer* LocalPlayer)
+{
+	TSharedPtr<FPlayerLayer>* PlayerLayerPtr = PlayerLayers.Find(LocalPlayer);
+	if ( PlayerLayerPtr == nullptr )
+	{
+		// Prevent any navigation outside of a player's layer once focus has been placed there.
+		TSharedRef<FNavigationMetaData> StopNavigation = MakeShareable(new FNavigationMetaData());
+		StopNavigation->SetNavigationStop(EUINavigation::Up);
+		StopNavigation->SetNavigationStop(EUINavigation::Down);
+		StopNavigation->SetNavigationStop(EUINavigation::Left);
+		StopNavigation->SetNavigationStop(EUINavigation::Right);
+		StopNavigation->SetNavigationStop(EUINavigation::Previous);
+		StopNavigation->SetNavigationStop(EUINavigation::Next);
+
+		// Create a new entry for the player
+		TSharedPtr<FPlayerLayer> NewLayer = MakeShareable(new FPlayerLayer());
+
+		// Create a new overlay widget to house any widgets we want to display for the player.
+		NewLayer->Widget = SNew(SOverlay)
+			.AddMetaData(StopNavigation);
+		
+		// Add the overlay to the player canvas, which we'll update every frame to match
+		// the dimensions of the player's split screen rect.
+		PlayerCanvas->AddSlot()
+			.HAlign(HAlign_Fill)
+			.VAlign(VAlign_Fill)
+			.Expose(NewLayer->Slot)
+			[
+				NewLayer->Widget.ToSharedRef()
+			];
+
+		PlayerLayerPtr = &PlayerLayers.Add(LocalPlayer, NewLayer);
+	}
+
+	return *PlayerLayerPtr;
+}
+
+void SGameLayerManager::RemoveMissingPlayerLayers(const TArray<ULocalPlayer*>& GamePlayers)
+{
+	TArray<ULocalPlayer*> ToRemove;
+
+	// Find the player layers for players that no longer exist
+	for ( TMap< ULocalPlayer*, TSharedPtr<FPlayerLayer> >::TIterator It(PlayerLayers); It; ++It )
+	{
+		ULocalPlayer* Key = ( *It ).Key;
+		if ( !GamePlayers.Contains(Key) )
+		{
+			ToRemove.Add(Key);
+		}
+	}
+
+	// Remove the missing players
+	for ( ULocalPlayer* Player : ToRemove )
+	{
+		RemovePlayerWidgets(Player);
+	}
+}
+
+void SGameLayerManager::RemovePlayerWidgets(ULocalPlayer* LocalPlayer)
+{
+	TSharedPtr<FPlayerLayer> Layer = PlayerLayers.FindRef(LocalPlayer);
+	PlayerCanvas->RemoveSlot(Layer->Widget.ToSharedRef());
+
+	PlayerLayers.Remove(LocalPlayer);
+}
+
+void SGameLayerManager::AddOrUpdatePlayerLayers(const FGeometry& AllottedGeometry, UGameViewportClient* ViewportClient, const TArray<ULocalPlayer*>& GamePlayers)
+{
+	ESplitScreenType::Type SplitType = ViewportClient->GetCurrentSplitscreenConfiguration();
+	TArray<struct FSplitscreenData>& SplitInfo = ViewportClient->SplitscreenInfo;
+
+	float InverseDPIScale = 1.0f / GetGameViewportDPIScale();
+
+	// Add and Update Player Layers
+	for ( int32 PlayerIndex = 0; PlayerIndex < GamePlayers.Num(); PlayerIndex++ )
+	{
+		ULocalPlayer* Player = GamePlayers[PlayerIndex];
+
+		if ( SplitType < SplitInfo.Num() && PlayerIndex < SplitInfo[SplitType].PlayerData.Num() )
+		{
+			TSharedPtr<FPlayerLayer> PlayerLayer = FindOrCreatePlayerLayer(Player);
+			FPerPlayerSplitscreenData& SplitData = SplitInfo[SplitType].PlayerData[PlayerIndex];
+
+			// Viewport Sizes
+			FVector2D Size, Position;
+			Size.X = SplitData.SizeX;
+			Size.Y = SplitData.SizeY;
+			Position.X = SplitData.OriginX;
+			Position.Y = SplitData.OriginY;
+
+			Size = Size * AllottedGeometry.GetLocalSize() * InverseDPIScale;
+			Position = Position * AllottedGeometry.GetLocalSize() * InverseDPIScale;
+
+			PlayerLayer->Slot->Size(Size);
+			PlayerLayer->Slot->Position(Position);
+		}
+	}
+}
\ No newline at end of file
diff --git a/Engine/Source/Runtime/Engine/Public/Slate/SGameLayerManager.h b/Engine/Source/Runtime/Engine/Public/Slate/SGameLayerManager.h
new file mode 100644
--- /dev/null
+++ b/Engine/Source/Runtime/Engine/Public/Slate/SGameLayerManager.h
@@ -0,0 +1,102 @@
+// Copyright 1998-2015 Epic Games, Inc. All Rights Reserved.
+
+#pragma once
+
+class FSceneViewport;
+class ULocalPlayer;
+
+/**
+ * Allows widgets to be managed for different users.
+ */
+class IGameLayerManager
+{
+public:
+	virtual void NotifyPlayerAdded(int32 PlayerIndex, ULocalPlayer* AddedPlayer) = 0;
+	virtual void NotifyPlayerRemoved(int32 PlayerIndex, ULocalPlayer* RemovedPlayer) = 0;
+
+	virtual void AddWidgetForPlayer(ULocalPlayer* Player, TSharedRef<SWidget> ViewportContent, const int32 ZOrder) = 0;
+	virtual void RemoveWidgetForPlayer(ULocalPlayer* Player, TSharedRef<SWidget> ViewportContent) = 0;
+	virtual void ClearWidgetsForPlayer(ULocalPlayer* Player) = 0;
+
+	virtual void ClearWidgets() = 0;
+};
+
+/**
+ * 
+ */
+class ENGINE_API SGameLayerManager : public SCompoundWidget, public IGameLayerManager
+{
+public:
+
+	SLATE_BEGIN_ARGS(SGameLayerManager)
+		: _UseScissor(true)
+	{
+		_Visibility = EVisibility::SelfHitTestInvisible;
+	}
+		/** Slot for this content (optional) */
+		SLATE_DEFAULT_SLOT(FArguments, Content)
+
+		SLATE_ATTRIBUTE(const FSceneViewport*, SceneViewport)
+		SLATE_ARGUMENT(bool, UseScissor)
+
+	SLATE_END_ARGS()
+
+	SGameLayerManager();
+
+	/**
+	 * Construct this widget
+	 *
+	 * @param	InArgs	The declaration data for this widget
+	 */
+	void Construct( const FArguments& InArgs );
+
+	// Begin IGameLayerManager
+	virtual void NotifyPlayerAdded(int32 PlayerIndex, ULocalPlayer* AddedPlayer);
+	virtual void NotifyPlayerRemoved(int32 PlayerIndex, ULocalPlayer* RemovedPlayer);
+
+	virtual void AddWidgetForPlayer(ULocalPlayer* Player, TSharedRef<SWidget> ViewportContent, const int32 ZOrder);
+	virtual void RemoveWidgetForPlayer(ULocalPlayer* Player, TSharedRef<SWidget> ViewportContent);
+	virtual void ClearWidgetsForPlayer(ULocalPlayer* Player);
+
+	virtual void ClearWidgets();
+	// Ened IGameLayerManager
+
+public:
+
+	// Begin SWidget overrides
+	virtual void Tick(const FGeometry& AllottedGeometry, const double InCurrentTime, const float InDeltaTime) override;
+	virtual bool OnVisualizeTooltip(const TSharedPtr<SWidget>& TooltipContent) override;
+	// End SWidget overrides
+
+private:
+	float GetGameViewportDPIScale() const;
+
+private:
+
+	struct FPlayerLayer : TSharedFromThis<FPlayerLayer>
+	{
+		TSharedPtr<SOverlay> Widget;
+		SCanvas::FSlot* Slot;
+
+		FPlayerLayer()
+			: Slot(nullptr)
+		{
+		}
+	};
+
+	void UpdateLayout();
+
+	TSharedPtr<FPlayerLayer> FindOrCreatePlayerLayer(ULocalPlayer* LocalPlayer);
+	void RemoveMissingPlayerLayers(const TArray<ULocalPlayer*>& GamePlayers);
+	void RemovePlayerWidgets(ULocalPlayer* LocalPlayer);
+	void AddOrUpdatePlayerLayers(const FGeometry& AllottedGeometry, UGameViewportClient* ViewportClient, const TArray<ULocalPlayer*>& GamePlayers);
+
+private:
+	FGeometry CachedGeometry;
+
+	TMap < ULocalPlayer*, TSharedPtr<FPlayerLayer> > PlayerLayers;
+
+	TAttribute<const FSceneViewport*> SceneViewport;
+	TSharedPtr<SCanvas> PlayerCanvas;
+	TSharedPtr<class STooltipPresenter> TooltipPresenter;
+};
diff --git a/Engine/Source/Runtime/UMG/Private/Components/WidgetComponent.cpp b/Engine/Source/Runtime/UMG/Private/Components/WidgetComponent.cpp
--- a/Engine/Source/Runtime/UMG/Private/Components/WidgetComponent.cpp
+++ b/Engine/Source/Runtime/UMG/Private/Components/WidgetComponent.cpp
@@ -10,6 +10,8 @@
 #include "DynamicMeshBuilder.h"
 #include "Scalability.h"
 
+DECLARE_CYCLE_STAT(TEXT("3DHitTesting"), STAT_Slate3DHitTesting, STATGROUP_Slate);
+
 class SVirtualWindow : public SWindow
 {
 	SLATE_BEGIN_ARGS( SVirtualWindow )
@@ -166,38 +168,37 @@ class FWidget3DHitTester : public ICustomHitTestPath
 public:
 	FWidget3DHitTester( UWorld* InWorld )
 		: World( InWorld )
+		, CachedFrame(-1)
 	{}
 
 	// ICustomHitTestPath implementation
 	virtual TArray<FWidgetAndPointer> GetBubblePathAndVirtualCursors(const FGeometry& InGeometry, FVector2D DesktopSpaceCoordinate, bool bIgnoreEnabledStatus) const override
 	{
+		SCOPE_CYCLE_COUNTER(STAT_Slate3DHitTesting);
+
 		if( World.IsValid() && ensure( World->IsGameWorld() ) )
 		{
 			ULocalPlayer* const TargetPlayer = GEngine->GetLocalPlayerFromControllerId(World.Get(), 0);
 			APlayerController* PlayerController = TargetPlayer->PlayerController;
 
-			if (TargetPlayer && PlayerController)
+			if ( UPrimitiveComponent* HitComponent = GetHitResultAtScreenPositionAndCache(TargetPlayer->PlayerController, InGeometry.AbsoluteToLocal(DesktopSpaceCoordinate)) )
 			{
-				FHitResult HitResult;
-				if (PlayerController->GetHitResultAtScreenPosition(InGeometry.AbsoluteToLocal(DesktopSpaceCoordinate), ECC_Visibility, true, HitResult))
+				if ( UWidgetComponent* WidgetComponent = Cast<UWidgetComponent>(HitComponent) )
 				{
-					UWidgetComponent* WidgetComponent = Cast<UWidgetComponent>(HitResult.Component.Get());
-					if ( WidgetComponent )
+					// Make sure the player is interacting with the front of the widget
+					// For widget components, the "front" faces the Z (or Up vector) direction
+					if ( FVector::DotProduct(WidgetComponent->GetUpVector(), CachedHitResult.ImpactPoint - CachedHitResult.TraceStart) < 0.f )
 					{
-						// Make sure the player is interacting with the front of the widget
-						// For widget components, the "front" faces the Z (or Up vector) direction
-						if ( FVector::DotProduct( WidgetComponent->GetUpVector(), HitResult.ImpactPoint - HitResult.TraceStart ) < 0.f )
+						// Make sure the player is close enough to the widget to interact with it
+						if ( FVector::DistSquared(CachedHitResult.TraceStart, CachedHitResult.ImpactPoint) <= FMath::Square(WidgetComponent->GetMaxInteractionDistance()) )
 						{
-							// Make sure the player is close enough to the widget to interact with it
-							if ( FVector::DistSquared( HitResult.TraceStart, HitResult.ImpactPoint ) <= FMath::Square( WidgetComponent->GetMaxInteractionDistance() ) )
-							{
-								return WidgetComponent->GetHitWidgetPath( HitResult, bIgnoreEnabledStatus );
-							}
+							return WidgetComponent->GetHitWidgetPath(CachedHitResult, bIgnoreEnabledStatus);
 						}
 					}
 				}
 			}
 		}
+
 		return TArray<FWidgetAndPointer>();
 	}
 
@@ -218,31 +219,29 @@ public:
 
 	virtual TSharedPtr<struct FVirtualPointerPosition> TranslateMouseCoordinateFor3DChild( const TSharedRef<SWidget>& ChildWidget, const FGeometry& ViewportGeometry, const FVector2D& ScreenSpaceMouseCoordinate, const FVector2D& LastScreenSpaceMouseCoordinate ) const override
 	{
-		if(World.IsValid() && ensure(World->IsGameWorld()))
+		if ( World.IsValid() && ensure(World->IsGameWorld()) )
 		{
 			ULocalPlayer* const TargetPlayer = GEngine->GetLocalPlayerFromControllerId(World.Get(), 0);
-			if(TargetPlayer && TargetPlayer->PlayerController)
+			if ( TargetPlayer && TargetPlayer->PlayerController )
 			{
 				// Check for a hit against any widget components in the world
-				for(TWeakObjectPtr<UWidgetComponent> Component : RegisteredComponents)
+				for ( TWeakObjectPtr<UWidgetComponent> Component : RegisteredComponents )
 				{
 					UWidgetComponent* WidgetComponent = Component.Get();
 					// Check if visible;
 					if ( WidgetComponent && WidgetComponent->GetSlateWidget() == ChildWidget )
 					{
 						FVector2D LocalMouseCoordinate = ViewportGeometry.AbsoluteToLocal(ScreenSpaceMouseCoordinate);
-						FVector2D LocalLastMouseCoordinate = ViewportGeometry.AbsoluteToLocal( LastScreenSpaceMouseCoordinate );
-
+						FVector2D LocalLastMouseCoordinate = ViewportGeometry.AbsoluteToLocal(LastScreenSpaceMouseCoordinate);
 
-						FHitResult HitResult;
-						if(TargetPlayer->PlayerController->GetHitResultAtScreenPosition(LocalMouseCoordinate, ECC_Visibility, true, HitResult))
+						if ( UPrimitiveComponent* HitComponent = GetHitResultAtScreenPositionAndCache(TargetPlayer->PlayerController, LocalMouseCoordinate) )
 						{
-							if ( WidgetComponent == HitResult.Component.Get() )
+							if ( WidgetComponent == HitComponent )
 							{
 								TSharedPtr<FVirtualPointerPosition> VirtualCursorPos = MakeShareable(new FVirtualPointerPosition);
 
 								FVector2D LocalHitLocation;
-								WidgetComponent->GetLocalHitLocation(HitResult, LocalHitLocation);
+								WidgetComponent->GetLocalHitLocation(CachedHitResult, LocalHitLocation);
 
 								VirtualCursorPos->CurrentCursorPosition = LocalHitLocation;
 								VirtualCursorPos->LastCursorPosition = LocalHitLocation;
@@ -259,6 +258,31 @@ public:
 	}
 	// End ICustomHitTestPath
 
+	UPrimitiveComponent* GetHitResultAtScreenPositionAndCache(APlayerController* PlayerController, FVector2D ScreenPosition) const
+	{
+		UPrimitiveComponent* HitComponent = nullptr;
+
+		if ( GFrameNumber != CachedFrame || CachedScreenPosition != ScreenPosition )
+		{
+			CachedFrame = GFrameNumber;
+			CachedScreenPosition = ScreenPosition;
+
+			if ( PlayerController )
+			{
+				if ( PlayerController->GetHitResultAtScreenPosition(ScreenPosition, ECC_Visibility, true, CachedHitResult) )
+				{
+					return CachedHitResult.Component.Get();
+				}
+			}
+		}
+		else
+		{
+			return CachedHitResult.Component.Get();
+		}
+
+		return nullptr;
+	}
+
 	void RegisterWidgetComponent( UWidgetComponent* InComponent )
 	{
 		RegisteredComponents.AddUnique( InComponent );
@@ -272,7 +296,12 @@ public:
 	uint32 GetNumRegisteredComponents() const { return RegisteredComponents.Num(); }
 	
 	UWorld* GetWorld() const { return World.Get(); }
+
 private:
+	mutable int64 CachedFrame;
+	mutable FVector2D CachedScreenPosition;
+	mutable FHitResult CachedHitResult;
+
 	TArray< TWeakObjectPtr<UWidgetComponent> > RegisteredComponents;
 	TWeakObjectPtr<UWorld> World;
 };
@@ -547,33 +576,23 @@ void UWidgetComponent::TickComponent(float DeltaTime, enum ELevelTick TickType,
 	{
 		if ( Widget && !Widget->IsDesignTime() )
 		{
-			ULocalPlayer* const TargetPlayer = GEngine->GetLocalPlayerFromControllerId(GetWorld(), 0);
-			APlayerController* PlayerController = TargetPlayer->PlayerController;
+			ULocalPlayer* TargetPlayer = OwnerPlayer ? OwnerPlayer : GEngine->GetLocalPlayerFromControllerId(GetWorld(), 0);
+			APlayerController* PlayerController = TargetPlayer ? TargetPlayer->PlayerController : nullptr;
 
 			if ( TargetPlayer && PlayerController && IsVisible() )
 			{
 				FVector WorldLocation = GetComponentLocation();
 
-				FVector2D ScreenLocation;
-				bool bProjected = PlayerController->ProjectWorldLocationToScreen(WorldLocation, ScreenLocation);
+				FVector2D ScreenPosition, ScreenPositionInvDPI;
+				const bool bProjected = UWidgetLayoutLibrary::ProjectWorldLocationToWidgetPosition(
+					PlayerController, WorldLocation, ScreenPosition, ScreenPositionInvDPI);
 
 				if ( bProjected )
 				{
 					Widget->SetVisibility(ESlateVisibility::SelfHitTestInvisible);
-						
-					// If the user has configured a resolution quality we need to multiply
-					// the pixels by the resolution quality to arrive at the true position in
-					// the viewport, as the rendered image will be stretched to fill whatever
-					// size the viewport is at.
-					Scalability::FQualityLevels ScalabilityQuality = Scalability::GetQualityLevels();
-					float QualityScale = ( ScalabilityQuality.ResolutionQuality / 100.0f );
-
-					FVector2D FinalScreenLocation = ScreenLocation / QualityScale;
-					FinalScreenLocation.X = FMath::RoundToInt(FinalScreenLocation.X);
-					FinalScreenLocation.Y = FMath::RoundToInt(FinalScreenLocation.Y);
-
+					
 					Widget->SetDesiredSizeInViewport(DrawSize);
-					Widget->SetPositionInViewport(FinalScreenLocation);
+					Widget->SetPositionInViewport(ScreenPositionInvDPI, false);
 					Widget->SetAlignmentInViewport(Pivot);
 				}
 				else
@@ -583,7 +602,8 @@ void UWidgetComponent::TickComponent(float DeltaTime, enum ELevelTick TickType,
 
 				if ( !Widget->IsInViewport() )
 				{
-					Widget->AddToViewport(ZOrder);
+					Widget->SetPlayerContext(TargetPlayer);
+					Widget->AddToPlayerScreen(ZOrder);
 				}
 			}
 			else if ( Widget->IsInViewport() )
@@ -708,6 +728,30 @@ void UWidgetComponent::InitWidget()
 	}
 }
 
+void UWidgetComponent::SetOwnerPlayer(ULocalPlayer* LocalPlayer)
+{
+	OwnerPlayer = LocalPlayer;
+}
+
+ULocalPlayer* UWidgetComponent::GetOwnerPlayer() const
+{
+	return OwnerPlayer;
+}
+
+void UWidgetComponent::SetWidget(UUserWidget* InWidget)
+{
+	if ( Widget )
+	{
+		Widget->RemoveFromParent();
+		Widget->MarkPendingKill();
+		Widget = nullptr;
+	}
+
+	Widget = InWidget;
+
+	UpdateWidget();
+}
+
 void UWidgetComponent::UpdateWidget()
 {
 	// Don't do any work if Slate is not initialized
diff --git a/Engine/Source/Runtime/UMG/Private/UserWidget.cpp b/Engine/Source/Runtime/UMG/Private/UserWidget.cpp
--- a/Engine/Source/Runtime/UMG/Private/UserWidget.cpp
+++ b/Engine/Source/Runtime/UMG/Private/UserWidget.cpp
@@ -616,6 +616,23 @@ UWidget* UUserWidget::GetRootWidget() const
 }
 
 void UUserWidget::AddToViewport(int32 ZOrder)
+{
+	AddToScreen(nullptr, ZOrder);
+}
+
+bool UUserWidget::AddToPlayerScreen(int32 ZOrder)
+{
+	if ( ULocalPlayer* LocalPlayer = GetOwningLocalPlayer() )
+	{
+		AddToScreen(LocalPlayer, ZOrder);
+		return true;
+	}
+
+	FMessageLog("PIE").Error(LOCTEXT("AddToPlayerScreen_NoPlayer", "AddToPlayerScreen Failed.  No Owning Player!"));
+	return false;
+}
+
+void UUserWidget::AddToScreen(ULocalPlayer* Player, int32 ZOrder)
 {
 	if ( !FullScreenWidget.IsValid() )
 	{
@@ -630,7 +647,15 @@ void UUserWidget::AddToViewport(int32 ZOrder)
 		{
 			if ( UGameViewportClient* ViewportClient = World->GetGameViewport() )
 			{
-				ViewportClient->AddViewportWidgetContent(RootWidget, 10 + ZOrder);
+				ZOrder = 10 + ZOrder;
+				if ( Player )
+				{
+					ViewportClient->AddViewportWidgetForPlayer(Player, RootWidget, ZOrder);
+				}
+				else
+				{
+					ViewportClient->AddViewportWidgetContent(RootWidget, ZOrder);
+				}
 			}
 		}
 	}
@@ -653,7 +678,14 @@ void UUserWidget::RemoveFromParent()
 		{
 			if ( UGameViewportClient* ViewportClient = World->GetGameViewport() )
 			{
-				ViewportClient->RemoveViewportWidgetContent(WidgetHost.ToSharedRef());
+				TSharedRef<SWidget> WidgetHostRef = WidgetHost.ToSharedRef();
+
+				ViewportClient->RemoveViewportWidgetContent(WidgetHostRef);
+
+				if ( ULocalPlayer* LocalPlayer = GetOwningLocalPlayer() )
+				{
+					ViewportClient->RemoveViewportWidgetForPlayer(LocalPlayer, WidgetHostRef);
+				}
 			}
 		}
 	}
@@ -712,12 +744,20 @@ class APawn* UUserWidget::GetOwningPlayerPawn() const
 	return nullptr;
 }
 
-void UUserWidget::SetPositionInViewport(FVector2D Position)
+void UUserWidget::SetPositionInViewport(FVector2D Position, bool bRemoveDPIScale )
 {
-	float Scale = UWidgetLayoutLibrary::GetViewportScale(this);
+	if ( bRemoveDPIScale )
+	{
+		float Scale = UWidgetLayoutLibrary::GetViewportScale(this);
 
-	ViewportOffsets.Left = Position.X / Scale;
-	ViewportOffsets.Top = Position.Y / Scale;
+		ViewportOffsets.Left = Position.X / Scale;
+		ViewportOffsets.Top = Position.Y / Scale;
+	}
+	else
+	{
+		ViewportOffsets.Left = Position.X;
+		ViewportOffsets.Top = Position.Y;
+	}
 
 	ViewportAnchors = FAnchors(0, 0);
 }
diff --git a/Engine/Source/Runtime/UMG/Private/WidgetLayoutLibrary.cpp b/Engine/Source/Runtime/UMG/Private/WidgetLayoutLibrary.cpp
--- a/Engine/Source/Runtime/UMG/Private/WidgetLayoutLibrary.cpp
+++ b/Engine/Source/Runtime/UMG/Private/WidgetLayoutLibrary.cpp
@@ -16,6 +16,39 @@ UWidgetLayoutLibrary::UWidgetLayoutLibrary(const FObjectInitializer& ObjectIniti
 {
 }
 
+bool UWidgetLayoutLibrary::ProjectWorldLocationToWidgetPosition(APlayerController* PlayerController, FVector WorldLocation, FVector2D& ScreenPosition, FVector2D& ScreenPositionInvDPI)
+{
+	if ( PlayerController )
+	{
+		FVector2D PixelLocation;
+		const bool bProjected = PlayerController->ProjectWorldLocationToScreen(WorldLocation, PixelLocation);
+
+		if ( bProjected )
+		{
+			// If the user has configured a resolution quality we need to multiply
+			// the pixels by the resolution quality to arrive at the true position in
+			// the viewport, as the rendered image will be stretched to fill whatever
+			// size the viewport is at.
+			Scalability::FQualityLevels ScalabilityQuality = Scalability::GetQualityLevels();
+			const float QualityScale = ( ScalabilityQuality.ResolutionQuality / 100.0f );
+
+			// Remove the resolution quality scale.
+			ScreenPosition = PixelLocation / QualityScale;
+
+			// Get the application / DPI scale
+			const float Scale = UWidgetLayoutLibrary::GetViewportScale(PlayerController);
+
+			// Apply inverse DPI scale so that the widget ends up in the expected position
+			ScreenPositionInvDPI.X = ScreenPosition.X / Scale;
+			ScreenPositionInvDPI.Y = ScreenPosition.Y / Scale;
+
+			return true;
+		}
+	}
+
+	return false;
+}
+
 float UWidgetLayoutLibrary::GetViewportScale(UObject* WorldContextObject)
 {
 	UWorld* World = GEngine->GetWorldFromContextObject(WorldContextObject);
@@ -92,4 +125,16 @@ UVerticalBoxSlot* UWidgetLayoutLibrary::SlotAsVerticalBoxSlot(UWidget* ChildWidg
 	return Cast<UVerticalBoxSlot>(ChildWidget->Slot);
 }
 
+void UWidgetLayoutLibrary::RemoveAllWidgets(UObject* WorldContextObject)
+{
+	UWorld* World = GEngine->GetWorldFromContextObject(WorldContextObject);
+	if ( World && World->IsGameWorld() )
+	{
+		if ( UGameViewportClient* ViewportClient = World->GetGameViewport() )
+		{
+			ViewportClient->RemoveAllViewportWidgets();
+		}
+	}
+}
+
 #undef LOCTEXT_NAMESPACE
\ No newline at end of file
diff --git a/Engine/Source/Runtime/UMG/Public/Blueprint/UserWidget.h b/Engine/Source/Runtime/UMG/Public/Blueprint/UserWidget.h
--- a/Engine/Source/Runtime/UMG/Public/Blueprint/UserWidget.h
+++ b/Engine/Source/Runtime/UMG/Public/Blueprint/UserWidget.h
@@ -160,11 +160,20 @@ public:
 	 * Adds it to the game's viewport and fills the entire screen, unless SetDesiredSizeInViewport is called
 	 * to explicitly set the size.
 	 *
-	 * @param ZOrder The higher the number the greater the priority when determining if it will above other widgets.
+	 * @param ZOrder The higher the number, the more on top this widget will be.
 	 */
 	UFUNCTION(BlueprintCallable, BlueprintCosmetic, Category="User Interface|Viewport", meta=( AdvancedDisplay = "ZOrder" ))
 	void AddToViewport(int32 ZOrder = 0);
 
+	/**
+	 * Adds the widget to the game's viewport in a section dedicated to the player.  This is valuable in a split screen
+	 * game where you need to only show a widget over a player's portion of the viewport.
+	 *
+	 * @param ZOrder The higher the number, the more on top this widget will be.
+	 */
+	UFUNCTION(BlueprintCallable, BlueprintCosmetic, Category="User Interface|Viewport", meta=( AdvancedDisplay = "ZOrder" ))
+	bool AddToPlayerScreen(int32 ZOrder = 0);
+
 	/**
 	 * Removes the widget from the viewport.
 	 */
@@ -176,9 +185,13 @@ public:
 
 	/**
 	 * Sets the widgets position in the viewport.
+	 * @param Position The 2D position to set the widget to in the viewport.
+	 * @param bRemoveDPIScale If you've already calculated inverse DPI, set this to false.  
+	 * Otherwise inverse DPI is applied to the position so that when the location is scaled 
+	 * by DPI, it ends up in the expected position.
 	 */
 	UFUNCTION(BlueprintCallable, BlueprintCosmetic, Category="User Interface|Viewport")
-	void SetPositionInViewport(FVector2D Position);
+	void SetPositionInViewport(FVector2D Position, bool bRemoveDPIScale = true);
 
 	/*  */
 	UFUNCTION(BlueprintCallable, BlueprintCosmetic, Category="User Interface|Viewport")
@@ -211,7 +224,7 @@ public:
 	 * @return The owning local player.
 	 */
 	UFUNCTION(BlueprintCallable, BlueprintCosmetic, Category="Player")
-	class ULocalPlayer* GetOwningLocalPlayer() const;
+	ULocalPlayer* GetOwningLocalPlayer() const;
 
 	/**
 	 * Sets the local player associated with this UI.
@@ -685,6 +698,9 @@ public:
 #endif
 
 protected:
+	/** Adds the widget to the screen, either to the viewport or to the player's screen depending on if the LocalPlayer is null. */
+	void AddToScreen(ULocalPlayer* LocalPlayer, int32 ZOrder);
+
 	virtual TSharedRef<SWidget> RebuildWidget() override;
 	virtual void OnWidgetRebuilt() override;
 
diff --git a/Engine/Source/Runtime/UMG/Public/Blueprint/WidgetLayoutLibrary.h b/Engine/Source/Runtime/UMG/Public/Blueprint/WidgetLayoutLibrary.h
--- a/Engine/Source/Runtime/UMG/Public/Blueprint/WidgetLayoutLibrary.h
+++ b/Engine/Source/Runtime/UMG/Public/Blueprint/WidgetLayoutLibrary.h
@@ -11,6 +11,18 @@ class UWidgetLayoutLibrary : public UBlueprintFunctionLibrary
 
 public:
 
+	/**
+	 * Gets the projected world to screen position for a player, then converts it into a widget
+	 * position, which takes into account any quality scaling.
+	 * @param PlayerController The player controller to project the position in the world to their screen.
+	 * @param WorldLocation The world location to project from.
+	 * @param ScreenPosition The position in the viewport without the inverse DPI being taken into account.
+	 * @param ScreenPositionInvDPI The psotion in the viewport with the inverse DPI applied to it.  This is useful for positioning child widgets of widgets in the viewport where you need to remove DPI before applying it.
+	 * @return true if the position projects onto the screen.
+	 */
+	UFUNCTION(BlueprintPure, BlueprintCosmetic, Category="Viewport")
+	static bool ProjectWorldLocationToWidgetPosition(APlayerController* PlayerController, FVector WorldLocation, FVector2D& ScreenPosition, FVector2D& ScreenPositionInvDPI);
+
 	/**
 	 * Gets the current DPI Scale being applied to the viewport and all the Widgets.
 	 */
@@ -73,4 +85,10 @@ public:
 	 */
 	UFUNCTION(BlueprintPure, Category = "Slot")
 	static UVerticalBoxSlot* SlotAsVerticalBoxSlot(UWidget* Widget);
+
+	/**
+	 * Removes all widgets from the viewport.
+	 */
+	UFUNCTION(BlueprintCallable, BlueprintCosmetic, Category="Viewport", meta=( HidePin="WorldContextObject", DefaultToSelf="WorldContextObject" ))
+	static void RemoveAllWidgets(UObject* WorldContextObject);
 };
diff --git a/Engine/Source/Runtime/UMG/Public/Components/WidgetComponent.h b/Engine/Source/Runtime/UMG/Public/Components/WidgetComponent.h
--- a/Engine/Source/Runtime/UMG/Public/Components/WidgetComponent.h
+++ b/Engine/Source/Runtime/UMG/Public/Components/WidgetComponent.h
@@ -9,14 +9,16 @@ struct FVirtualPointerPosition;
 UENUM()
 enum class EWidgetSpace : uint8
 {
+	// The widget is rendered in the world as mesh, it can be occluded like any other mesh in the world.
 	World,
+	// The widget is rendered in the screen, completely outside of the world, never occluded.
 	Screen
 };
 
 /**
-* Beware! This feature is experimental and may be substantially changed or removed in future releases.
-* A 3D instance of a Widget Blueprint that can be interacted with in the world.
-*/
+ * Beware! This feature is experimental and may be substantially changed or removed in future releases.
+ * A 3D instance of a Widget Blueprint that can be interacted with in the world.
+ */
 UCLASS(ClassGroup=Experimental, hidecategories=(Object,Activation,"Components|Activation",Sockets,Base,Lighting,LOD,Mesh), editinlinenew, meta=(BlueprintSpawnableComponent,  DevelopmentStatus=Experimental) )
 class UMG_API UWidgetComponent : public UPrimitiveComponent
 {
@@ -87,6 +89,18 @@ public:
 	/** @return The window containing the user widget content */
 	TSharedPtr<SWidget> GetSlateWidget() const;
 
+	/** Sets the widget to use directly. */
+	UFUNCTION(BlueprintCallable, Category=UI)
+	void SetWidget(UUserWidget* Widget);
+
+	/** Sets the local player that owns this widget component. */
+	UFUNCTION(BlueprintCallable, Category=UI)
+	void SetOwnerPlayer(ULocalPlayer* LocalPlayer);
+
+	/** Gets the local player that owns this widget component. */
+	UFUNCTION(BlueprintCallable, Category=UI)
+	ULocalPlayer* GetOwnerPlayer() const;
+
 	/** @return The draw size of the quad in the world */
 	UFUNCTION(BlueprintCallable, Category=UI)
 	FVector2D GetDrawSize() const;
@@ -134,6 +148,13 @@ private:
 	UPROPERTY(EditAnywhere, Category=UI, meta=(ClampMin="0.0", UIMax="5000.0", ClampMax="100000.0"))
 	float MaxInteractionDistance;
 
+	/**
+	 * The owner player for a widget component, if this widget is drawn on the screen, this controls
+	 * what player's screen it appears on for split screen, if not set, users player 0.
+	 */
+	UPROPERTY()
+	ULocalPlayer* OwnerPlayer;
+
 	/** The background color of the component */
 	UPROPERTY(EditAnywhere, Category=Rendering)
 	FLinearColor BackgroundColor;
@@ -196,6 +217,3 @@ private:
 	/** The hit tester to use for this component */
 	static TSharedPtr<class FWidget3DHitTester> WidgetHitTester;
 };
-
- 
-
