Fix for missing include in WidgetComponent for the UWidgetLayoutLibrary usage.

diff --git a/Engine/Source/Runtime/UMG/Private/Components/WidgetComponent.cpp b/Engine/Source/Runtime/UMG/Private/Components/WidgetComponent.cpp
--- a/Engine/Source/Runtime/UMG/Private/Components/WidgetComponent.cpp
+++ b/Engine/Source/Runtime/UMG/Private/Components/WidgetComponent.cpp
@@ -9,6 +9,7 @@
 #endif // !UE_SERVER
 #include "DynamicMeshBuilder.h"
 #include "Scalability.h"
+#include "WidgetLayoutLibrary.h"
 
 DECLARE_CYCLE_STAT(TEXT("3DHitTesting"), STAT_Slate3DHitTesting, STATGROUP_Slate);
 
