Missing chages for 2481343.

diff --git a/Engine/Source/Runtime/Engine/Classes/Components/SkyLightComponent.h b/Engine/Source/Runtime/Engine/Classes/Components/SkyLightComponent.h
--- a/Engine/Source/Runtime/Engine/Classes/Components/SkyLightComponent.h
+++ b/Engine/Source/Runtime/Engine/Classes/Components/SkyLightComponent.h
@@ -85,9 +85,7 @@ enum ESkyLightSourceType
 UCLASS(ClassGroup=Lights, HideCategories=(Trigger,Activation,"Components|Activation",Physics), meta=(BlueprintSpawnableComponent))
 class ENGINE_API USkyLightComponent : public ULightComponentBase
 {
-	GENERATED_BODY()
-public:
-	USkyLightComponent(const FObjectInitializer& ObjectInitializer = FObjectInitializer::Get());
+	GENERATED_UCLASS_BODY()
 
 	/** Indicates where to get the light contribution from. */
 	UPROPERTY(EditAnywhere, BlueprintReadOnly, Category=Light)
diff --git a/Engine/Source/Runtime/Engine/Classes/Components/SphereComponent.h b/Engine/Source/Runtime/Engine/Classes/Components/SphereComponent.h
--- a/Engine/Source/Runtime/Engine/Classes/Components/SphereComponent.h
+++ b/Engine/Source/Runtime/Engine/Classes/Components/SphereComponent.h
@@ -11,9 +11,7 @@
 UCLASS(ClassGroup="Collision", editinlinenew, hidecategories=(Object,LOD,Lighting,TextureStreaming), meta=(DisplayName="Sphere Collision", BlueprintSpawnableComponent))
 class ENGINE_API USphereComponent : public UShapeComponent
 {
-	GENERATED_BODY()
-public:
-	USphereComponent(const FObjectInitializer& ObjectInitializer = FObjectInitializer::Get());
+	GENERATED_UCLASS_BODY()
 
 protected:
 	/** The radius of the sphere **/
diff --git a/Engine/Source/Runtime/Engine/Classes/Components/WindDirectionalSourceComponent.h b/Engine/Source/Runtime/Engine/Classes/Components/WindDirectionalSourceComponent.h
--- a/Engine/Source/Runtime/Engine/Classes/Components/WindDirectionalSourceComponent.h
+++ b/Engine/Source/Runtime/Engine/Classes/Components/WindDirectionalSourceComponent.h
@@ -7,9 +7,7 @@
 UCLASS(collapsecategories, hidecategories=(Object, Mobility), editinlinenew)
 class UWindDirectionalSourceComponent : public USceneComponent
 {
-	GENERATED_BODY()
-public:
-	UWindDirectionalSourceComponent(const FObjectInitializer& ObjectInitializer = FObjectInitializer::Get());
+	GENERATED_UCLASS_BODY()
 
 	UPROPERTY(interp, Category=WindDirectionalSourceComponent)
 	float Strength;
diff --git a/Engine/Source/Runtime/Engine/Classes/Debug/LogVisualizerHUD.h b/Engine/Source/Runtime/Engine/Classes/Debug/LogVisualizerHUD.h
--- a/Engine/Source/Runtime/Engine/Classes/Debug/LogVisualizerHUD.h
+++ b/Engine/Source/Runtime/Engine/Classes/Debug/LogVisualizerHUD.h
@@ -7,9 +7,7 @@
 UCLASS()
 class ALogVisualizerHUD : public ADebugCameraHUD
 {
-	GENERATED_BODY()
-public:
-	ALogVisualizerHUD(const FObjectInitializer& ObjectInitializer = FObjectInitializer::Get());
+	GENERATED_UCLASS_BODY()
 
 	FFontRenderInfo TextRenderInfo;
 
diff --git a/Engine/Source/Runtime/Engine/Classes/DeviceProfiles/DeviceProfileManager.h b/Engine/Source/Runtime/Engine/Classes/DeviceProfiles/DeviceProfileManager.h
--- a/Engine/Source/Runtime/Engine/Classes/DeviceProfiles/DeviceProfileManager.h
+++ b/Engine/Source/Runtime/Engine/Classes/DeviceProfiles/DeviceProfileManager.h
@@ -19,11 +19,9 @@ DECLARE_MULTICAST_DELEGATE( FOnManagerUpdated );
  UCLASS( config=DeviceProfiles, transient )
 class ENGINE_API UDeviceProfileManager : public UObject
 {
-	GENERATED_BODY()
-public:
-	UDeviceProfileManager(const FObjectInitializer& ObjectInitializer = FObjectInitializer::Get());
 public:
 
+	GENERATED_UCLASS_BODY()
 
 	/**
 	 * Startup and select the active device profile
@@ -130,4 +128,4 @@ private:
 	static FString DeviceProfileFileName;
 
 	int32 RenameIndex;
-};
+};
\ No newline at end of file
diff --git a/Engine/Source/Runtime/Engine/Classes/Distributions/Distribution.h b/Engine/Source/Runtime/Engine/Classes/Distributions/Distribution.h
--- a/Engine/Source/Runtime/Engine/Classes/Distributions/Distribution.h
+++ b/Engine/Source/Runtime/Engine/Classes/Distributions/Distribution.h
@@ -58,9 +58,7 @@ struct FRawDistribution
 UCLASS(DefaultToInstanced, collapsecategories, hidecategories=Object, editinlinenew, abstract,MinimalAPI)
 class UDistribution : public UObject, public FCurveEdInterface
 {
-	GENERATED_BODY()
-public:
-	ENGINE_API UDistribution(const FObjectInitializer& ObjectInitializer = FObjectInitializer::Get());
+	GENERATED_UCLASS_BODY()
 
 	/** Default value for initializing and checking correct values on UDistributions. */
 	static const float DefaultValue;
diff --git a/Engine/Source/Runtime/Engine/Classes/Distributions/DistributionFloat.h b/Engine/Source/Runtime/Engine/Classes/Distributions/DistributionFloat.h
--- a/Engine/Source/Runtime/Engine/Classes/Distributions/DistributionFloat.h
+++ b/Engine/Source/Runtime/Engine/Classes/Distributions/DistributionFloat.h
@@ -20,7 +20,7 @@ struct FFloatDistribution
 USTRUCT()
 struct FRawDistributionFloat : public FRawDistribution
 {
-	GENERATED_BODY()
+	GENERATED_USTRUCT_BODY()
 
 private:
 	UPROPERTY()
@@ -71,11 +71,10 @@ public:
 	
 };
 
-UCLASS(abstract,MinimalAPI)
+UCLASS(abstract, customconstructor,MinimalAPI)
 class UDistributionFloat : public UDistribution
 {
-	GENERATED_BODY()
-public:
+	GENERATED_UCLASS_BODY()
 
 	/** Can this variable be baked out to a FRawDistribution? Should be true 99% of the time*/
 	UPROPERTY(EditAnywhere, Category=Baked)
diff --git a/Engine/Source/Runtime/Engine/Classes/Distributions/DistributionVectorParameterBase.h b/Engine/Source/Runtime/Engine/Classes/Distributions/DistributionVectorParameterBase.h
--- a/Engine/Source/Runtime/Engine/Classes/Distributions/DistributionVectorParameterBase.h
+++ b/Engine/Source/Runtime/Engine/Classes/Distributions/DistributionVectorParameterBase.h
@@ -8,9 +8,7 @@
 UCLASS(abstract, collapsecategories, hidecategories=Object, editinlinenew)
 class UDistributionVectorParameterBase : public UDistributionVectorConstant
 {
-	GENERATED_BODY()
-public:
-	UDistributionVectorParameterBase(const FObjectInitializer& ObjectInitializer = FObjectInitializer::Get());
+	GENERATED_UCLASS_BODY()
 
 	UPROPERTY(EditAnywhere, Category=DistributionVectorParameterBase)
 	FName ParameterName;
diff --git a/Engine/Source/Runtime/Engine/Classes/Engine/BlendableInterface.h b/Engine/Source/Runtime/Engine/Classes/Engine/BlendableInterface.h
--- a/Engine/Source/Runtime/Engine/Classes/Engine/BlendableInterface.h
+++ b/Engine/Source/Runtime/Engine/Classes/Engine/BlendableInterface.h
@@ -31,9 +31,7 @@ enum EBlendableLocation
 UINTERFACE()
 class UBlendableInterface : public UInterface
 {
-	GENERATED_BODY()
-public:
-	UBlendableInterface(const FObjectInitializer& ObjectInitializer = FObjectInitializer::Get());
+	GENERATED_UINTERFACE_BODY()
 };
 
 /**
@@ -41,8 +39,7 @@ public:
  */
 class IBlendableInterface
 {
-	GENERATED_BODY()
-public:
+	GENERATED_IINTERFACE_BODY()
 		
 	// @param Weight 0..1, 1=fully take the values from this object, crash if outside the valid range
 	virtual void OverrideBlendableSettings(class FSceneView& View, float Weight) const = 0;
diff --git a/Engine/Source/Runtime/Engine/Classes/Engine/CollisionProfile.h b/Engine/Source/Runtime/Engine/Classes/Engine/CollisionProfile.h
--- a/Engine/Source/Runtime/Engine/Classes/Engine/CollisionProfile.h
+++ b/Engine/Source/Runtime/Engine/Classes/Engine/CollisionProfile.h
@@ -13,7 +13,7 @@
 USTRUCT(BlueprintType)
 struct FCollisionProfileName
 {
-	GENERATED_BODY()
+	GENERATED_USTRUCT_BODY()
 
 	FCollisionProfileName()
 		: Name(NAME_None)
@@ -30,7 +30,7 @@ struct FCollisionProfileName
 USTRUCT()
 struct ENGINE_API FCollisionResponseTemplate
 {
-	GENERATED_BODY()
+	GENERATED_USTRUCT_BODY()
 
 	UPROPERTY()
 	FName Name;
@@ -77,7 +77,7 @@ struct ENGINE_API FCollisionResponseTemplate
 USTRUCT()
 struct FCustomChannelSetup
 {
-	GENERATED_BODY()
+	GENERATED_USTRUCT_BODY()
 
 	/** Which channel you'd like to customize **/
 	UPROPERTY()
@@ -122,7 +122,7 @@ struct FCustomChannelSetup
 USTRUCT()
 struct ENGINE_API FCustomProfile
 {
-	GENERATED_BODY()
+	GENERATED_USTRUCT_BODY()
 
 	UPROPERTY()
 	FName Name;
@@ -140,9 +140,7 @@ UCLASS(abstract, config=Engine, defaultconfig, MinimalAPI)
 class UCollisionProfile
 	: public UObject
 {
-	GENERATED_BODY()
-public:
-	ENGINE_API UCollisionProfile(const FObjectInitializer& ObjectInitializer = FObjectInitializer::Get());
+	GENERATED_UCLASS_BODY()
 
 private:
 
diff --git a/Engine/Source/Runtime/Engine/Classes/Engine/Console.h b/Engine/Source/Runtime/Engine/Classes/Engine/Console.h
--- a/Engine/Source/Runtime/Engine/Classes/Engine/Console.h
+++ b/Engine/Source/Runtime/Engine/Classes/Engine/Console.h
@@ -12,7 +12,7 @@
 USTRUCT()
 struct FAutoCompleteNode
 {
-	GENERATED_BODY()
+	GENERATED_USTRUCT_BODY()
 
 	/** Char for node in the tree */
 	UPROPERTY()
@@ -55,9 +55,7 @@ class ENGINE_API UConsole
 	: public UObject
 	, public FOutputDevice
 {
-	GENERATED_BODY()
-public:
-	UConsole(const FObjectInitializer& ObjectInitializer = FObjectInitializer::Get());
+	GENERATED_UCLASS_BODY()
 
 	/** The player which the next console command should be executed in the context of.  If nullptr, execute in the viewport. */
 	UPROPERTY()
diff --git a/Engine/Source/Runtime/Engine/Classes/Engine/GameViewportClient.h b/Engine/Source/Runtime/Engine/Classes/Engine/GameViewportClient.h
--- a/Engine/Source/Runtime/Engine/Classes/Engine/GameViewportClient.h
+++ b/Engine/Source/Runtime/Engine/Classes/Engine/GameViewportClient.h
@@ -51,9 +51,7 @@ enum EStereoscopicPass
 UCLASS(Within=Engine, transient, config=Engine)
 class ENGINE_API UGameViewportClient : public UScriptViewportClient, public FExec
 {
-	GENERATED_BODY()
-public:
-	UGameViewportClient(const FObjectInitializer& ObjectInitializer = FObjectInitializer::Get());
+	GENERATED_UCLASS_BODY()
 
 public:
 	virtual ~UGameViewportClient();
diff --git a/Engine/Source/Runtime/Engine/Classes/Engine/InputAxisDelegateBinding.h b/Engine/Source/Runtime/Engine/Classes/Engine/InputAxisDelegateBinding.h
--- a/Engine/Source/Runtime/Engine/Classes/Engine/InputAxisDelegateBinding.h
+++ b/Engine/Source/Runtime/Engine/Classes/Engine/InputAxisDelegateBinding.h
@@ -7,7 +7,7 @@
 USTRUCT()
 struct ENGINE_API FBlueprintInputAxisDelegateBinding : public FBlueprintInputDelegateBinding
 {
-	GENERATED_BODY()
+	GENERATED_USTRUCT_BODY()
 
 	UPROPERTY()
 	FName InputAxisName;
@@ -26,9 +26,7 @@ struct ENGINE_API FBlueprintInputAxisDelegateBinding : public FBlueprintInputDel
 UCLASS()
 class ENGINE_API UInputAxisDelegateBinding : public UInputDelegateBinding
 {
-	GENERATED_BODY()
-public:
-	UInputAxisDelegateBinding(const FObjectInitializer& ObjectInitializer = FObjectInitializer::Get());
+	GENERATED_UCLASS_BODY()
 
 	UPROPERTY()
 	TArray<FBlueprintInputAxisDelegateBinding> InputAxisDelegateBindings;
diff --git a/Engine/Source/Runtime/Engine/Classes/Engine/LevelStreamingAlwaysLoaded.h b/Engine/Source/Runtime/Engine/Classes/Engine/LevelStreamingAlwaysLoaded.h
--- a/Engine/Source/Runtime/Engine/Classes/Engine/LevelStreamingAlwaysLoaded.h
+++ b/Engine/Source/Runtime/Engine/Classes/Engine/LevelStreamingAlwaysLoaded.h
@@ -13,9 +13,7 @@
 UCLASS(MinimalAPI)
 class ULevelStreamingAlwaysLoaded : public ULevelStreaming
 {
-	GENERATED_BODY()
-public:
-	ENGINE_API ULevelStreamingAlwaysLoaded(const FObjectInitializer& ObjectInitializer = FObjectInitializer::Get());
+	GENERATED_UCLASS_BODY()
 
 	// Begin ULevelStreaming Interface
 	virtual bool ShouldBeLoaded() const override;
diff --git a/Engine/Source/Runtime/Engine/Classes/Engine/NiagaraSpriteRendererProperties.h b/Engine/Source/Runtime/Engine/Classes/Engine/NiagaraSpriteRendererProperties.h
--- a/Engine/Source/Runtime/Engine/Classes/Engine/NiagaraSpriteRendererProperties.h
+++ b/Engine/Source/Runtime/Engine/Classes/Engine/NiagaraSpriteRendererProperties.h
@@ -8,10 +8,8 @@
 UCLASS()
 class UNiagaraSpriteRendererProperties : public UNiagaraEffectRendererProperties
 {
-	GENERATED_BODY()
-public:
-	UNiagaraSpriteRendererProperties(const FObjectInitializer& ObjectInitializer);
 public:
+	GENERATED_UCLASS_BODY()
 
 	UNiagaraSpriteRendererProperties()
 	{
diff --git a/Engine/Source/Runtime/Engine/Classes/Engine/RendererSettings.h b/Engine/Source/Runtime/Engine/Classes/Engine/RendererSettings.h
--- a/Engine/Source/Runtime/Engine/Classes/Engine/RendererSettings.h
+++ b/Engine/Source/Runtime/Engine/Classes/Engine/RendererSettings.h
@@ -91,9 +91,7 @@ UCLASS(config=Engine, defaultconfig)
 class ENGINE_API URendererSettings
 	: public UObject
 {
-	GENERATED_BODY()
-public:
-	URendererSettings(const FObjectInitializer& ObjectInitializer = FObjectInitializer::Get());
+	GENERATED_UCLASS_BODY()
 
 	UPROPERTY(config, EditAnywhere, Category=Mobile, meta=(
 		ConsoleVariable="r.MobileHDR",DisplayName="Mobile HDR",
diff --git a/Engine/Source/Runtime/Engine/Classes/Engine/Scene.h b/Engine/Source/Runtime/Engine/Classes/Engine/Scene.h
--- a/Engine/Source/Runtime/Engine/Classes/Engine/Scene.h
+++ b/Engine/Source/Runtime/Engine/Classes/Engine/Scene.h
@@ -34,7 +34,7 @@ enum EAntiAliasingMethod
 USTRUCT(BlueprintType)
 struct FPostProcessSettings
 {
-	GENERATED_BODY()
+	GENERATED_USTRUCT_BODY()
 
 	// first all bOverride_... as they get grouped together into bitfields
 
@@ -1040,11 +1040,9 @@ struct FPostProcessSettings
 UCLASS()
 class UScene : public UObject
 {
-	GENERATED_BODY()
-public:
-	UScene(const FObjectInitializer& ObjectInitializer = FObjectInitializer::Get());
+	GENERATED_UCLASS_BODY()
 
 
 	/** bits needed to store DPG value */
 	#define SDPG_NumBits	3
-};
+};
\ No newline at end of file
diff --git a/Engine/Source/Runtime/Engine/Classes/Engine/TextRenderActor.h b/Engine/Source/Runtime/Engine/Classes/Engine/TextRenderActor.h
--- a/Engine/Source/Runtime/Engine/Classes/Engine/TextRenderActor.h
+++ b/Engine/Source/Runtime/Engine/Classes/Engine/TextRenderActor.h
@@ -11,9 +11,7 @@
 UCLASS(MinimalAPI, ComponentWrapperClass, hideCategories = (Collision, Attachment, Actor))
 class ATextRenderActor : public AActor
 {
-	GENERATED_BODY()
-public:
-	ENGINE_API ATextRenderActor(const FObjectInitializer& ObjectInitializer = FObjectInitializer::Get());
+	GENERATED_UCLASS_BODY()
 
 	friend class UActorFactoryTextRender;
 
diff --git a/Engine/Source/Runtime/Engine/Classes/Engine/UserInterfaceSettings.h b/Engine/Source/Runtime/Engine/Classes/Engine/UserInterfaceSettings.h
--- a/Engine/Source/Runtime/Engine/Classes/Engine/UserInterfaceSettings.h
+++ b/Engine/Source/Runtime/Engine/Classes/Engine/UserInterfaceSettings.h
@@ -44,9 +44,7 @@ UCLASS(config=Engine, defaultconfig)
 class ENGINE_API UUserInterfaceSettings
 	: public UObject
 {
-	GENERATED_BODY()
-public:
-	UUserInterfaceSettings(const FObjectInitializer& ObjectInitializer = FObjectInitializer::Get());
+	GENERATED_UCLASS_BODY()
 
 public:
 
diff --git a/Engine/Source/Runtime/Engine/Classes/GameFramework/Controller.h b/Engine/Source/Runtime/Engine/Classes/GameFramework/Controller.h
--- a/Engine/Source/Runtime/Engine/Classes/GameFramework/Controller.h
+++ b/Engine/Source/Runtime/Engine/Classes/GameFramework/Controller.h
@@ -31,9 +31,7 @@ DECLARE_DYNAMIC_MULTICAST_DELEGATE_FourParams( FInstigatedAnyDamageSignature, fl
 UCLASS(abstract, notplaceable, NotBlueprintable, HideCategories=(Collision,Rendering,"Utilities|Transformation")) 
 class ENGINE_API AController : public AActor, public INavAgentInterface
 {
-	GENERATED_BODY()
-public:
-	AController(const FObjectInitializer& ObjectInitializer = FObjectInitializer::Get());
+	GENERATED_UCLASS_BODY()
 
 private:
 	/** Pawn currently being controlled by this controller.  Use Pawn.Possess() to take control of a pawn */
@@ -155,14 +153,12 @@ public:
 	class APlayerController* CastToPlayerController();
 
 	/** Replicated function to set the pawn location and rotation, allowing server to force (ex. teleports). */
-	UFUNCTION(reliable, client="ClientSetLocation_Implementation")
+	UFUNCTION(reliable, client)
 	void ClientSetLocation(FVector NewLocation, FRotator NewRotation);
-	virtual void ClientSetLocation_Implementation(FVector NewLocation, FRotator NewRotation);
 
 	/** Replicated function to set the pawn rotation, allowing the server to force. */
-	UFUNCTION(reliable, client="ClientSetRotation_Implementation")
+	UFUNCTION(reliable, client)
 	void ClientSetRotation(FRotator NewRotation, bool bResetCamera = false);
-	virtual void ClientSetRotation_Implementation(FRotator NewRotation, bool bResetCamera);
 
 	/** Return the Pawn that is currently 'controlled' by this PlayerController */
 	UFUNCTION(BlueprintCallable, Category="Pawn", meta=(FriendlyName="Get Controlled Pawn"))
diff --git a/Engine/Source/Runtime/Engine/Classes/GameFramework/DefaultPawn.h b/Engine/Source/Runtime/Engine/Classes/GameFramework/DefaultPawn.h
--- a/Engine/Source/Runtime/Engine/Classes/GameFramework/DefaultPawn.h
+++ b/Engine/Source/Runtime/Engine/Classes/GameFramework/DefaultPawn.h
@@ -20,9 +20,7 @@ class UStaticMeshComponent;
 UCLASS(config=Game, Blueprintable, BlueprintType)
 class ENGINE_API ADefaultPawn : public APawn
 {
-	GENERATED_BODY()
-public:
-	ADefaultPawn(const FObjectInitializer& ObjectInitializer = FObjectInitializer::Get());
+	GENERATED_UCLASS_BODY()
 
 	// Begin Pawn overrides
 	virtual UPawnMovementComponent* GetMovementComponent() const override;
diff --git a/Engine/Source/Runtime/Engine/Classes/GameFramework/PainCausingVolume.h b/Engine/Source/Runtime/Engine/Classes/GameFramework/PainCausingVolume.h
--- a/Engine/Source/Runtime/Engine/Classes/GameFramework/PainCausingVolume.h
+++ b/Engine/Source/Runtime/Engine/Classes/GameFramework/PainCausingVolume.h
@@ -16,9 +16,7 @@
 UCLASS()
 class ENGINE_API APainCausingVolume : public APhysicsVolume
 {
-	GENERATED_BODY()
-public:
-	APainCausingVolume(const FObjectInitializer& ObjectInitializer = FObjectInitializer::Get());
+	GENERATED_UCLASS_BODY()
 
 	/** Whether volume currently causes damage. */
 	UPROPERTY(EditAnywhere, BlueprintReadWrite, Category=PainCausingVolume)
diff --git a/Engine/Source/Runtime/Engine/Classes/GameFramework/SpectatorPawn.h b/Engine/Source/Runtime/Engine/Classes/GameFramework/SpectatorPawn.h
--- a/Engine/Source/Runtime/Engine/Classes/GameFramework/SpectatorPawn.h
+++ b/Engine/Source/Runtime/Engine/Classes/GameFramework/SpectatorPawn.h
@@ -13,9 +13,7 @@
 UCLASS(config=Game, Blueprintable, BlueprintType, notplaceable)
 class ENGINE_API ASpectatorPawn : public ADefaultPawn
 {
-	GENERATED_BODY()
-public:
-	ASpectatorPawn(const FObjectInitializer& ObjectInitializer = FObjectInitializer::Get());
+	GENERATED_UCLASS_BODY()
 
 	// Begin Pawn overrides
 	/** Overridden to avoid changing network role. If subclasses want networked behavior, call the Pawn::PossessedBy() instead. */
diff --git a/Engine/Source/Runtime/Engine/Classes/Interfaces/Interface_AssetUserData.h b/Engine/Source/Runtime/Engine/Classes/Interfaces/Interface_AssetUserData.h
--- a/Engine/Source/Runtime/Engine/Classes/Interfaces/Interface_AssetUserData.h
+++ b/Engine/Source/Runtime/Engine/Classes/Interfaces/Interface_AssetUserData.h
@@ -9,15 +9,12 @@ class UAssetUserData;
 UINTERFACE(MinimalApi, meta=(CannotImplementInterfaceInBlueprint))
 class UInterface_AssetUserData : public UInterface
 {
-	GENERATED_BODY()
-public:
-	ENGINE_API UInterface_AssetUserData(const FObjectInitializer& ObjectInitializer = FObjectInitializer::Get());
+	GENERATED_UINTERFACE_BODY()
 };
 
 class IInterface_AssetUserData
 {
-	GENERATED_BODY()
-public:
+	GENERATED_IINTERFACE_BODY()
 
 	virtual void AddAssetUserData(UAssetUserData* InUserData) {}
 
diff --git a/Engine/Source/Runtime/Engine/Classes/Kismet/KismetSystemLibrary.h b/Engine/Source/Runtime/Engine/Classes/Kismet/KismetSystemLibrary.h
--- a/Engine/Source/Runtime/Engine/Classes/Kismet/KismetSystemLibrary.h
+++ b/Engine/Source/Runtime/Engine/Classes/Kismet/KismetSystemLibrary.h
@@ -49,7 +49,7 @@ namespace EQuitPreference
 USTRUCT()
 struct FGenericStruct
 {
-	GENERATED_BODY()
+	GENERATED_USTRUCT_BODY()
 
 	UPROPERTY()
 	int32 Data;
@@ -58,9 +58,7 @@ struct FGenericStruct
 UCLASS(MinimalAPI)
 class UKismetSystemLibrary : public UBlueprintFunctionLibrary
 {
-	GENERATED_BODY()
-public:
-	ENGINE_API UKismetSystemLibrary(const FObjectInitializer& ObjectInitializer = FObjectInitializer::Get());
+	GENERATED_UCLASS_BODY()
 
 	// --- Globally useful functions ------------------------------
 	/** Prints a stack trace to the log, so you can see how a blueprint got to this node */
diff --git a/Engine/Source/Runtime/Engine/Classes/Materials/MaterialExpressionAppendVector.h b/Engine/Source/Runtime/Engine/Classes/Materials/MaterialExpressionAppendVector.h
--- a/Engine/Source/Runtime/Engine/Classes/Materials/MaterialExpressionAppendVector.h
+++ b/Engine/Source/Runtime/Engine/Classes/Materials/MaterialExpressionAppendVector.h
@@ -8,9 +8,7 @@
 UCLASS(collapsecategories, hidecategories=Object)
 class UMaterialExpressionAppendVector : public UMaterialExpression
 {
-	GENERATED_BODY()
-public:
-	UMaterialExpressionAppendVector(const FObjectInitializer& ObjectInitializer = FObjectInitializer::Get());
+	GENERATED_UCLASS_BODY()
 
 	UPROPERTY()
 	FExpressionInput A;
diff --git a/Engine/Source/Runtime/Engine/Classes/Materials/MaterialExpressionComponentMask.h b/Engine/Source/Runtime/Engine/Classes/Materials/MaterialExpressionComponentMask.h
--- a/Engine/Source/Runtime/Engine/Classes/Materials/MaterialExpressionComponentMask.h
+++ b/Engine/Source/Runtime/Engine/Classes/Materials/MaterialExpressionComponentMask.h
@@ -8,9 +8,7 @@
 UCLASS(collapsecategories, hidecategories=Object, MinimalAPI)
 class UMaterialExpressionComponentMask : public UMaterialExpression
 {
-	GENERATED_BODY()
-public:
-	ENGINE_API UMaterialExpressionComponentMask(const FObjectInitializer& ObjectInitializer = FObjectInitializer::Get());
+	GENERATED_UCLASS_BODY()
 
 	UPROPERTY()
 	FExpressionInput Input;
diff --git a/Engine/Source/Runtime/Engine/Classes/Materials/MaterialExpressionConstant.h b/Engine/Source/Runtime/Engine/Classes/Materials/MaterialExpressionConstant.h
--- a/Engine/Source/Runtime/Engine/Classes/Materials/MaterialExpressionConstant.h
+++ b/Engine/Source/Runtime/Engine/Classes/Materials/MaterialExpressionConstant.h
@@ -8,9 +8,7 @@
 UCLASS(collapsecategories, hidecategories=Object, MinimalAPI)
 class UMaterialExpressionConstant : public UMaterialExpression
 {
-	GENERATED_BODY()
-public:
-	ENGINE_API UMaterialExpressionConstant(const FObjectInitializer& ObjectInitializer = FObjectInitializer::Get());
+	GENERATED_UCLASS_BODY()
 
 	UPROPERTY(EditAnywhere, Category=MaterialExpressionConstant, DisplayName="Value")
 	float R;
diff --git a/Engine/Source/Runtime/Engine/Classes/Materials/MaterialExpressionConstant2Vector.h b/Engine/Source/Runtime/Engine/Classes/Materials/MaterialExpressionConstant2Vector.h
--- a/Engine/Source/Runtime/Engine/Classes/Materials/MaterialExpressionConstant2Vector.h
+++ b/Engine/Source/Runtime/Engine/Classes/Materials/MaterialExpressionConstant2Vector.h
@@ -8,9 +8,7 @@
 UCLASS(collapsecategories, hidecategories=Object, MinimalAPI)
 class UMaterialExpressionConstant2Vector : public UMaterialExpression
 {
-	GENERATED_BODY()
-public:
-	ENGINE_API UMaterialExpressionConstant2Vector(const FObjectInitializer& ObjectInitializer = FObjectInitializer::Get());
+	GENERATED_UCLASS_BODY()
 
 	UPROPERTY(EditAnywhere, Category=MaterialExpressionConstant2Vector)
 	float R;
diff --git a/Engine/Source/Runtime/Engine/Classes/Materials/MaterialExpressionDepthFade.h b/Engine/Source/Runtime/Engine/Classes/Materials/MaterialExpressionDepthFade.h
--- a/Engine/Source/Runtime/Engine/Classes/Materials/MaterialExpressionDepthFade.h
+++ b/Engine/Source/Runtime/Engine/Classes/Materials/MaterialExpressionDepthFade.h
@@ -11,9 +11,7 @@
 UCLASS(collapsecategories, hidecategories=Object)
 class UMaterialExpressionDepthFade : public UMaterialExpression
 {
-	GENERATED_BODY()
-public:
-	UMaterialExpressionDepthFade(const FObjectInitializer& ObjectInitializer = FObjectInitializer::Get());
+	GENERATED_UCLASS_BODY()
 
 	/** Input opacity which will be scaled by the result of the fade. */
 	UPROPERTY(meta = (RequiredInput = "false", ToolTip = "Defaults to 'OpacityDefault' if not specified"))
diff --git a/Engine/Source/Runtime/Engine/Classes/Materials/MaterialExpressionDeriveNormalZ.h b/Engine/Source/Runtime/Engine/Classes/Materials/MaterialExpressionDeriveNormalZ.h
--- a/Engine/Source/Runtime/Engine/Classes/Materials/MaterialExpressionDeriveNormalZ.h
+++ b/Engine/Source/Runtime/Engine/Classes/Materials/MaterialExpressionDeriveNormalZ.h
@@ -9,9 +9,7 @@
 UCLASS(collapsecategories, hidecategories=Object)
 class UMaterialExpressionDeriveNormalZ : public UMaterialExpression
 {
-	GENERATED_BODY()
-public:
-	UMaterialExpressionDeriveNormalZ(const FObjectInitializer& ObjectInitializer = FObjectInitializer::Get());
+	GENERATED_UCLASS_BODY()
 
 	UPROPERTY()
 	FExpressionInput InXY;
diff --git a/Engine/Source/Runtime/Engine/Classes/Materials/MaterialExpressionDotProduct.h b/Engine/Source/Runtime/Engine/Classes/Materials/MaterialExpressionDotProduct.h
--- a/Engine/Source/Runtime/Engine/Classes/Materials/MaterialExpressionDotProduct.h
+++ b/Engine/Source/Runtime/Engine/Classes/Materials/MaterialExpressionDotProduct.h
@@ -8,9 +8,7 @@
 UCLASS()
 class UMaterialExpressionDotProduct : public UMaterialExpression
 {
-	GENERATED_BODY()
-public:
-	UMaterialExpressionDotProduct(const FObjectInitializer& ObjectInitializer = FObjectInitializer::Get());
+	GENERATED_UCLASS_BODY()
 
 	UPROPERTY()
 	FExpressionInput A;
diff --git a/Engine/Source/Runtime/Engine/Classes/Materials/MaterialExpressionObjectPositionWS.h b/Engine/Source/Runtime/Engine/Classes/Materials/MaterialExpressionObjectPositionWS.h
--- a/Engine/Source/Runtime/Engine/Classes/Materials/MaterialExpressionObjectPositionWS.h
+++ b/Engine/Source/Runtime/Engine/Classes/Materials/MaterialExpressionObjectPositionWS.h
@@ -8,9 +8,7 @@
 UCLASS(collapsecategories, hidecategories=Object)
 class UMaterialExpressionObjectPositionWS : public UMaterialExpression
 {
-	GENERATED_BODY()
-public:
-	UMaterialExpressionObjectPositionWS(const FObjectInitializer& ObjectInitializer = FObjectInitializer::Get());
+	GENERATED_UCLASS_BODY()
 
 
 	// Begin UMaterialExpression Interface
diff --git a/Engine/Source/Runtime/Engine/Classes/Materials/MaterialExpressionSceneTexture.h b/Engine/Source/Runtime/Engine/Classes/Materials/MaterialExpressionSceneTexture.h
--- a/Engine/Source/Runtime/Engine/Classes/Materials/MaterialExpressionSceneTexture.h
+++ b/Engine/Source/Runtime/Engine/Classes/Materials/MaterialExpressionSceneTexture.h
@@ -62,9 +62,7 @@ enum ESceneTextureId
 UCLASS(collapsecategories, hidecategories=Object)
 class UMaterialExpressionSceneTexture : public UMaterialExpression
 {
-	GENERATED_BODY()
-public:
-	UMaterialExpressionSceneTexture(const FObjectInitializer& ObjectInitializer = FObjectInitializer::Get());
+	GENERATED_UCLASS_BODY()
 
 	/** UV in 0..1 range */
 	UPROPERTY(meta = (RequiredInput = "false", ToolTip = "Ignored if not specified"))
diff --git a/Engine/Source/Runtime/Engine/Classes/Materials/MaterialExpressionStaticBoolParameter.h b/Engine/Source/Runtime/Engine/Classes/Materials/MaterialExpressionStaticBoolParameter.h
--- a/Engine/Source/Runtime/Engine/Classes/Materials/MaterialExpressionStaticBoolParameter.h
+++ b/Engine/Source/Runtime/Engine/Classes/Materials/MaterialExpressionStaticBoolParameter.h
@@ -8,9 +8,7 @@
 UCLASS(collapsecategories, hidecategories=Object, MinimalAPI)
 class UMaterialExpressionStaticBoolParameter : public UMaterialExpressionParameter
 {
-	GENERATED_BODY()
-public:
-	ENGINE_API UMaterialExpressionStaticBoolParameter(const FObjectInitializer& ObjectInitializer = FObjectInitializer::Get());
+	GENERATED_UCLASS_BODY()
 
 	UPROPERTY(EditAnywhere, Category=MaterialExpressionStaticBoolParameter)
 	uint32 DefaultValue:1;
diff --git a/Engine/Source/Runtime/Engine/Classes/Materials/MaterialExpressionTextureSampleParameter2D.h b/Engine/Source/Runtime/Engine/Classes/Materials/MaterialExpressionTextureSampleParameter2D.h
--- a/Engine/Source/Runtime/Engine/Classes/Materials/MaterialExpressionTextureSampleParameter2D.h
+++ b/Engine/Source/Runtime/Engine/Classes/Materials/MaterialExpressionTextureSampleParameter2D.h
@@ -8,9 +8,7 @@
 UCLASS(collapsecategories, hidecategories=Object)
 class ENGINE_API UMaterialExpressionTextureSampleParameter2D : public UMaterialExpressionTextureSampleParameter
 {
-	GENERATED_BODY()
-public:
-	UMaterialExpressionTextureSampleParameter2D(const FObjectInitializer& ObjectInitializer = FObjectInitializer::Get());
+	GENERATED_UCLASS_BODY()
 
 
 	// Begin UMaterialExpression Interface
diff --git a/Engine/Source/Runtime/Engine/Classes/Materials/MaterialExpressionTransform.h b/Engine/Source/Runtime/Engine/Classes/Materials/MaterialExpressionTransform.h
--- a/Engine/Source/Runtime/Engine/Classes/Materials/MaterialExpressionTransform.h
+++ b/Engine/Source/Runtime/Engine/Classes/Materials/MaterialExpressionTransform.h
@@ -28,9 +28,7 @@ enum EMaterialVectorCoordTransform
 UCLASS(collapsecategories, hidecategories=Object)
 class UMaterialExpressionTransform : public UMaterialExpression
 {
-	GENERATED_BODY()
-public:
-	UMaterialExpressionTransform(const FObjectInitializer& ObjectInitializer = FObjectInitializer::Get());
+	GENERATED_UCLASS_BODY()
 
 	/** input expression for this transform */
 	UPROPERTY()
diff --git a/Engine/Source/Runtime/Engine/Classes/Matinee/InterpTrackInstFloatProp.h b/Engine/Source/Runtime/Engine/Classes/Matinee/InterpTrackInstFloatProp.h
--- a/Engine/Source/Runtime/Engine/Classes/Matinee/InterpTrackInstFloatProp.h
+++ b/Engine/Source/Runtime/Engine/Classes/Matinee/InterpTrackInstFloatProp.h
@@ -8,9 +8,7 @@
 UCLASS(MinimalAPI)
 class UInterpTrackInstFloatProp : public UInterpTrackInstProperty
 {
-	GENERATED_BODY()
-public:
-	ENGINE_API UInterpTrackInstFloatProp(const FObjectInitializer& ObjectInitializer = FObjectInitializer::Get());
+	GENERATED_UCLASS_BODY()
 
 	/** Pointer to float property in TrackObject. */
 	float* FloatProp;
diff --git a/Engine/Source/Runtime/Engine/Classes/Matinee/InterpTrackInstLinearColorProp.h b/Engine/Source/Runtime/Engine/Classes/Matinee/InterpTrackInstLinearColorProp.h
--- a/Engine/Source/Runtime/Engine/Classes/Matinee/InterpTrackInstLinearColorProp.h
+++ b/Engine/Source/Runtime/Engine/Classes/Matinee/InterpTrackInstLinearColorProp.h
@@ -7,9 +7,7 @@
 UCLASS()
 class UInterpTrackInstLinearColorProp : public UInterpTrackInstProperty
 {
-	GENERATED_BODY()
-public:
-	UInterpTrackInstLinearColorProp(const FObjectInitializer& ObjectInitializer = FObjectInitializer::Get());
+	GENERATED_UCLASS_BODY()
 
 	/** Pointer to color property in TrackObject. */
 	FLinearColor* ColorProp;
diff --git a/Engine/Source/Runtime/Engine/Classes/Matinee/InterpTrackInstVisibility.h b/Engine/Source/Runtime/Engine/Classes/Matinee/InterpTrackInstVisibility.h
--- a/Engine/Source/Runtime/Engine/Classes/Matinee/InterpTrackInstVisibility.h
+++ b/Engine/Source/Runtime/Engine/Classes/Matinee/InterpTrackInstVisibility.h
@@ -9,9 +9,7 @@
 UCLASS()
 class UInterpTrackInstVisibility : public UInterpTrackInst
 {
-	GENERATED_BODY()
-public:
-	UInterpTrackInstVisibility(const FObjectInitializer& ObjectInitializer = FObjectInitializer::Get());
+	GENERATED_UCLASS_BODY()
 
 	UPROPERTY(EditAnywhere, Category=InterpTrackInstVisibility)
 	TEnumAsByte<enum EVisibilityTrackAction> Action;
diff --git a/Engine/Source/Runtime/Engine/Classes/Particles/Emitter.h b/Engine/Source/Runtime/Engine/Classes/Particles/Emitter.h
--- a/Engine/Source/Runtime/Engine/Classes/Particles/Emitter.h
+++ b/Engine/Source/Runtime/Engine/Classes/Particles/Emitter.h
@@ -47,9 +47,7 @@ DECLARE_DYNAMIC_MULTICAST_DELEGATE_EightParams( FParticleCollisionSignature, FNa
 UCLASS(hideCategories=(Activation,"Components|Activation",Input,Collision,"Game|Damage"), ComponentWrapperClass)
 class ENGINE_API AEmitter : public AActor
 {
-	GENERATED_BODY()
-public:
-	AEmitter(const FObjectInitializer& ObjectInitializer = FObjectInitializer::Get());
+	GENERATED_UCLASS_BODY()
 
 private_subobject:
 	DEPRECATED_FORGAME(4.6, "ParticleSystemComponent should not be accessed directly, please use GetParticleSystemComponent() function instead. ParticleSystemComponent will soon be private and your code will not compile.")
diff --git a/Engine/Source/Runtime/Engine/Classes/Particles/Event/ParticleModuleEventSendToGame.h b/Engine/Source/Runtime/Engine/Classes/Particles/Event/ParticleModuleEventSendToGame.h
--- a/Engine/Source/Runtime/Engine/Classes/Particles/Event/ParticleModuleEventSendToGame.h
+++ b/Engine/Source/Runtime/Engine/Classes/Particles/Event/ParticleModuleEventSendToGame.h
@@ -7,9 +7,7 @@
 UCLASS(abstract, editinlinenew, hidecategories=Object)
 class UParticleModuleEventSendToGame : public UObject
 {
-	GENERATED_BODY()
-public:
-	UParticleModuleEventSendToGame(const FObjectInitializer& ObjectInitializer = FObjectInitializer::Get());
+	GENERATED_UCLASS_BODY()
 
 
 	/** This is our function to allow subclasses to "do the event action" **/
diff --git a/Engine/Source/Runtime/Engine/Classes/Particles/Location/ParticleModuleLocationEmitterDirect.h b/Engine/Source/Runtime/Engine/Classes/Particles/Location/ParticleModuleLocationEmitterDirect.h
--- a/Engine/Source/Runtime/Engine/Classes/Particles/Location/ParticleModuleLocationEmitterDirect.h
+++ b/Engine/Source/Runtime/Engine/Classes/Particles/Location/ParticleModuleLocationEmitterDirect.h
@@ -15,9 +15,7 @@
 UCLASS(editinlinenew, hidecategories=Object, meta=(DisplayName = "Emitter Direct Location"))
 class UParticleModuleLocationEmitterDirect : public UParticleModuleLocationBase
 {
-	GENERATED_BODY()
-public:
-	UParticleModuleLocationEmitterDirect(const FObjectInitializer& ObjectInitializer = FObjectInitializer::Get());
+	GENERATED_UCLASS_BODY()
 
 	/** The name of the emitter to use as a source for the location of the particles. */
 	UPROPERTY(EditAnywhere, export, noclear, Category=Location)
diff --git a/Engine/Source/Runtime/Engine/Classes/Particles/Orientation/ParticleModuleOrientationAxisLock.h b/Engine/Source/Runtime/Engine/Classes/Particles/Orientation/ParticleModuleOrientationAxisLock.h
--- a/Engine/Source/Runtime/Engine/Classes/Particles/Orientation/ParticleModuleOrientationAxisLock.h
+++ b/Engine/Source/Runtime/Engine/Classes/Particles/Orientation/ParticleModuleOrientationAxisLock.h
@@ -35,9 +35,7 @@ enum EParticleAxisLock
 UCLASS(editinlinenew, hidecategories=Object, MinimalAPI, meta=(DisplayName = "Lock Axis"))
 class UParticleModuleOrientationAxisLock : public UParticleModuleOrientationBase
 {
-	GENERATED_BODY()
-public:
-	ENGINE_API UParticleModuleOrientationAxisLock(const FObjectInitializer& ObjectInitializer = FObjectInitializer::Get());
+	GENERATED_UCLASS_BODY()
 
 	/**
 	 *	The lock axis flag setting.
diff --git a/Engine/Source/Runtime/Engine/Classes/Particles/Parameter/ParticleModuleParameterDynamic.h b/Engine/Source/Runtime/Engine/Classes/Particles/Parameter/ParticleModuleParameterDynamic.h
--- a/Engine/Source/Runtime/Engine/Classes/Particles/Parameter/ParticleModuleParameterDynamic.h
+++ b/Engine/Source/Runtime/Engine/Classes/Particles/Parameter/ParticleModuleParameterDynamic.h
@@ -37,7 +37,7 @@ enum EEmitterDynamicParameterValue
 USTRUCT()
 struct FEmitterDynamicParameter
 {
-	GENERATED_BODY()
+	GENERATED_USTRUCT_BODY()
 
 	/** The parameter name - from the material DynamicParameter expression. READ-ONLY */
 	UPROPERTY(Category=EmitterDynamicParameter, VisibleAnywhere, BlueprintReadWrite)
@@ -86,9 +86,7 @@ struct FEmitterDynamicParameter
 UCLASS(editinlinenew, hidecategories=Object, MinimalAPI, meta=(DisplayName = "Dynamic"))
 class UParticleModuleParameterDynamic : public UParticleModuleParameterBase
 {
-	GENERATED_BODY()
-public:
-	ENGINE_API UParticleModuleParameterDynamic(const FObjectInitializer& ObjectInitializer = FObjectInitializer::Get());
+	GENERATED_UCLASS_BODY()
 
 	/** The dynamic parameters this module uses. */
 	UPROPERTY(EditAnywhere, editfixedsize, Category=ParticleModuleParameterDynamic)
diff --git a/Engine/Source/Runtime/Engine/Classes/Particles/ParticleEmitter.h b/Engine/Source/Runtime/Engine/Classes/Particles/ParticleEmitter.h
--- a/Engine/Source/Runtime/Engine/Classes/Particles/ParticleEmitter.h
+++ b/Engine/Source/Runtime/Engine/Classes/Particles/ParticleEmitter.h
@@ -55,7 +55,7 @@ enum EEmitterRenderMode
 USTRUCT()
 struct FParticleBurst
 {
-	GENERATED_BODY()
+	GENERATED_USTRUCT_BODY()
 
 	/** The number of particles to burst */
 	UPROPERTY(EditAnywhere, Category=ParticleBurst)
@@ -83,9 +83,7 @@ struct FParticleBurst
 UCLASS(hidecategories=Object, editinlinenew, abstract, MinimalAPI)
 class UParticleEmitter : public UObject
 {
-	GENERATED_BODY()
-public:
-	ENGINE_API UParticleEmitter(const FObjectInitializer& ObjectInitializer = FObjectInitializer::Get());
+	GENERATED_UCLASS_BODY()
 
 	//=============================================================================
 	//	General variables
diff --git a/Engine/Source/Runtime/Engine/Classes/Particles/Rotation/ParticleModuleRotationOverLifetime.h b/Engine/Source/Runtime/Engine/Classes/Particles/Rotation/ParticleModuleRotationOverLifetime.h
--- a/Engine/Source/Runtime/Engine/Classes/Particles/Rotation/ParticleModuleRotationOverLifetime.h
+++ b/Engine/Source/Runtime/Engine/Classes/Particles/Rotation/ParticleModuleRotationOverLifetime.h
@@ -8,9 +8,7 @@
 UCLASS(editinlinenew, hidecategories=Object, MinimalAPI, meta=(DisplayName = "Rotation/Life"))
 class UParticleModuleRotationOverLifetime : public UParticleModuleRotationBase
 {
-	GENERATED_BODY()
-public:
-	ENGINE_API UParticleModuleRotationOverLifetime(const FObjectInitializer& ObjectInitializer = FObjectInitializer::Get());
+	GENERATED_UCLASS_BODY()
 
 	/** 
 	 *	The rotation of the particle (1.0 = 360 degrees).
diff --git a/Engine/Source/Runtime/Engine/Classes/Particles/Spawn/ParticleModuleSpawnPerUnit.h b/Engine/Source/Runtime/Engine/Classes/Particles/Spawn/ParticleModuleSpawnPerUnit.h
--- a/Engine/Source/Runtime/Engine/Classes/Particles/Spawn/ParticleModuleSpawnPerUnit.h
+++ b/Engine/Source/Runtime/Engine/Classes/Particles/Spawn/ParticleModuleSpawnPerUnit.h
@@ -8,9 +8,7 @@
 UCLASS(editinlinenew, hidecategories=Object, meta=(DisplayName = "Spawn PerUnit"))
 class UParticleModuleSpawnPerUnit : public UParticleModuleSpawnBase
 {
-	GENERATED_BODY()
-public:
-	UParticleModuleSpawnPerUnit(const FObjectInitializer& ObjectInitializer = FObjectInitializer::Get());
+	GENERATED_UCLASS_BODY()
 
 	/** 
 	 *	The scalar to apply to the distance traveled.
diff --git a/Engine/Source/Runtime/Engine/Classes/Particles/SubUV/ParticleModuleSubUVMovie.h b/Engine/Source/Runtime/Engine/Classes/Particles/SubUV/ParticleModuleSubUVMovie.h
--- a/Engine/Source/Runtime/Engine/Classes/Particles/SubUV/ParticleModuleSubUVMovie.h
+++ b/Engine/Source/Runtime/Engine/Classes/Particles/SubUV/ParticleModuleSubUVMovie.h
@@ -8,9 +8,7 @@
 UCLASS(editinlinenew, hidecategories=Object, hidecategories=SubUV, meta=(DisplayName = "SubUV Movie"))
 class UParticleModuleSubUVMovie : public UParticleModuleSubUV
 {
-	GENERATED_BODY()
-public:
-	UParticleModuleSubUVMovie(const FObjectInitializer& ObjectInitializer = FObjectInitializer::Get());
+	GENERATED_UCLASS_BODY()
 
 	/**
 	 *	If true, use the emitter time to look up the frame rate.
diff --git a/Engine/Source/Runtime/Engine/Classes/PhysicsEngine/PhysicsAsset.h b/Engine/Source/Runtime/Engine/Classes/PhysicsEngine/PhysicsAsset.h
--- a/Engine/Source/Runtime/Engine/Classes/PhysicsEngine/PhysicsAsset.h
+++ b/Engine/Source/Runtime/Engine/Classes/PhysicsEngine/PhysicsAsset.h
@@ -18,9 +18,7 @@
 UCLASS(hidecategories=Object, BlueprintType, MinimalAPI)
 class UPhysicsAsset : public UObject
 {
-	GENERATED_BODY()
-public:
-	ENGINE_API UPhysicsAsset(const FObjectInitializer& ObjectInitializer = FObjectInitializer::Get());
+	GENERATED_UCLASS_BODY()
 
 #if WITH_EDITORONLY_DATA
 	/** 
@@ -126,4 +124,4 @@ public:
 };
 
 #define		RB_MinSizeToLockDOF				(0.1)
-#define		RB_MinAngleToLockDOF			(5.0)
+#define		RB_MinAngleToLockDOF			(5.0)
\ No newline at end of file
diff --git a/Engine/Source/Runtime/Engine/Classes/PhysicsEngine/PhysicsSpringComponent.h b/Engine/Source/Runtime/Engine/Classes/PhysicsEngine/PhysicsSpringComponent.h
--- a/Engine/Source/Runtime/Engine/Classes/PhysicsEngine/PhysicsSpringComponent.h
+++ b/Engine/Source/Runtime/Engine/Classes/PhysicsEngine/PhysicsSpringComponent.h
@@ -12,9 +12,7 @@
 UCLASS(hidecategories=(Object, Mobility, LOD), ClassGroup=Physics, showcategories=Trigger)
 class ENGINE_API UPhysicsSpringComponent : public USceneComponent
 {
-	GENERATED_BODY()
-public:
-	UPhysicsSpringComponent(const FObjectInitializer& ObjectInitializer = FObjectInitializer::Get());
+	GENERATED_UCLASS_BODY()
 
 	/** Specifies how much strength the spring has. The higher the SpringStiffness the more force the spring can push on a body with. */
 	UPROPERTY(BlueprintReadWrite, EditAnywhere, Category = Physics)
diff --git a/Engine/Source/Runtime/Engine/Classes/Sound/DialogueSoundWaveProxy.h b/Engine/Source/Runtime/Engine/Classes/Sound/DialogueSoundWaveProxy.h
--- a/Engine/Source/Runtime/Engine/Classes/Sound/DialogueSoundWaveProxy.h
+++ b/Engine/Source/Runtime/Engine/Classes/Sound/DialogueSoundWaveProxy.h
@@ -9,9 +9,7 @@
 UCLASS()
 class UDialogueSoundWaveProxy : public USoundBase
 {
-	GENERATED_BODY()
-public:
-	UDialogueSoundWaveProxy(const FObjectInitializer& ObjectInitializer = FObjectInitializer::Get());
+	GENERATED_UCLASS_BODY()
 
 	friend class UDialogueWave;
 
diff --git a/Engine/Source/Runtime/Engine/Classes/Sound/SoundClass.h b/Engine/Source/Runtime/Engine/Classes/Sound/SoundClass.h
--- a/Engine/Source/Runtime/Engine/Classes/Sound/SoundClass.h
+++ b/Engine/Source/Runtime/Engine/Classes/Sound/SoundClass.h
@@ -17,7 +17,7 @@ namespace EAudioOutputTarget
 USTRUCT()
 struct FSoundClassEditorData
 {
-	GENERATED_BODY()
+	GENERATED_USTRUCT_BODY()
 
 	int32 NodePosX;
 
@@ -43,7 +43,7 @@ struct FSoundClassEditorData
 USTRUCT()
 struct FSoundClassProperties
 {
-	GENERATED_BODY()
+	GENERATED_USTRUCT_BODY()
 
 	/** Volume multiplier. */
 	UPROPERTY(EditAnywhere, Category=SoundClassProperties)
@@ -132,7 +132,7 @@ struct FSoundClassProperties
 USTRUCT()
 struct FPassiveSoundMixModifier
 {
-	GENERATED_BODY()
+	GENERATED_USTRUCT_BODY()
 
 	/** The SoundMix to activate */
 	UPROPERTY(EditAnywhere, Category=PassiveSoundMixModifier)
@@ -158,9 +158,7 @@ struct FPassiveSoundMixModifier
 UCLASS(hidecategories=object, MinimalAPI)
 class USoundClass : public UObject
 {
-	GENERATED_BODY()
-public:
-	ENGINE_API USoundClass(const FObjectInitializer& ObjectInitializer = FObjectInitializer::Get());
+	GENERATED_UCLASS_BODY()
 
 	/** Configurable properties like volume and priority. */
 	UPROPERTY(EditAnywhere, Category=SoundClass)
diff --git a/Engine/Source/Runtime/Engine/Classes/Sound/SoundNodeEnveloper.h b/Engine/Source/Runtime/Engine/Classes/Sound/SoundNodeEnveloper.h
--- a/Engine/Source/Runtime/Engine/Classes/Sound/SoundNodeEnveloper.h
+++ b/Engine/Source/Runtime/Engine/Classes/Sound/SoundNodeEnveloper.h
@@ -10,9 +10,7 @@
 UCLASS(hidecategories=Object, editinlinenew, MinimalAPI, meta=( DisplayName="Enveloper" ))
 class USoundNodeEnveloper : public USoundNode
 {
-	GENERATED_BODY()
-public:
-	ENGINE_API USoundNodeEnveloper(const FObjectInitializer& ObjectInitializer = FObjectInitializer::Get());
+	GENERATED_UCLASS_BODY()
 
 	// The time in seconds where the envelope's loop begins.
 	UPROPERTY(EditAnywhere, Category=Looping)
diff --git a/Engine/Source/Runtime/Engine/Classes/Sound/SoundNodeParamCrossFade.h b/Engine/Source/Runtime/Engine/Classes/Sound/SoundNodeParamCrossFade.h
--- a/Engine/Source/Runtime/Engine/Classes/Sound/SoundNodeParamCrossFade.h
+++ b/Engine/Source/Runtime/Engine/Classes/Sound/SoundNodeParamCrossFade.h
@@ -10,9 +10,7 @@
 UCLASS(editinlinenew, MinimalAPI, meta=( DisplayName="Crossfade by Param" ))
 class USoundNodeParamCrossFade : public USoundNodeDistanceCrossFade
 {
-	GENERATED_BODY()
-public:
-	ENGINE_API USoundNodeParamCrossFade(const FObjectInitializer& ObjectInitializer = FObjectInitializer::Get());
+	GENERATED_UCLASS_BODY()
 
 public:
 
diff --git a/Engine/Source/Runtime/Engine/Classes/Sound/SoundNodeSoundClass.h b/Engine/Source/Runtime/Engine/Classes/Sound/SoundNodeSoundClass.h
--- a/Engine/Source/Runtime/Engine/Classes/Sound/SoundNodeSoundClass.h
+++ b/Engine/Source/Runtime/Engine/Classes/Sound/SoundNodeSoundClass.h
@@ -11,9 +11,7 @@
 UCLASS(hidecategories=Object, editinlinenew, MinimalAPI, meta=( DisplayName="SoundClass" ))
 class USoundNodeSoundClass : public USoundNode
 {
-	GENERATED_BODY()
-public:
-	ENGINE_API USoundNodeSoundClass(const FObjectInitializer& ObjectInitializer = FObjectInitializer::Get());
+	GENERATED_UCLASS_BODY()
 
 	UPROPERTY(EditAnywhere, Category=SoundClass)
 	USoundClass* SoundClassOverride;
diff --git a/Engine/Source/Runtime/Engine/Classes/VectorField/VectorFieldVolume.h b/Engine/Source/Runtime/Engine/Classes/VectorField/VectorFieldVolume.h
--- a/Engine/Source/Runtime/Engine/Classes/VectorField/VectorFieldVolume.h
+++ b/Engine/Source/Runtime/Engine/Classes/VectorField/VectorFieldVolume.h
@@ -11,9 +11,7 @@
 UCLASS(hidecategories=(Object, Advanced, Collision), MinimalAPI)
 class AVectorFieldVolume : public AActor
 {
-	GENERATED_BODY()
-public:
-	ENGINE_API AVectorFieldVolume(const FObjectInitializer& ObjectInitializer = FObjectInitializer::Get());
+	GENERATED_UCLASS_BODY()
 
 private_subobject:
 	DEPRECATED_FORGAME(4.6, "VectorFieldComponent should not be accessed directly, please use GetVectorFieldComponent() function instead. VectorFieldComponent will soon be private and your code will not compile.")
diff --git a/Engine/Source/Runtime/Engine/Classes/Vehicles/VehicleAnimInstance.h b/Engine/Source/Runtime/Engine/Classes/Vehicles/VehicleAnimInstance.h
--- a/Engine/Source/Runtime/Engine/Classes/Vehicles/VehicleAnimInstance.h
+++ b/Engine/Source/Runtime/Engine/Classes/Vehicles/VehicleAnimInstance.h
@@ -10,9 +10,7 @@
 UCLASS(transient)
 class ENGINE_API UVehicleAnimInstance : public UAnimInstance
 {
-	GENERATED_BODY()
-public:
-	UVehicleAnimInstance(const FObjectInitializer& ObjectInitializer = FObjectInitializer::Get());
+	GENERATED_UCLASS_BODY()
 
 	/** Makes a montage jump to the end of a named section. */
 	UFUNCTION(BlueprintCallable, Category="Animation")
diff --git a/Engine/Source/Runtime/GameplayAbilities/Public/Abilities/GameplayAbilityTargetActor_Trace.h b/Engine/Source/Runtime/GameplayAbilities/Public/Abilities/GameplayAbilityTargetActor_Trace.h
--- a/Engine/Source/Runtime/GameplayAbilities/Public/Abilities/GameplayAbilityTargetActor_Trace.h
+++ b/Engine/Source/Runtime/GameplayAbilities/Public/Abilities/GameplayAbilityTargetActor_Trace.h
@@ -10,9 +10,7 @@
 UCLASS(Abstract, Blueprintable, notplaceable, config=Game)
 class GAMEPLAYABILITIES_API AGameplayAbilityTargetActor_Trace : public AGameplayAbilityTargetActor
 {
-	GENERATED_BODY()
-public:
-	AGameplayAbilityTargetActor_Trace(const FObjectInitializer& ObjectInitializer = FObjectInitializer::Get());
+	GENERATED_UCLASS_BODY()
 
 public:
 	virtual void EndPlay(const EEndPlayReason::Type EndPlayReason) override;
@@ -49,4 +47,4 @@ protected:
 	FGameplayAbilityTargetDataHandle MakeTargetData(FHitResult HitResult) const;
 
 	TWeakObjectPtr<AGameplayAbilityWorldReticle> ReticleActor;
-};
+};
\ No newline at end of file
diff --git a/Engine/Source/Runtime/GameplayAbilities/Public/AbilitySystemGlobals.h b/Engine/Source/Runtime/GameplayAbilities/Public/AbilitySystemGlobals.h
--- a/Engine/Source/Runtime/GameplayAbilities/Public/AbilitySystemGlobals.h
+++ b/Engine/Source/Runtime/GameplayAbilities/Public/AbilitySystemGlobals.h
@@ -21,9 +21,7 @@ struct FGameplayEffectSpec;
 UCLASS(config=Game)
 class GAMEPLAYABILITIES_API UAbilitySystemGlobals : public UObject
 {
-	GENERATED_BODY()
-public:
-	UAbilitySystemGlobals(const FObjectInitializer& ObjectInitializer = FObjectInitializer::Get());
+	GENERATED_UCLASS_BODY()
 
 	/** Gets the single instance of the globals object, will create it as necessary */
 	static UAbilitySystemGlobals& Get();
diff --git a/Engine/Source/Runtime/GameplayAbilities/Public/GameplayAbilitySpec.h b/Engine/Source/Runtime/GameplayAbilities/Public/GameplayAbilitySpec.h
--- a/Engine/Source/Runtime/GameplayAbilities/Public/GameplayAbilitySpec.h
+++ b/Engine/Source/Runtime/GameplayAbilities/Public/GameplayAbilitySpec.h
@@ -14,7 +14,7 @@
 USTRUCT(BlueprintType)
 struct FGameplayAbilitySpecHandle
 {
-	GENERATED_BODY()
+	GENERATED_USTRUCT_BODY()
 
 	FGameplayAbilitySpecHandle()
 	: Handle(INDEX_NONE)
@@ -89,7 +89,6 @@ enum class EGameplayEffectGrantedAbilityRemovePolicy : uint8
 USTRUCT(BlueprintType)
 struct FGameplayAbilitySpecDef
 {
-	GENERATED_BODY()
 	FGameplayAbilitySpecDef()
 		: Level(1)
 		, InputID(INDEX_NONE)
@@ -97,6 +96,7 @@ struct FGameplayAbilitySpecDef
 	{
 	}
 
+	GENERATED_USTRUCT_BODY()
 
 	UPROPERTY(EditDefaultsOnly, Category="Ability Definition", NotReplicated)
 	TSubclassOf<UGameplayAbility> Ability;
@@ -130,7 +130,7 @@ struct FGameplayAbilitySpecDef
 USTRUCT(BlueprintType)
 struct GAMEPLAYABILITIES_API FGameplayAbilityActivationInfo
 {
-	GENERATED_BODY()
+	GENERATED_USTRUCT_BODY()
 
 	FGameplayAbilityActivationInfo()
 		: ActivationMode(EGameplayAbilityActivationMode::Authority)
@@ -186,7 +186,7 @@ private:
 USTRUCT()
 struct GAMEPLAYABILITIES_API FGameplayAbilitySpec : public FFastArraySerializerItem
 {
-	GENERATED_BODY()
+	GENERATED_USTRUCT_BODY()
 
 	FGameplayAbilitySpec()
 	: Ability(nullptr), Level(1), InputID(INDEX_NONE), SourceObject(nullptr), InputPressed(false), ActiveCount(0), RemoveAfterActivation(false), PendingRemove(false)
@@ -278,7 +278,7 @@ struct GAMEPLAYABILITIES_API FGameplayAbilitySpec : public FFastArraySerializerI
 USTRUCT()
 struct GAMEPLAYABILITIES_API FGameplayAbilitySpecContainer : public FFastArraySerializer
 {
-	GENERATED_BODY()
+	GENERATED_USTRUCT_BODY()
 
 	/** List of activatable abilities */
 	UPROPERTY()
@@ -315,4 +315,4 @@ private:
 	UAbilitySystemComponent& AbilitySystemComponent;
 };
 
-#define ABILITYLIST_SCOPE_LOCK()	FScopedAbilityListLock ActiveScopeLock(*this);
+#define ABILITYLIST_SCOPE_LOCK()	FScopedAbilityListLock ActiveScopeLock(*this);
\ No newline at end of file
diff --git a/Engine/Source/Runtime/GameplayAbilities/Public/GameplayCueNotify_Actor.h b/Engine/Source/Runtime/GameplayAbilities/Public/GameplayCueNotify_Actor.h
--- a/Engine/Source/Runtime/GameplayAbilities/Public/GameplayCueNotify_Actor.h
+++ b/Engine/Source/Runtime/GameplayAbilities/Public/GameplayCueNotify_Actor.h
@@ -22,9 +22,7 @@
 UCLASS(Blueprintable, meta = (ShowWorldContextPin), hidecategories = (Replication))
 class GAMEPLAYABILITIES_API AGameplayCueNotify_Actor : public AActor
 {
-	GENERATED_BODY()
-public:
-	AGameplayCueNotify_Actor(const FObjectInitializer& ObjectInitializer = FObjectInitializer::Get());
+	GENERATED_UCLASS_BODY()
 
 	/** Does this GameplayCueNotify handle this type of GameplayCueEvent? */
 	virtual bool HandlesEvent(EGameplayCueEvent::Type EventType) const;
@@ -47,19 +45,15 @@ public:
 
 	UFUNCTION(BlueprintNativeEvent, Category = "GameplayCueNotify")
 	bool OnExecute(AActor* MyTarget, FGameplayCueParameters Parameters);
-	virtual bool OnExecute_Implementation(AActor* MyTarget, FGameplayCueParameters Parameters);
 
 	UFUNCTION(BlueprintNativeEvent, Category = "GameplayCueNotify")
 	bool OnActive(AActor* MyTarget, FGameplayCueParameters Parameters);
-	virtual bool OnActive_Implementation(AActor* MyTarget, FGameplayCueParameters Parameters);
 
 	UFUNCTION(BlueprintNativeEvent, Category = "GameplayCueNotify")
 	bool WhileActive(AActor* MyTarget, FGameplayCueParameters Parameters);
-	virtual bool WhileActive_Implementation(AActor* MyTarget, FGameplayCueParameters Parameters);
 
 	UFUNCTION(BlueprintNativeEvent, Category = "GameplayCueNotify")
 	bool OnRemove(AActor* MyTarget, FGameplayCueParameters Parameters);
-	virtual bool OnRemove_Implementation(AActor* MyTarget, FGameplayCueParameters Parameters);
 
 	UPROPERTY(EditDefaultsOnly, Category = GameplayCue)
 	FGameplayTag	GameplayCueTag;
@@ -78,4 +72,4 @@ public:
 
 private:
 	virtual void DeriveGameplayCueTagFromAssetName();
-};
+};
\ No newline at end of file
diff --git a/Engine/Source/Runtime/GameplayAbilities/Public/GameplayEffectExecutionCalculation.h b/Engine/Source/Runtime/GameplayAbilities/Public/GameplayEffectExecutionCalculation.h
--- a/Engine/Source/Runtime/GameplayAbilities/Public/GameplayEffectExecutionCalculation.h
+++ b/Engine/Source/Runtime/GameplayAbilities/Public/GameplayEffectExecutionCalculation.h
@@ -14,7 +14,7 @@ class UAbilitySystemComponent;
 USTRUCT()
 struct GAMEPLAYABILITIES_API FGameplayEffectCustomExecutionParameters
 {
-	GENERATED_BODY()
+	GENERATED_USTRUCT_BODY()
 
 public:
 
@@ -112,7 +112,7 @@ private:
 USTRUCT()
 struct GAMEPLAYABILITIES_API FGameplayEffectCustomExecutionOutput
 {
-	GENERATED_BODY()
+	GENERATED_USTRUCT_BODY()
 
 public:
 
@@ -168,9 +168,7 @@ private:
 UCLASS(BlueprintType, Blueprintable, Abstract)
 class GAMEPLAYABILITIES_API UGameplayEffectExecutionCalculation : public UGameplayEffectCalculation
 {
-	GENERATED_BODY()
-public:
-	UGameplayEffectExecutionCalculation(const FObjectInitializer& ObjectInitializer = FObjectInitializer::Get());
+	GENERATED_UCLASS_BODY()
 
 protected:
 
@@ -212,5 +210,4 @@ public:
 	 */
 	UFUNCTION(BlueprintNativeEvent, Category="Calculation")
 	void Execute(const FGameplayEffectCustomExecutionParameters& ExecutionParams, FGameplayEffectCustomExecutionOutput& OutExecutionOutput) const;
-	virtual void Execute_Implementation(const FGameplayEffectCustomExecutionParameters& ExecutionParams, FGameplayEffectCustomExecutionOutput& OutExecutionOutput) const;
 };
diff --git a/Engine/Source/Runtime/GameplayTags/Classes/GameplayTagsSettings.h b/Engine/Source/Runtime/GameplayTags/Classes/GameplayTagsSettings.h
--- a/Engine/Source/Runtime/GameplayTags/Classes/GameplayTagsSettings.h
+++ b/Engine/Source/Runtime/GameplayTags/Classes/GameplayTagsSettings.h
@@ -22,9 +22,7 @@
 UCLASS(config=Engine, defaultconfig, notplaceable)
 class GAMEPLAYTAGS_API UGameplayTagsSettings : public UObject
 {
-	GENERATED_BODY()
-public:
-	UGameplayTagsSettings(const FObjectInitializer& ObjectInitializer = FObjectInitializer::Get());
+	GENERATED_UCLASS_BODY()
 
 	UPROPERTY(config, EditAnywhere, Category=GameplayTags)
 	TArray<FString>		GameplayTags;
diff --git a/Engine/Source/Runtime/Landscape/Classes/LandscapeGizmoActiveActor.h b/Engine/Source/Runtime/Landscape/Classes/LandscapeGizmoActiveActor.h
--- a/Engine/Source/Runtime/Landscape/Classes/LandscapeGizmoActiveActor.h
+++ b/Engine/Source/Runtime/Landscape/Classes/LandscapeGizmoActiveActor.h
@@ -20,7 +20,7 @@ enum ELandscapeGizmoType
 USTRUCT()
 struct FGizmoSelectData
 {
-	GENERATED_BODY()
+	GENERATED_USTRUCT_BODY()
 
 #if WITH_EDITORONLY_DATA
 	UPROPERTY()
@@ -46,9 +46,7 @@ struct FGizmoSelectData
 UCLASS(notplaceable, MinimalAPI)
 class ALandscapeGizmoActiveActor : public ALandscapeGizmoActor
 {
-	GENERATED_BODY()
-public:
-	LANDSCAPE_API ALandscapeGizmoActiveActor(const FObjectInitializer& ObjectInitializer = FObjectInitializer::Get());
+	GENERATED_UCLASS_BODY()
 
 #if WITH_EDITORONLY_DATA
 	UPROPERTY(Transient)
diff --git a/Engine/Source/Runtime/Landscape/Classes/LandscapeGizmoActor.h b/Engine/Source/Runtime/Landscape/Classes/LandscapeGizmoActor.h
--- a/Engine/Source/Runtime/Landscape/Classes/LandscapeGizmoActor.h
+++ b/Engine/Source/Runtime/Landscape/Classes/LandscapeGizmoActor.h
@@ -9,9 +9,7 @@
 UCLASS(notplaceable, MinimalAPI, NotBlueprintable)
 class ALandscapeGizmoActor : public AActor
 {
-	GENERATED_BODY()
-public:
-	LANDSCAPE_API ALandscapeGizmoActor(const FObjectInitializer& ObjectInitializer = FObjectInitializer::Get());
+	GENERATED_UCLASS_BODY()
 
 #if WITH_EDITORONLY_DATA
 	UPROPERTY(EditAnywhere, Category=Gizmo)
diff --git a/Engine/Source/Runtime/Landscape/Classes/LandscapeGizmoRenderComponent.h b/Engine/Source/Runtime/Landscape/Classes/LandscapeGizmoRenderComponent.h
--- a/Engine/Source/Runtime/Landscape/Classes/LandscapeGizmoRenderComponent.h
+++ b/Engine/Source/Runtime/Landscape/Classes/LandscapeGizmoRenderComponent.h
@@ -7,9 +7,7 @@
 UCLASS(hidecategories=Object)
 class ULandscapeGizmoRenderComponent : public UPrimitiveComponent
 {
-	GENERATED_BODY()
-public:
-	ULandscapeGizmoRenderComponent(const FObjectInitializer& ObjectInitializer = FObjectInitializer::Get());
+	GENERATED_UCLASS_BODY()
 
 
 	// Begin UPrimitiveComponent Interface
diff --git a/Engine/Source/Runtime/Landscape/Classes/LandscapeInfoMap.h b/Engine/Source/Runtime/Landscape/Classes/LandscapeInfoMap.h
--- a/Engine/Source/Runtime/Landscape/Classes/LandscapeInfoMap.h
+++ b/Engine/Source/Runtime/Landscape/Classes/LandscapeInfoMap.h
@@ -9,9 +9,7 @@
 UCLASS()
 class ULandscapeInfoMap : public UObject
 {
-	GENERATED_BODY()
-public:
-	ULandscapeInfoMap(const FObjectInitializer& ObjectInitializer = FObjectInitializer::Get());
+	GENERATED_UCLASS_BODY()
 
 public:
 	void PostDuplicate(bool bDuplicateForPIE) override;
diff --git a/Engine/Source/Runtime/Landscape/Classes/LandscapeMeshProxyComponent.h b/Engine/Source/Runtime/Landscape/Classes/LandscapeMeshProxyComponent.h
--- a/Engine/Source/Runtime/Landscape/Classes/LandscapeMeshProxyComponent.h
+++ b/Engine/Source/Runtime/Landscape/Classes/LandscapeMeshProxyComponent.h
@@ -11,9 +11,7 @@
 UCLASS(MinimalAPI)
 class ULandscapeMeshProxyComponent : public UStaticMeshComponent
 {
-	GENERATED_BODY()
-public:
-	LANDSCAPE_API ULandscapeMeshProxyComponent(const FObjectInitializer& ObjectInitializer = FObjectInitializer::Get());
+	GENERATED_UCLASS_BODY()
 
 private:
 	/* The landscape this proxy was generated for */
diff --git a/Engine/Source/Runtime/Landscape/Classes/Materials/MaterialExpressionLandscapeLayerBlend.h b/Engine/Source/Runtime/Landscape/Classes/Materials/MaterialExpressionLandscapeLayerBlend.h
--- a/Engine/Source/Runtime/Landscape/Classes/Materials/MaterialExpressionLandscapeLayerBlend.h
+++ b/Engine/Source/Runtime/Landscape/Classes/Materials/MaterialExpressionLandscapeLayerBlend.h
@@ -17,7 +17,7 @@ enum ELandscapeLayerBlendType
 USTRUCT()
 struct FLayerBlendInput
 {
-	GENERATED_BODY()
+	GENERATED_USTRUCT_BODY()
 
 	UPROPERTY(EditAnywhere, Category=LayerBlendInput)
 	FName LayerName;
@@ -54,9 +54,7 @@ struct FLayerBlendInput
 UCLASS(collapsecategories, hidecategories=Object)
 class UMaterialExpressionLandscapeLayerBlend : public UMaterialExpression
 {
-	GENERATED_BODY()
-public:
-	UMaterialExpressionLandscapeLayerBlend(const FObjectInitializer& ObjectInitializer = FObjectInitializer::Get());
+	GENERATED_UCLASS_BODY()
 
 	UPROPERTY(EditAnywhere, Category=MaterialExpressionLandscapeLayerBlend)
 	TArray<FLayerBlendInput> Layers;
diff --git a/Engine/Source/Runtime/Landscape/Classes/Materials/MaterialExpressionLandscapeVisibilityMask.h b/Engine/Source/Runtime/Landscape/Classes/Materials/MaterialExpressionLandscapeVisibilityMask.h
--- a/Engine/Source/Runtime/Landscape/Classes/Materials/MaterialExpressionLandscapeVisibilityMask.h
+++ b/Engine/Source/Runtime/Landscape/Classes/Materials/MaterialExpressionLandscapeVisibilityMask.h
@@ -8,9 +8,7 @@
 UCLASS(collapsecategories, hidecategories=Object)
 class UMaterialExpressionLandscapeVisibilityMask : public UMaterialExpression
 {
-	GENERATED_BODY()
-public:
-	UMaterialExpressionLandscapeVisibilityMask(const FObjectInitializer& ObjectInitializer = FObjectInitializer::Get());
+	GENERATED_UCLASS_BODY()
 
 	/** GUID that should be unique within the material, this is used for parameter renaming. */
 	UPROPERTY()
diff --git a/Engine/Source/Runtime/MovieSceneCore/Classes/MovieScene.h b/Engine/Source/Runtime/MovieSceneCore/Classes/MovieScene.h
--- a/Engine/Source/Runtime/MovieSceneCore/Classes/MovieScene.h
+++ b/Engine/Source/Runtime/MovieSceneCore/Classes/MovieScene.h
@@ -12,7 +12,7 @@ MOVIESCENECORE_API DECLARE_LOG_CATEGORY_EXTERN(LogSequencerRuntime, Log, All);
 USTRUCT()
 struct FMovieSceneSpawnable 
 {
-	GENERATED_BODY()
+	GENERATED_USTRUCT_BODY()
 
 public:
 
@@ -75,7 +75,7 @@ private:
 USTRUCT()
 struct FMovieScenePossessable
 {
-	GENERATED_BODY()
+	GENERATED_USTRUCT_BODY( FMovieScenePossessable )
 
 public:
 
@@ -134,7 +134,7 @@ private:
 USTRUCT()
 struct FMovieSceneEditorData
 {
-	GENERATED_BODY()
+	GENERATED_USTRUCT_BODY()
 
 	/** List of collapsed sequencer nodes.  We store collapsed instead of expanded so that new nodes with no saved state are expanded by default */
 	UPROPERTY()
@@ -147,7 +147,7 @@ struct FMovieSceneEditorData
 USTRUCT()
 struct FMovieSceneObjectBinding
 {
-	GENERATED_BODY()
+	GENERATED_USTRUCT_BODY()
 
 	FMovieSceneObjectBinding()
 	{}
@@ -218,9 +218,7 @@ private:
 UCLASS( MinimalAPI )
 class UMovieScene : public UObject
 {
-	GENERATED_BODY()
-public:
-	MOVIESCENECORE_API UMovieScene(const FObjectInitializer& ObjectInitializer = FObjectInitializer::Get());
+	GENERATED_UCLASS_BODY()
 public:
 
 #if WITH_EDITOR
diff --git a/Engine/Source/Runtime/MovieSceneCore/Classes/MovieSceneSection.h b/Engine/Source/Runtime/MovieSceneCore/Classes/MovieSceneSection.h
--- a/Engine/Source/Runtime/MovieSceneCore/Classes/MovieSceneSection.h
+++ b/Engine/Source/Runtime/MovieSceneCore/Classes/MovieSceneSection.h
@@ -10,9 +10,7 @@
 UCLASS( abstract, MinimalAPI )
 class UMovieSceneSection : public UObject
 {
-	GENERATED_BODY()
-public:
-	MOVIESCENECORE_API UMovieSceneSection(const FObjectInitializer& ObjectInitializer = FObjectInitializer::Get());
+	GENERATED_UCLASS_BODY()
 public:
 
 	/**
diff --git a/Engine/Source/Runtime/MovieSceneCoreTypes/Classes/MovieSceneColorSection.h b/Engine/Source/Runtime/MovieSceneCoreTypes/Classes/MovieSceneColorSection.h
--- a/Engine/Source/Runtime/MovieSceneCoreTypes/Classes/MovieSceneColorSection.h
+++ b/Engine/Source/Runtime/MovieSceneCoreTypes/Classes/MovieSceneColorSection.h
@@ -12,9 +12,7 @@ struct FColorKey;
 UCLASS( MinimalAPI )
 class UMovieSceneColorSection : public UMovieSceneSection
 {
-	GENERATED_BODY()
-public:
-	MOVIESCENECORETYPES_API UMovieSceneColorSection(const FObjectInitializer& ObjectInitializer = FObjectInitializer::Get());
+	GENERATED_UCLASS_BODY()
 public:
 	/** MovieSceneSection interface */
 	virtual void MoveSection( float DeltaPosition ) override;
diff --git a/Engine/Source/Runtime/Online/OnlineSubsystemUtils/Classes/FindTurnBasedMatchCallbackProxy.h b/Engine/Source/Runtime/Online/OnlineSubsystemUtils/Classes/FindTurnBasedMatchCallbackProxy.h
--- a/Engine/Source/Runtime/Online/OnlineSubsystemUtils/Classes/FindTurnBasedMatchCallbackProxy.h
+++ b/Engine/Source/Runtime/Online/OnlineSubsystemUtils/Classes/FindTurnBasedMatchCallbackProxy.h
@@ -31,9 +31,7 @@ private:
 UCLASS(MinimalAPI)
 class UFindTurnBasedMatchCallbackProxy : public UOnlineBlueprintCallProxyBase
 {
-	GENERATED_BODY()
-public:
-	ONLINESUBSYSTEMUTILS_API UFindTurnBasedMatchCallbackProxy(const FObjectInitializer& ObjectInitializer = FObjectInitializer::Get());
+	GENERATED_UCLASS_BODY()
 
     virtual ~UFindTurnBasedMatchCallbackProxy();
     
diff --git a/Engine/Source/Runtime/SlateCore/Public/Input/Events.h b/Engine/Source/Runtime/SlateCore/Public/Input/Events.h
--- a/Engine/Source/Runtime/SlateCore/Public/Input/Events.h
+++ b/Engine/Source/Runtime/SlateCore/Public/Input/Events.h
@@ -56,7 +56,7 @@ namespace EKeyboardFocusCause
 USTRUCT()
 struct FFocusEvent
 {
-	GENERATED_BODY()
+	GENERATED_USTRUCT_BODY()
 
 public:
 
@@ -128,7 +128,7 @@ struct FVirtualPointerPosition
 USTRUCT(BlueprintType)
 struct FInputEvent
 {
-	GENERATED_BODY()
+	GENERATED_USTRUCT_BODY()
 
 public:
 
@@ -358,7 +358,7 @@ USTRUCT(BlueprintType)
 struct FKeyEvent
 	: public FInputEvent
 {
-	GENERATED_BODY()
+	GENERATED_USTRUCT_BODY()
 
 public:
 	/**
@@ -445,7 +445,7 @@ USTRUCT(BlueprintType)
 struct FKeyboardEvent : public FKeyEvent
 //struct DEPRECATED(4.6, "Use FKeyEvent") FKeyboardEvent : public FKeyEvent
 {
-	GENERATED_BODY()
+	GENERATED_USTRUCT_BODY()
 public:
 	FKeyEvent() {}
 
@@ -461,7 +461,7 @@ USTRUCT(BlueprintType)
 struct FAnalogInputEvent
 	: public FKeyEvent
 {
-	GENERATED_BODY()
+	GENERATED_USTRUCT_BODY()
 
 public:
 	/**
@@ -515,7 +515,7 @@ USTRUCT()
 struct FCharacterEvent
 	: public FInputEvent
 {
-	GENERATED_BODY()
+	GENERATED_USTRUCT_BODY()
 
 public:
 	/**
@@ -587,7 +587,7 @@ USTRUCT(BlueprintType)
 struct FPointerEvent
 	: public FInputEvent
 {
-	GENERATED_BODY()
+	GENERATED_USTRUCT_BODY()
 public:
 
 	/**
@@ -782,7 +782,7 @@ USTRUCT(BlueprintType)
 struct FControllerEvent
 	: public FInputEvent
 {
-	GENERATED_BODY()
+	GENERATED_USTRUCT_BODY()
 public:
 
 	FControllerEvent()
@@ -816,7 +816,7 @@ USTRUCT(BlueprintType)
 struct FMotionEvent
 	: public FInputEvent
 {
-	GENERATED_BODY()
+	GENERATED_USTRUCT_BODY()
 
 public:
 	/**
@@ -884,7 +884,7 @@ USTRUCT(BlueprintType)
 struct FNavigationEvent
 	: public FInputEvent
 {
-	GENERATED_BODY()
+	GENERATED_USTRUCT_BODY()
 
 public:
 	/**
diff --git a/Engine/Source/Runtime/UMG/Public/Animation/MovieScene2DTransformTrack.h b/Engine/Source/Runtime/UMG/Public/Animation/MovieScene2DTransformTrack.h
--- a/Engine/Source/Runtime/UMG/Public/Animation/MovieScene2DTransformTrack.h
+++ b/Engine/Source/Runtime/UMG/Public/Animation/MovieScene2DTransformTrack.h
@@ -20,9 +20,7 @@ struct F2DTransformKey
 UCLASS( MinimalAPI )
 class UMovieScene2DTransformTrack : public UMovieScenePropertyTrack
 {
-	GENERATED_BODY()
-public:
-	UMG_API UMovieScene2DTransformTrack(const FObjectInitializer& ObjectInitializer = FObjectInitializer::Get());
+	GENERATED_UCLASS_BODY()
 public:
 	/** UMovieSceneTrack interface */
 	virtual UMovieSceneSection* CreateNewSection() override;
@@ -42,4 +40,4 @@ public:
 
 	UMG_API bool AddKeyToSection( float Time, const F2DTransformKey& TransformKey );
 
-};
+};
\ No newline at end of file
diff --git a/Engine/Source/Runtime/UMG/Public/Animation/MovieSceneMarginSection.h b/Engine/Source/Runtime/UMG/Public/Animation/MovieSceneMarginSection.h
--- a/Engine/Source/Runtime/UMG/Public/Animation/MovieSceneMarginSection.h
+++ b/Engine/Source/Runtime/UMG/Public/Animation/MovieSceneMarginSection.h
@@ -14,9 +14,7 @@ struct FMarginKey;
 UCLASS(MinimalAPI)
 class UMovieSceneMarginSection : public UMovieSceneSection
 {
-	GENERATED_BODY()
-public:
-	UMG_API UMovieSceneMarginSection(const FObjectInitializer& ObjectInitializer = FObjectInitializer::Get());
+	GENERATED_UCLASS_BODY()
 public:
 	/** UMovieSceneSection interface */
 	virtual void MoveSection( float DeltaPosition ) override;
diff --git a/Engine/Source/Runtime/UMG/Public/Blueprint/AsyncTaskDownloadImage.h b/Engine/Source/Runtime/UMG/Public/Blueprint/AsyncTaskDownloadImage.h
--- a/Engine/Source/Runtime/UMG/Public/Blueprint/AsyncTaskDownloadImage.h
+++ b/Engine/Source/Runtime/UMG/Public/Blueprint/AsyncTaskDownloadImage.h
@@ -12,9 +12,7 @@ DECLARE_DYNAMIC_MULTICAST_DELEGATE_OneParam(FDownloadImageDelegate, UTexture2D*,
 UCLASS()
 class UMG_API UAsyncTaskDownloadImage : public UBlueprintAsyncActionBase
 {
-	GENERATED_BODY()
-public:
-	UAsyncTaskDownloadImage(const FObjectInitializer& ObjectInitializer = FObjectInitializer::Get());
+	GENERATED_UCLASS_BODY()
 
 public:
 	UFUNCTION(BlueprintCallable, meta=( BlueprintInternalUseOnly="true" ))
diff --git a/Engine/Source/Runtime/UMG/Public/Blueprint/UserWidget.h b/Engine/Source/Runtime/UMG/Public/Blueprint/UserWidget.h
--- a/Engine/Source/Runtime/UMG/Public/Blueprint/UserWidget.h
+++ b/Engine/Source/Runtime/UMG/Public/Blueprint/UserWidget.h
@@ -30,7 +30,7 @@ struct FLocalPlayerContext;
 USTRUCT(BlueprintType)
 struct UMG_API FPaintContext
 {
-	GENERATED_BODY()
+	GENERATED_USTRUCT_BODY()
 
 public:
 
@@ -83,7 +83,7 @@ public:
 USTRUCT()
 struct UMG_API FNamedSlotBinding
 {
-	GENERATED_BODY()
+	GENERATED_USTRUCT_BODY()
 
 public:
 
@@ -126,9 +126,7 @@ DECLARE_DYNAMIC_MULTICAST_DELEGATE(FOnConstructEvent);
 UCLASS(Abstract, editinlinenew, BlueprintType, Blueprintable, meta=( Category="User Controls" ))
 class UMG_API UUserWidget : public UWidget, public INamedSlotInterface
 {
-	GENERATED_BODY()
-public:
-	UUserWidget(const FObjectInitializer& ObjectInitializer = FObjectInitializer::Get());
+	GENERATED_UCLASS_BODY()
 
 public:
 	//UObject interface
@@ -255,7 +253,6 @@ public:
 	 */
 	UFUNCTION(BlueprintNativeEvent, BlueprintCosmetic, Category="User Interface", meta=( Keywords="Begin Play" ))
 	void Construct();
-	virtual void Construct_Implementation();
 
 	/**
 	 * Ticks this widget.  Override in derived classes, but always call the parent implementation.
@@ -265,11 +262,9 @@ public:
 	 */
 	UFUNCTION(BlueprintNativeEvent, BlueprintCosmetic, Category="User Interface")
 	void Tick(FGeometry MyGeometry, float InDeltaTime);
-	virtual void Tick_Implementation(FGeometry MyGeometry, float InDeltaTime);
 
 	UFUNCTION(BlueprintNativeEvent, BlueprintCosmetic, Category="User Interface | Painting")
 	void OnPaint(UPARAM(ref) FPaintContext& Context) const;
-	virtual void OnPaint_Implementation(FPaintContext& Context) const;
 
 	/**
 	 * Called when keyboard focus is given to this widget.  This event does not bubble.
@@ -280,7 +275,6 @@ public:
 	 */
 	UFUNCTION(BlueprintNativeEvent, BlueprintCosmetic, Category="Input")
 	FEventReply OnFocusReceived(FGeometry MyGeometry, FFocusEvent InFocusEvent);
-	virtual FEventReply OnFocusReceived_Implementation(FGeometry MyGeometry, FFocusEvent InFocusEvent);
 
 	/**
 	 * Called when this widget loses focus.  This event does not bubble.
@@ -289,7 +283,6 @@ public:
 	 */
 	UFUNCTION(BlueprintNativeEvent, BlueprintCosmetic, Category="Input")
 	void OnFocusLost(FFocusEvent InFocusEvent);
-	virtual void OnFocusLost_Implementation(FFocusEvent InFocusEvent);
 
 	/**
 	 * Called after a character is entered while this widget has focus
@@ -300,7 +293,6 @@ public:
 	 */
 	UFUNCTION(BlueprintNativeEvent, BlueprintCosmetic, Category="Input")
 	FEventReply OnKeyChar(FGeometry MyGeometry, FCharacterEvent InCharacterEvent);
-	virtual FEventReply OnKeyChar_Implementation(FGeometry MyGeometry, FCharacterEvent InCharacterEvent);
 
 	/**
 	 * Called after a key (keyboard, controller, ...) is pressed when this widget or a child of this widget has focus
@@ -315,7 +307,6 @@ public:
 	 */
 	UFUNCTION(BlueprintNativeEvent, Category="Input")
 	FEventReply OnPreviewKeyDown(FGeometry MyGeometry, FKeyEvent InKeyEvent);
-	virtual FEventReply OnPreviewKeyDown_Implementation(FGeometry MyGeometry, FKeyEvent InKeyEvent);
 
 	/**
 	 * Called after a key (keyboard, controller, ...) is pressed when this widget has focus (this event bubbles if not handled)
@@ -326,7 +317,6 @@ public:
 	 */
 	UFUNCTION(BlueprintNativeEvent, BlueprintCosmetic, Category="Input")
 	FEventReply OnKeyDown(FGeometry MyGeometry, FKeyEvent InKeyEvent);
-	virtual FEventReply OnKeyDown_Implementation(FGeometry MyGeometry, FKeyEvent InKeyEvent);
 
 	/**
 	 * Called after a key (keyboard, controller, ...) is released when this widget has focus
@@ -337,7 +327,6 @@ public:
 	 */
 	UFUNCTION(BlueprintNativeEvent, BlueprintCosmetic, Category="Input")
 	FEventReply OnKeyUp(FGeometry MyGeometry, FKeyEvent InKeyEvent);
-	virtual FEventReply OnKeyUp_Implementation(FGeometry MyGeometry, FKeyEvent InKeyEvent);
 
 	/**
 	* Called when an analog value changes on a button that supports analog
@@ -348,7 +337,6 @@ public:
 	*/
 	UFUNCTION(BlueprintNativeEvent, Category = "Input")
 	FEventReply OnAnalogValueChanged(FGeometry MyGeometry, FAnalogInputEvent InAnalogInputEvent);
-	virtual FEventReply OnAnalogValueChanged_Implementation(FGeometry MyGeometry, FAnalogInputEvent InAnalogInputEvent);
 
 	/**
 	 * The system calls this method to notify the widget that a mouse button was pressed within it. This event is bubbled.
@@ -359,7 +347,6 @@ public:
 	 */
 	UFUNCTION(BlueprintNativeEvent, BlueprintCosmetic, Category="Mouse")
 	FEventReply OnMouseButtonDown(FGeometry MyGeometry, const FPointerEvent& MouseEvent);
-	virtual FEventReply OnMouseButtonDown_Implementation(FGeometry MyGeometry, const FPointerEvent& MouseEvent);
 
 	/**
 	 * Just like OnMouseButtonDown, but tunnels instead of bubbling.
@@ -374,7 +361,6 @@ public:
 	 */
 	UFUNCTION(BlueprintNativeEvent, BlueprintCosmetic, Category="Mouse")
 	FEventReply OnPreviewMouseButtonDown(FGeometry MyGeometry, const FPointerEvent& MouseEvent);
-	virtual FEventReply OnPreviewMouseButtonDown_Implementation(FGeometry MyGeometry, const FPointerEvent& MouseEvent);
 
 	/**
 	 * The system calls this method to notify the widget that a mouse button was release within it. This event is bubbled.
@@ -385,7 +371,6 @@ public:
 	 */
 	UFUNCTION(BlueprintNativeEvent, BlueprintCosmetic, Category="Mouse")
 	FEventReply OnMouseButtonUp(FGeometry MyGeometry, const FPointerEvent& MouseEvent);
-	virtual FEventReply OnMouseButtonUp_Implementation(FGeometry MyGeometry, const FPointerEvent& MouseEvent);
 
 	/**
 	 * The system calls this method to notify the widget that a mouse moved within it. This event is bubbled.
@@ -396,7 +381,6 @@ public:
 	 */
 	UFUNCTION(BlueprintNativeEvent, BlueprintCosmetic, Category="Mouse")
 	FEventReply OnMouseMove(FGeometry MyGeometry, const FPointerEvent& MouseEvent);
-	virtual FEventReply OnMouseMove_Implementation(FGeometry MyGeometry, const FPointerEvent& MouseEvent);
 
 	/**
 	 * The system will use this event to notify a widget that the cursor has entered it. This event is NOT bubbled.
@@ -406,7 +390,6 @@ public:
 	 */
 	UFUNCTION(BlueprintNativeEvent, BlueprintCosmetic, Category="Mouse")
 	void OnMouseEnter(FGeometry MyGeometry, const FPointerEvent& MouseEvent);
-	virtual void OnMouseEnter_Implementation(FGeometry MyGeometry, const FPointerEvent& MouseEvent);
 
 	/**
 	 * The system will use this event to notify a widget that the cursor has left it. This event is NOT bubbled.
@@ -415,7 +398,6 @@ public:
 	 */
 	UFUNCTION(BlueprintNativeEvent, BlueprintCosmetic, Category="Mouse")
 	void OnMouseLeave(const FPointerEvent& MouseEvent);
-	virtual void OnMouseLeave_Implementation(const FPointerEvent& MouseEvent);
 
 	/**
 	 * Called when the mouse wheel is spun. This event is bubbled.
@@ -425,7 +407,6 @@ public:
 	 */
 	UFUNCTION(BlueprintNativeEvent, BlueprintCosmetic, Category="Mouse")
 	FEventReply OnMouseWheel(FGeometry MyGeometry, const FPointerEvent& MouseEvent);
-	virtual FEventReply OnMouseWheel_Implementation(FGeometry MyGeometry, const FPointerEvent& MouseEvent);
 
 	/**
 	 * Called when a mouse button is double clicked.  Override this in derived classes.
@@ -436,7 +417,6 @@ public:
 	 */
 	UFUNCTION(BlueprintNativeEvent, BlueprintCosmetic, Category="Mouse")
 	FEventReply OnMouseButtonDoubleClick(FGeometry InMyGeometry, const FPointerEvent& InMouseEvent);
-	virtual FEventReply OnMouseButtonDoubleClick_Implementation(FGeometry InMyGeometry, const FPointerEvent& InMouseEvent);
 
 	// TODO
 	//UFUNCTION(BlueprintNativeEvent, Category="Mouse")
@@ -454,7 +434,6 @@ public:
 	 */
 	UFUNCTION(BlueprintNativeEvent, BlueprintCosmetic, Category="Drag and Drop")
 	void OnDragDetected(FGeometry MyGeometry, const FPointerEvent& PointerEvent, UDragDropOperation*& Operation);
-	virtual void OnDragDetected_Implementation(FGeometry MyGeometry, const FPointerEvent& PointerEvent, UDragDropOperation*& Operation);
 
 	/**
 	 * Called when the user cancels the drag operation, typically when they simply release the mouse button after
@@ -465,7 +444,6 @@ public:
 	 */
 	UFUNCTION(BlueprintNativeEvent, BlueprintCosmetic, Category="Drag and Drop")
 	void OnDragCancelled(const FPointerEvent& PointerEvent, UDragDropOperation* Operation);
-	virtual void OnDragCancelled_Implementation(const FPointerEvent& PointerEvent, UDragDropOperation* Operation);
 	
 	/**
 	 * Called during drag and drop when the drag enters the widget.
@@ -476,7 +454,6 @@ public:
 	 */
 	UFUNCTION(BlueprintNativeEvent, BlueprintCosmetic, Category="Drag and Drop")
 	void OnDragEnter(FGeometry MyGeometry, FPointerEvent PointerEvent, UDragDropOperation* Operation);
-	virtual void OnDragEnter_Implementation(FGeometry MyGeometry, FPointerEvent PointerEvent, UDragDropOperation* Operation);
 
 	/**
 	 * Called during drag and drop when the drag leaves the widget.
@@ -486,7 +463,6 @@ public:
 	 */
 	UFUNCTION(BlueprintNativeEvent, BlueprintCosmetic, Category="Drag and Drop")
 	void OnDragLeave(FPointerEvent PointerEvent, UDragDropOperation* Operation);
-	virtual void OnDragLeave_Implementation(FPointerEvent PointerEvent, UDragDropOperation* Operation);
 
 	/**
 	 * Called during drag and drop when the the mouse is being dragged over a widget.
@@ -499,7 +475,6 @@ public:
 	 */
 	UFUNCTION(BlueprintNativeEvent, BlueprintCosmetic, Category="Drag and Drop")
 	bool OnDragOver(FGeometry MyGeometry, FPointerEvent PointerEvent, UDragDropOperation* Operation);
-	virtual bool OnDragOver_Implementation(FGeometry MyGeometry, FPointerEvent PointerEvent, UDragDropOperation* Operation);
 
 	/**
 	 * Called when the user is dropping something onto a widget.  Ends the drag and drop operation, even if no widget handles this.
@@ -512,20 +487,16 @@ public:
 	 */
 	UFUNCTION(BlueprintNativeEvent, BlueprintCosmetic, Category="Drag and Drop")
 	bool OnDrop(FGeometry MyGeometry, FPointerEvent PointerEvent, UDragDropOperation* Operation);
-	virtual bool OnDrop_Implementation(FGeometry MyGeometry, FPointerEvent PointerEvent, UDragDropOperation* Operation);
 
 
 	UFUNCTION(BlueprintNativeEvent, meta = (DeprecatedFunction, DeprecationMessage = "Use OnKeyDown() instead"), Category = "Gamepad Input")
 	FEventReply OnControllerButtonPressed(FGeometry MyGeometry, FControllerEvent ControllerEvent);
-	virtual FEventReply OnControllerButtonPressed_Implementation(FGeometry MyGeometry, FControllerEvent ControllerEvent);
 
 	UFUNCTION(BlueprintNativeEvent, meta = (DeprecatedFunction, DeprecationMessage = "Use OnKeyUp() instead"), Category = "Gamepad Input")
 	FEventReply OnControllerButtonReleased(FGeometry MyGeometry, FControllerEvent ControllerEvent);
-	virtual FEventReply OnControllerButtonReleased_Implementation(FGeometry MyGeometry, FControllerEvent ControllerEvent);
 
 	UFUNCTION(BlueprintNativeEvent, meta = (DeprecatedFunction, DeprecationMessage = "Use OnAnalogValueChanged() instead"), Category = "Gamepad Input")
 	FEventReply OnControllerAnalogValueChanged(FGeometry MyGeometry, FControllerEvent ControllerEvent);
-	virtual FEventReply OnControllerAnalogValueChanged_Implementation(FGeometry MyGeometry, FControllerEvent ControllerEvent);
 
 	/**
 	 * Called when the user performs a gesture on trackpad. This event is bubbled.
@@ -536,7 +507,6 @@ public:
 	 */
 	UFUNCTION(BlueprintNativeEvent, BlueprintCosmetic, Category="Touch Input")
 	FEventReply OnTouchGesture(FGeometry MyGeometry, const FPointerEvent& GestureEvent);
-	virtual FEventReply OnTouchGesture_Implementation(FGeometry MyGeometry, const FPointerEvent& GestureEvent);
 
 	/**
 	 * Called when a touchpad touch is started (finger down)
@@ -546,7 +516,6 @@ public:
 	 */
 	UFUNCTION(BlueprintNativeEvent, BlueprintCosmetic, Category="Touch Input")
 	FEventReply OnTouchStarted(FGeometry MyGeometry, const FPointerEvent& InTouchEvent);
-	virtual FEventReply OnTouchStarted_Implementation(FGeometry MyGeometry, const FPointerEvent& InTouchEvent);
 
 	/**
 	 * Called when a touchpad touch is moved  (finger moved)
@@ -556,7 +525,6 @@ public:
 	 */
 	UFUNCTION(BlueprintNativeEvent, BlueprintCosmetic, Category="Touch Input")
 	FEventReply OnTouchMoved(FGeometry MyGeometry, const FPointerEvent& InTouchEvent);
-	virtual FEventReply OnTouchMoved_Implementation(FGeometry MyGeometry, const FPointerEvent& InTouchEvent);
 
 	/**
 	 * Called when a touchpad touch is ended (finger lifted)
@@ -566,7 +534,6 @@ public:
 	 */
 	UFUNCTION(BlueprintNativeEvent, BlueprintCosmetic, Category="Touch Input")
 	FEventReply OnTouchEnded(FGeometry MyGeometry, const FPointerEvent& InTouchEvent);
-	virtual FEventReply OnTouchEnded_Implementation(FGeometry MyGeometry, const FPointerEvent& InTouchEvent);
 
 	/**
 	 * Called when motion is detected (controller or device)
@@ -577,7 +544,6 @@ public:
 	 */
 	UFUNCTION(BlueprintNativeEvent, BlueprintCosmetic, Category="Touch Input")
 	FEventReply OnMotionDetected(FGeometry MyGeometry, FMotionEvent InMotionEvent);
-	virtual FEventReply OnMotionDetected_Implementation(FGeometry MyGeometry, FMotionEvent InMotionEvent);
 
 	/**
 	 * Called when an animation is started.
@@ -586,7 +552,6 @@ public:
 	 */
 	UFUNCTION( BlueprintNativeEvent, BlueprintCosmetic, Category = "Animation" )
 	void OnAnimationStarted( const UWidgetAnimation* Animation );
-	virtual void OnAnimationStarted_Implementation(const UWidgetAnimation* Animation);
 
 	/**
 	 * Called when an animation has either played all the way through or is stopped
@@ -595,7 +560,6 @@ public:
 	 */
 	UFUNCTION( BlueprintNativeEvent, BlueprintCosmetic, Category = "Animation" )
 	void OnAnimationFinished( const UWidgetAnimation* Animation );
-	virtual void OnAnimationFinished_Implementation(const UWidgetAnimation* Animation);
 
 public:
 
diff --git a/Engine/Source/Runtime/UMG/Public/Blueprint/WidgetLayoutLibrary.h b/Engine/Source/Runtime/UMG/Public/Blueprint/WidgetLayoutLibrary.h
--- a/Engine/Source/Runtime/UMG/Public/Blueprint/WidgetLayoutLibrary.h
+++ b/Engine/Source/Runtime/UMG/Public/Blueprint/WidgetLayoutLibrary.h
@@ -7,9 +7,7 @@
 UCLASS(MinimalAPI)
 class UWidgetLayoutLibrary : public UBlueprintFunctionLibrary
 {
-	GENERATED_BODY()
-public:
-	UMG_API UWidgetLayoutLibrary(const FObjectInitializer& ObjectInitializer = FObjectInitializer::Get());
+	GENERATED_UCLASS_BODY()
 
 public:
 
diff --git a/Engine/Source/Runtime/UMG/Public/Components/ScaleBoxSlot.h b/Engine/Source/Runtime/UMG/Public/Components/ScaleBoxSlot.h
--- a/Engine/Source/Runtime/UMG/Public/Components/ScaleBoxSlot.h
+++ b/Engine/Source/Runtime/UMG/Public/Components/ScaleBoxSlot.h
@@ -10,9 +10,7 @@
 UCLASS()
 class UMG_API UScaleBoxSlot : public UPanelSlot
 {
-	GENERATED_BODY()
-public:
-	UScaleBoxSlot(const FObjectInitializer& ObjectInitializer = FObjectInitializer::Get());
+	GENERATED_UCLASS_BODY()
 
 public:
 	
diff --git a/Engine/Source/Runtime/UMG/Public/Components/ScrollBox.h b/Engine/Source/Runtime/UMG/Public/Components/ScrollBox.h
--- a/Engine/Source/Runtime/UMG/Public/Components/ScrollBox.h
+++ b/Engine/Source/Runtime/UMG/Public/Components/ScrollBox.h
@@ -10,9 +10,7 @@
 UCLASS()
 class UMG_API UScrollBox : public UPanelWidget
 {
-	GENERATED_BODY()
-public:
-	UScrollBox(const FObjectInitializer& ObjectInitializer = FObjectInitializer::Get());
+	GENERATED_UCLASS_BODY()
 
 public:
 
diff --git a/Engine/Source/Runtime/UMG/Public/Components/Spacer.h b/Engine/Source/Runtime/UMG/Public/Components/Spacer.h
--- a/Engine/Source/Runtime/UMG/Public/Components/Spacer.h
+++ b/Engine/Source/Runtime/UMG/Public/Components/Spacer.h
@@ -1,4 +1,4 @@
-// Copyright 1998-2015 Epic Games, Inc. All Rights Reserved.
+﻿// Copyright 1998-2015 Epic Games, Inc. All Rights Reserved.
 
 #pragma once
 
@@ -12,9 +12,7 @@
 UCLASS()
 class UMG_API USpacer : public UWidget
 {
-	GENERATED_BODY()
-public:
-	USpacer(const FObjectInitializer& ObjectInitializer = FObjectInitializer::Get());
+	GENERATED_UCLASS_BODY()
 
 public:
 
diff --git a/Engine/Source/Runtime/UMG/Public/Components/Throbber.h b/Engine/Source/Runtime/UMG/Public/Components/Throbber.h
--- a/Engine/Source/Runtime/UMG/Public/Components/Throbber.h
+++ b/Engine/Source/Runtime/UMG/Public/Components/Throbber.h
@@ -13,9 +13,7 @@ class USlateBrushAsset;
 UCLASS(ClassGroup=UserInterface)
 class UMG_API UThrobber : public UWidget
 {
-	GENERATED_BODY()
-public:
-	UThrobber(const FObjectInitializer& ObjectInitializer = FObjectInitializer::Get());
+	GENERATED_UCLASS_BODY()
 
 public:
 	
diff --git a/Engine/Source/Runtime/UMG/Public/Components/UniformGridSlot.h b/Engine/Source/Runtime/UMG/Public/Components/UniformGridSlot.h
--- a/Engine/Source/Runtime/UMG/Public/Components/UniformGridSlot.h
+++ b/Engine/Source/Runtime/UMG/Public/Components/UniformGridSlot.h
@@ -13,9 +13,7 @@
 UCLASS()
 class UMG_API UUniformGridSlot : public UPanelSlot
 {
-	GENERATED_BODY()
-public:
-	UUniformGridSlot(const FObjectInitializer& ObjectInitializer = FObjectInitializer::Get());
+	GENERATED_UCLASS_BODY()
 
 public:
 
@@ -69,4 +67,4 @@ private:
 
 	/** A raw pointer to the slot to allow us to adjust the size, padding...etc at runtime. */
 	SUniformGridPanel::FSlot* Slot;
-};
+};
\ No newline at end of file
