UMG - PR #962: Scrollbox Update (Contributed by [REDACTED_PERSON]), additional small cleanup to make it const.

diff --git a/Engine/Source/Runtime/UMG/Private/Components/ScrollBox.cpp b/Engine/Source/Runtime/UMG/Private/Components/ScrollBox.cpp
--- a/Engine/Source/Runtime/UMG/Private/Components/ScrollBox.cpp
+++ b/Engine/Source/Runtime/UMG/Private/Components/ScrollBox.cpp
@@ -90,6 +90,16 @@ void UScrollBox::SynchronizeProperties()
 	MyScrollBox->SetScrollBarAlwaysVisible(AlwaysShowScrollbar);
 }
 
+float UScrollBox::GetScrollOffset() const
+{
+	if ( MyScrollBox.IsValid() )
+	{
+		return MyScrollBox->GetScrollOffset();
+	}
+
+	return 0;
+}
+
 void UScrollBox::SetScrollOffset(float NewScrollOffset)
 {
 	DesiredScrollOffset = NewScrollOffset;
diff --git a/Engine/Source/Runtime/UMG/Public/Components/ScrollBox.h b/Engine/Source/Runtime/UMG/Public/Components/ScrollBox.h
--- a/Engine/Source/Runtime/UMG/Public/Components/ScrollBox.h
+++ b/Engine/Source/Runtime/UMG/Public/Components/ScrollBox.h
@@ -61,6 +61,12 @@ public:
 	 */
 	UFUNCTION(BlueprintCallable, Category="Widget")
 	void SetScrollOffset(float NewScrollOffset);
+	
+	/**
+	 * Gets the scroll offset of the scrollbox in Slate Units.
+	 */
+	UFUNCTION(BlueprintCallable, Category="Widget")
+	float GetScrollOffset() const;
 
 	/** Scrolls the ScrollBox to the top instantly */
 	UFUNCTION(BlueprintCallable, Category="Widget")
