Merging CL 2488532

//depot/UE4-UT/...

to //depot/UE4/...

diff --git a/Engine/Source/Runtime/UMG/Public/Components/ComboBoxString.h b/Engine/Source/Runtime/UMG/Public/Components/ComboBoxString.h
--- a/Engine/Source/Runtime/UMG/Public/Components/ComboBoxString.h
+++ b/Engine/Source/Runtime/UMG/Public/Components/ComboBoxString.h
@@ -111,13 +111,13 @@ public:
 
 protected:
 	/** Called by slate when it needs to generate a new item for the combobox */
-	TSharedRef<SWidget> HandleGenerateWidget(TSharedPtr<FString> Item) const;
+	virtual TSharedRef<SWidget> HandleGenerateWidget(TSharedPtr<FString> Item) const;
 
 	/** Called by slate when the underlying comobobox selection changes */
-	void HandleSelectionChanged(TSharedPtr<FString> Item, ESelectInfo::Type SelectionType);
+	virtual void HandleSelectionChanged(TSharedPtr<FString> Item, ESelectInfo::Type SelectionType);
 
 	/** Called by slate when the underlying comobobox is opening */
-	void HandleOpening();
+	virtual void HandleOpening();
 
 	// UWidget interface
 	virtual TSharedRef<SWidget> RebuildWidget() override;
