Fixed Layout category labels on Blueprint-visible functions on UMG classes
[UE-9414] - Blueprint Palette Layout options are not all under Layout

diff --git a/Engine/Source/Runtime/UMG/Public/Components/ButtonSlot.h b/Engine/Source/Runtime/UMG/Public/Components/ButtonSlot.h
--- a/Engine/Source/Runtime/UMG/Public/Components/ButtonSlot.h
+++ b/Engine/Source/Runtime/UMG/Public/Components/ButtonSlot.h
@@ -15,26 +15,26 @@ class UMG_API UButtonSlot : public UPanelSlot
 public:
 	
 	/** The padding area between the slot and the content it contains. */
-	UPROPERTY(EditDefaultsOnly, BlueprintReadOnly, Category="Layout (Button Slot)")
+	UPROPERTY(EditDefaultsOnly, BlueprintReadOnly, Category="Layout|Button Slot")
 	FMargin Padding;
 
 	/** The alignment of the object horizontally. */
-	UPROPERTY(EditDefaultsOnly, BlueprintReadOnly, Category="Layout (Button Slot)")
+	UPROPERTY(EditDefaultsOnly, BlueprintReadOnly, Category="Layout|Button Slot")
 	TEnumAsByte<EHorizontalAlignment> HorizontalAlignment;
 
 	/** The alignment of the object vertically. */
-	UPROPERTY(EditDefaultsOnly, BlueprintReadOnly, Category="Layout (Button Slot)")
+	UPROPERTY(EditDefaultsOnly, BlueprintReadOnly, Category="Layout|Button Slot")
 	TEnumAsByte<EVerticalAlignment> VerticalAlignment;
 
 public:
 
-	UFUNCTION(BlueprintCallable, Category="Layout (Button Slot)")
+	UFUNCTION(BlueprintCallable, Category="Layout|Button Slot")
 	void SetPadding(FMargin InPadding);
 
-	UFUNCTION(BlueprintCallable, Category="Layout (Button Slot)")
+	UFUNCTION(BlueprintCallable, Category="Layout|Button Slot")
 	void SetHorizontalAlignment(EHorizontalAlignment InHorizontalAlignment);
 
-	UFUNCTION(BlueprintCallable, Category="Layout (Button Slot)")
+	UFUNCTION(BlueprintCallable, Category="Layout|Button Slot")
 	void SetVerticalAlignment(EVerticalAlignment InVerticalAlignment);
 
 	// UPanelSlot interface
diff --git a/Engine/Source/Runtime/UMG/Public/Components/CanvasPanelSlot.h b/Engine/Source/Runtime/UMG/Public/Components/CanvasPanelSlot.h
--- a/Engine/Source/Runtime/UMG/Public/Components/CanvasPanelSlot.h
+++ b/Engine/Source/Runtime/UMG/Public/Components/CanvasPanelSlot.h
@@ -36,80 +36,80 @@ class UMG_API UCanvasPanelSlot : public UPanelSlot
 public:
 
 	/** The anchoring information for the slot */
-	UPROPERTY(EditDefaultsOnly, BlueprintReadOnly, Category="Layout (Canvas Slot)")
+	UPROPERTY(EditDefaultsOnly, BlueprintReadOnly, Category="Layout|Canvas Slot")
 	FAnchorData LayoutData;
 
 	/** When AutoSize is true we use the widget's desired size */
-	UPROPERTY(EditDefaultsOnly, BlueprintReadOnly, Category="Layout (Canvas Slot)", AdvancedDisplay, meta=( DisplayName="Size To Content" ))
+	UPROPERTY(EditDefaultsOnly, BlueprintReadOnly, Category="Layout|Canvas Slot", AdvancedDisplay, meta=( DisplayName="Size To Content" ))
 	bool bAutoSize;
 
 	/** The order priority this widget is rendered in.  Higher values are rendered last (and so they will appear to be on top). */
-	UPROPERTY(EditDefaultsOnly, BlueprintReadOnly, Category="Layout (Canvas Slot)")
+	UPROPERTY(EditDefaultsOnly, BlueprintReadOnly, Category="Layout|Canvas Slot")
 	int32 ZOrder;
 
 public:
 	/** Sets the layout data of the slot */
-	UFUNCTION(BlueprintCallable, Category="Layout (Canvas Slot)")
+	UFUNCTION(BlueprintCallable, Category="Layout|Canvas Slot")
 	void SetLayout(const FAnchorData& InLayoutData);
 
 	/** Gets the layout data of the slot */
-	UFUNCTION(BlueprintCallable, Category="Layout (Canvas Slot)")
+	UFUNCTION(BlueprintCallable, Category="Layout|Canvas Slot")
 	FAnchorData GetLayout() const;
 
 	/** Sets the position of the slot */
-	UFUNCTION(BlueprintCallable, Category="Layout (Canvas Slot)")
+	UFUNCTION(BlueprintCallable, Category="Layout|Canvas Slot")
 	void SetPosition(FVector2D InPosition);
 
 	/** Gets the position of the slot */
-	UFUNCTION(BlueprintCallable, Category="Layout (Canvas Slot)")
+	UFUNCTION(BlueprintCallable, Category="Layout|Canvas Slot")
 	FVector2D GetPosition() const;
 
 	/** Sets the size of the slot */
-	UFUNCTION(BlueprintCallable, Category="Layout (Canvas Slot)")
+	UFUNCTION(BlueprintCallable, Category="Layout|Canvas Slot")
 	void SetSize(FVector2D InSize);
 
 	/** Gets the size of the slot */
-	UFUNCTION(BlueprintCallable, Category="Layout (Canvas Slot)")
+	UFUNCTION(BlueprintCallable, Category="Layout|Canvas Slot")
 	FVector2D GetSize() const;
 
 	/** Sets the offset data of the slot, which could be position and size, or margins depending on the anchor points */
-	UFUNCTION(BlueprintCallable, Category="Layout (Canvas Slot)")
+	UFUNCTION(BlueprintCallable, Category="Layout|Canvas Slot")
 	void SetOffsets(FMargin InOffset);
 
 	/** Gets the offset data of the slot, which could be position and size, or margins depending on the anchor points */
-	UFUNCTION(BlueprintCallable, Category="Layout (Canvas Slot)")
+	UFUNCTION(BlueprintCallable, Category="Layout|Canvas Slot")
 	FMargin GetOffsets() const;
 	
 	/** Sets the anchors on the slot */
-	UFUNCTION(BlueprintCallable, Category="Layout (Canvas Slot)")
+	UFUNCTION(BlueprintCallable, Category="Layout|Canvas Slot")
 	void SetAnchors(FAnchors InAnchors);
 
 	/** Gets the anchors on the slot */
-	UFUNCTION(BlueprintCallable, Category="Layout (Canvas Slot)")
+	UFUNCTION(BlueprintCallable, Category="Layout|Canvas Slot")
 	FAnchors GetAnchors() const;
 
 	/** Sets the alignment on the slot */
-	UFUNCTION(BlueprintCallable, Category="Layout (Canvas Slot)")
+	UFUNCTION(BlueprintCallable, Category="Layout|Canvas Slot")
 	void SetAlignment(FVector2D InAlignment);
 
 	/** Gets the alignment on the slot */
-	UFUNCTION(BlueprintCallable, Category="Layout (Canvas Slot)")
+	UFUNCTION(BlueprintCallable, Category="Layout|Canvas Slot")
 	FVector2D GetAlignment() const;
 
 	/** Sets if the slot to be auto-sized */
-	UFUNCTION(BlueprintCallable, Category="Layout (Canvas Slot)")
+	UFUNCTION(BlueprintCallable, Category="Layout|Canvas Slot")
 	void SetAutoSize(bool InbAutoSize);
 
 	/** Gets if the slot to be auto-sized */
-	UFUNCTION(BlueprintCallable, Category="Layout (Canvas Slot)")
+	UFUNCTION(BlueprintCallable, Category="Layout|Canvas Slot")
 	bool GetAutoSize() const;
 
 	/** Sets the z-order on the slot */
-	UFUNCTION(BlueprintCallable, Category="Layout (Canvas Slot)")
+	UFUNCTION(BlueprintCallable, Category="Layout|Canvas Slot")
 	void SetZOrder(int32 InZOrder);
 
 	/** Gets the z-order on the slot */
-	UFUNCTION(BlueprintCallable, Category="Layout (Canvas Slot)")
+	UFUNCTION(BlueprintCallable, Category="Layout|Canvas Slot")
 	int32 GetZOrder() const;
 
 public:
diff --git a/Engine/Source/Runtime/UMG/Public/Components/GridSlot.h b/Engine/Source/Runtime/UMG/Public/Components/GridSlot.h
--- a/Engine/Source/Runtime/UMG/Public/Components/GridSlot.h
+++ b/Engine/Source/Runtime/UMG/Public/Components/GridSlot.h
@@ -18,59 +18,59 @@ class UMG_API UGridSlot : public UPanelSlot
 public:
 
 	/** The alignment of the object horizontally. */
-	UPROPERTY(EditDefaultsOnly, BlueprintReadOnly, Category="Layout (Grid Slot)")
+	UPROPERTY(EditDefaultsOnly, BlueprintReadOnly, Category="Layout|Grid Slot")
 	TEnumAsByte<EHorizontalAlignment> HorizontalAlignment;
 
 	/** The alignment of the object vertically. */
-	UPROPERTY(EditDefaultsOnly, BlueprintReadOnly, Category="Layout (Grid Slot)")
+	UPROPERTY(EditDefaultsOnly, BlueprintReadOnly, Category="Layout|Grid Slot")
 	TEnumAsByte<EVerticalAlignment> VerticalAlignment;
 	
 	/** The row index of the cell this slot is in */
-	UPROPERTY(EditDefaultsOnly, BlueprintReadOnly, meta=( UIMin = "0" ), Category="Layout (Grid Slot)")
+	UPROPERTY(EditDefaultsOnly, BlueprintReadOnly, meta=( UIMin = "0" ), Category="Layout|Grid Slot")
 	int32 Row;
 	
 	/**  */
-	UPROPERTY(EditDefaultsOnly, BlueprintReadOnly, Category="Layout (Grid Slot)")
+	UPROPERTY(EditDefaultsOnly, BlueprintReadOnly, Category="Layout|Grid Slot")
 	int32 RowSpan;
 
 	/** The column index of the cell this slot is in */
-	UPROPERTY(EditDefaultsOnly, BlueprintReadOnly, meta=( UIMin = "0" ), Category="Layout (Grid Slot)")
+	UPROPERTY(EditDefaultsOnly, BlueprintReadOnly, meta=( UIMin = "0" ), Category="Layout|Grid Slot")
 	int32 Column;
 
 	/**  */
-	UPROPERTY(EditDefaultsOnly, BlueprintReadOnly, Category="Layout (Grid Slot)")
+	UPROPERTY(EditDefaultsOnly, BlueprintReadOnly, Category="Layout|Grid Slot")
 	int32 ColumnSpan;
 
 	/** Positive values offset this cell to be hit-tested and drawn on top of others. Default is 0; i.e. no offset. */
-	UPROPERTY(EditDefaultsOnly, BlueprintReadOnly, Category="Layout (Grid Slot)")
+	UPROPERTY(EditDefaultsOnly, BlueprintReadOnly, Category="Layout|Grid Slot")
 	int32 Layer;
 
 	/** Offset this slot's content by some amount; positive values offset to lower right*/
-	UPROPERTY(EditDefaultsOnly, BlueprintReadOnly, Category="Layout (Grid Slot)")
+	UPROPERTY(EditDefaultsOnly, BlueprintReadOnly, Category="Layout|Grid Slot")
 	FVector2D Nudge;
 
 	/** Sets the row index of the slot, this determines what cell the slot is in the panel */
-	UFUNCTION(BlueprintCallable, Category="Layout (Grid Slot)")
+	UFUNCTION(BlueprintCallable, Category="Layout|Grid Slot")
 	void SetRow(int32 InRow);
 
 	/** How many rows this this slot spans over */
-	UFUNCTION(BlueprintCallable, Category="Layout (Grid Slot)")
+	UFUNCTION(BlueprintCallable, Category="Layout|Grid Slot")
 	void SetRowSpan(int32 InRowSpan);
 
 	/** Sets the column index of the slot, this determines what cell the slot is in the panel */
-	UFUNCTION(BlueprintCallable, Category="Layout (Grid Slot)")
+	UFUNCTION(BlueprintCallable, Category="Layout|Grid Slot")
 	void SetColumn(int32 InColumn);
 
 	/** How many columns this slot spans over */
-	UFUNCTION(BlueprintCallable, Category="Layout (Grid Slot)")
+	UFUNCTION(BlueprintCallable, Category="Layout|Grid Slot")
 	void SetColumnSpan(int32 InColumnSpan);
 
 	/**  */
-	UFUNCTION(BlueprintCallable, Category="Layout (Grid Slot)")
+	UFUNCTION(BlueprintCallable, Category="Layout|Grid Slot")
 	void SetHorizontalAlignment(EHorizontalAlignment InHorizontalAlignment);
 
 	/**  */
-	UFUNCTION(BlueprintCallable, Category="Layout (Grid Slot)")
+	UFUNCTION(BlueprintCallable, Category="Layout|Grid Slot")
 	void SetVerticalAlignment(EVerticalAlignment InVerticalAlignment);
 
 	// UPanelSlot interface
diff --git a/Engine/Source/Runtime/UMG/Public/Components/HorizontalBoxSlot.h b/Engine/Source/Runtime/UMG/Public/Components/HorizontalBoxSlot.h
--- a/Engine/Source/Runtime/UMG/Public/Components/HorizontalBoxSlot.h
+++ b/Engine/Source/Runtime/UMG/Public/Components/HorizontalBoxSlot.h
@@ -14,31 +14,31 @@ class UMG_API UHorizontalBoxSlot : public UPanelSlot
 public:
 	
 	/** The amount of padding between the slots parent and the content. */
-	UPROPERTY(EditDefaultsOnly, BlueprintReadOnly, Category="Layout (Horizontal Box Slot)")
+	UPROPERTY(EditDefaultsOnly, BlueprintReadOnly, Category="Layout|Horizontal Box Slot")
 	FMargin Padding;
 
 	/** How much space this slot should occupy in the direction of the panel. */
-	UPROPERTY(EditDefaultsOnly, BlueprintReadOnly, Category="Layout (Horizontal Box Slot)")
+	UPROPERTY(EditDefaultsOnly, BlueprintReadOnly, Category="Layout|Horizontal Box Slot")
 	FSlateChildSize Size;
 
-	UPROPERTY(EditDefaultsOnly, BlueprintReadOnly, Category="Layout (Horizontal Box Slot)")
+	UPROPERTY(EditDefaultsOnly, BlueprintReadOnly, Category="Layout|Horizontal Box Slot")
 	TEnumAsByte<EHorizontalAlignment> HorizontalAlignment;
 
-	UPROPERTY(EditDefaultsOnly, BlueprintReadOnly, Category="Layout (Horizontal Box Slot)")
+	UPROPERTY(EditDefaultsOnly, BlueprintReadOnly, Category="Layout|Horizontal Box Slot")
 	TEnumAsByte<EVerticalAlignment> VerticalAlignment;
 
 public:
 
-	UFUNCTION(BlueprintCallable, Category="Layout(Horizontal Box Slot)")
+	UFUNCTION(BlueprintCallable, Category="Layout|Horizontal Box Slot")
 	void SetPadding(FMargin InPadding);
 
-	UFUNCTION(BlueprintCallable, Category="Layout (Horizontal Box Slot)")
+	UFUNCTION(BlueprintCallable, Category="Layout|Horizontal Box Slot")
 	void SetSize(FSlateChildSize InSize);
 
-	UFUNCTION(BlueprintCallable, Category="Layout (Horizontal Box Slot)")
+	UFUNCTION(BlueprintCallable, Category="Layout|Horizontal Box Slot")
 	void SetHorizontalAlignment(EHorizontalAlignment InHorizontalAlignment);
 
-	UFUNCTION(BlueprintCallable, Category="Layout (Horizontal Box Slot)")
+	UFUNCTION(BlueprintCallable, Category="Layout|Horizontal Box Slot")
 	void SetVerticalAlignment(EVerticalAlignment InVerticalAlignment);
 
 	void BuildSlot(TSharedRef<SHorizontalBox> HorizontalBox);
diff --git a/Engine/Source/Runtime/UMG/Public/Components/OverlaySlot.h b/Engine/Source/Runtime/UMG/Public/Components/OverlaySlot.h
--- a/Engine/Source/Runtime/UMG/Public/Components/OverlaySlot.h
+++ b/Engine/Source/Runtime/UMG/Public/Components/OverlaySlot.h
@@ -17,26 +17,26 @@ class UMG_API UOverlaySlot : public UPanelSlot
 public:
 	
 	/** The padding area between the slot and the content it contains. */
-	UPROPERTY(EditDefaultsOnly, BlueprintReadOnly, Category="Layout (Overlay Slot)")
+	UPROPERTY(EditDefaultsOnly, BlueprintReadOnly, Category="Layout|Overlay Slot")
 	FMargin Padding;
 
 	/** The alignment of the object horizontally. */
-	UPROPERTY(EditDefaultsOnly, BlueprintReadOnly, Category="Layout (Overlay Slot)")
+	UPROPERTY(EditDefaultsOnly, BlueprintReadOnly, Category="Layout|Overlay Slot")
 	TEnumAsByte<EHorizontalAlignment> HorizontalAlignment;
 
 	/** The alignment of the object vertically. */
-	UPROPERTY(EditDefaultsOnly, BlueprintReadOnly, Category="Layout (Overlay Slot)")
+	UPROPERTY(EditDefaultsOnly, BlueprintReadOnly, Category="Layout|Overlay Slot")
 	TEnumAsByte<EVerticalAlignment> VerticalAlignment;
 
 public:
 
-	UFUNCTION(BlueprintCallable, Category="Layout (Overlay Slot)")
+	UFUNCTION(BlueprintCallable, Category="Layout|Overlay Slot")
 	void SetPadding(FMargin InPadding);
 
-	UFUNCTION(BlueprintCallable, Category="Layout (Overlay Slot)")
+	UFUNCTION(BlueprintCallable, Category="Layout|Overlay Slot")
 	void SetHorizontalAlignment(EHorizontalAlignment InHorizontalAlignment);
 
-	UFUNCTION(BlueprintCallable, Category="Layout (Overlay Slot)")
+	UFUNCTION(BlueprintCallable, Category="Layout|Overlay Slot")
 	void SetVerticalAlignment(EVerticalAlignment InVerticalAlignment);
 
 public:
diff --git a/Engine/Source/Runtime/UMG/Public/Components/ScaleBoxSlot.h b/Engine/Source/Runtime/UMG/Public/Components/ScaleBoxSlot.h
--- a/Engine/Source/Runtime/UMG/Public/Components/ScaleBoxSlot.h
+++ b/Engine/Source/Runtime/UMG/Public/Components/ScaleBoxSlot.h
@@ -15,26 +15,26 @@ class UMG_API UScaleBoxSlot : public UPanelSlot
 public:
 	
 	/** The padding area between the slot and the content it contains. */
-	UPROPERTY(EditDefaultsOnly, BlueprintReadOnly, Category="Layout (ScaleBox Slot)")
+	UPROPERTY(EditDefaultsOnly, BlueprintReadOnly, Category="Layout|ScaleBox Slot")
 	FMargin Padding;
 
 	/** The alignment of the object horizontally. */
-	UPROPERTY(EditDefaultsOnly, BlueprintReadOnly, Category="Layout (ScaleBox Slot)")
+	UPROPERTY(EditDefaultsOnly, BlueprintReadOnly, Category="Layout|ScaleBox Slot")
 	TEnumAsByte<EHorizontalAlignment> HorizontalAlignment;
 
 	/** The alignment of the object vertically. */
-	UPROPERTY(EditDefaultsOnly, BlueprintReadOnly, Category="Layout (ScaleBox Slot)" )
+	UPROPERTY(EditDefaultsOnly, BlueprintReadOnly, Category="Layout|ScaleBox Slot" )
 	TEnumAsByte<EVerticalAlignment> VerticalAlignment;
 
 public:
 
-	UFUNCTION(BlueprintCallable, Category="Layout (ScaleBox Slot)")
+	UFUNCTION(BlueprintCallable, Category="Layout|ScaleBox Slot")
 	void SetPadding(FMargin InPadding);
 
-	UFUNCTION(BlueprintCallable, Category="Layout (ScaleBox Slot)")
+	UFUNCTION(BlueprintCallable, Category="Layout|ScaleBox Slot")
 	void SetHorizontalAlignment(EHorizontalAlignment InHorizontalAlignment);
 
-	UFUNCTION(BlueprintCallable, Category="Layout (ScaleBox Slot)")
+	UFUNCTION(BlueprintCallable, Category="Layout|ScaleBox Slot")
 	void SetVerticalAlignment(EVerticalAlignment InVerticalAlignment);
 
 	// UPanelSlot interface
diff --git a/Engine/Source/Runtime/UMG/Public/Components/ScrollBoxSlot.h b/Engine/Source/Runtime/UMG/Public/Components/ScrollBoxSlot.h
--- a/Engine/Source/Runtime/UMG/Public/Components/ScrollBoxSlot.h
+++ b/Engine/Source/Runtime/UMG/Public/Components/ScrollBoxSlot.h
@@ -15,19 +15,19 @@ class UMG_API UScrollBoxSlot : public UPanelSlot
 public:
 	
 	/** The padding area between the slot and the content it contains. */
-	UPROPERTY(EditDefaultsOnly, BlueprintReadOnly, Category="Layout (ScrollBox Slot)")
+	UPROPERTY(EditDefaultsOnly, BlueprintReadOnly, Category="Layout|ScrollBox Slot")
 	FMargin Padding;
 
 	/** The alignment of the object horizontally. */
-	UPROPERTY(EditDefaultsOnly, BlueprintReadOnly, Category="Layout (ScrollBox Slot)")
+	UPROPERTY(EditDefaultsOnly, BlueprintReadOnly, Category="Layout|ScrollBox Slot")
 	TEnumAsByte<EHorizontalAlignment> HorizontalAlignment;
 
 public:
 
-	UFUNCTION(BlueprintCallable, Category="Layout (ScrollBox Slot)")
+	UFUNCTION(BlueprintCallable, Category="Layout|ScrollBox Slot")
 	void SetPadding(FMargin InPadding);
 
-	UFUNCTION(BlueprintCallable, Category="Layout (ScrollBox Slot)")
+	UFUNCTION(BlueprintCallable, Category="Layout|ScrollBox Slot")
 	void SetHorizontalAlignment(EHorizontalAlignment InHorizontalAlignment);
 
 public:
diff --git a/Engine/Source/Runtime/UMG/Public/Components/SizeBox.h b/Engine/Source/Runtime/UMG/Public/Components/SizeBox.h
--- a/Engine/Source/Runtime/UMG/Public/Components/SizeBox.h
+++ b/Engine/Source/Runtime/UMG/Public/Components/SizeBox.h
@@ -72,45 +72,45 @@ public:
 public:
 		
 	/** When specified, ignore the content's desired size and report the WidthOverride as the Box's desired width. */
-	UFUNCTION(BlueprintCallable, Category="Layout")
+	UFUNCTION(BlueprintCallable, Category="Layout|Size Box")
 	void SetWidthOverride(float InWidthOverride);
 
-	UFUNCTION(BlueprintCallable, Category="Layout")
+	UFUNCTION(BlueprintCallable, Category="Layout|Size Box")
 	void ClearWidthOverride();
 
 	/** When specified, ignore the content's desired size and report the HeightOverride as the Box's desired height. */
-	UFUNCTION(BlueprintCallable, Category="Layout")
+	UFUNCTION(BlueprintCallable, Category="Layout|Size Box")
 	void SetHeightOverride(float InHeightOverride);
 
-	UFUNCTION(BlueprintCallable, Category="Layout")
+	UFUNCTION(BlueprintCallable, Category="Layout|Size Box")
 	void ClearHeightOverride();
 
 	/** When specified, will report the MinDesiredWidth if larger than the content's desired width. */
-	UFUNCTION(BlueprintCallable, Category="Layout")
+	UFUNCTION(BlueprintCallable, Category="Layout|Size Box")
 	void SetMinDesiredWidth(float InMinDesiredWidth);
 
-	UFUNCTION(BlueprintCallable, Category="Layout")
+	UFUNCTION(BlueprintCallable, Category="Layout|Size Box")
 	void ClearMinDesiredWidth();
 
 	/** When specified, will report the MinDesiredHeight if larger than the content's desired height. */
-	UFUNCTION(BlueprintCallable, Category="Layout")
+	UFUNCTION(BlueprintCallable, Category="Layout|Size Box")
 	void SetMinDesiredHeight(float InMinDesiredHeight);
 
-	UFUNCTION(BlueprintCallable, Category="Layout")
+	UFUNCTION(BlueprintCallable, Category="Layout|Size Box")
 	void ClearMinDesiredHeight();
 
 	/** When specified, will report the MaxDesiredWidth if smaller than the content's desired width. */
-	UFUNCTION(BlueprintCallable, Category="Layout")
+	UFUNCTION(BlueprintCallable, Category="Layout|Size Box")
 	void SetMaxDesiredWidth(float InMaxDesiredWidth);
 
-	UFUNCTION(BlueprintCallable, Category="Layout")
+	UFUNCTION(BlueprintCallable, Category="Layout|Size Box")
 	void ClearMaxDesiredWidth();
 
 	/** When specified, will report the MaxDesiredHeight if smaller than the content's desired height. */
-	UFUNCTION(BlueprintCallable, Category="Layout")
+	UFUNCTION(BlueprintCallable, Category="Layout|Size Box")
 	void SetMaxDesiredHeight(float InMaxDesiredHeight);
 
-	UFUNCTION(BlueprintCallable, Category="Layout")
+	UFUNCTION(BlueprintCallable, Category="Layout|Size Box")
 	void ClearMaxDesiredHeight();
 
 public:
diff --git a/Engine/Source/Runtime/UMG/Public/Components/SizeBoxSlot.h b/Engine/Source/Runtime/UMG/Public/Components/SizeBoxSlot.h
--- a/Engine/Source/Runtime/UMG/Public/Components/SizeBoxSlot.h
+++ b/Engine/Source/Runtime/UMG/Public/Components/SizeBoxSlot.h
@@ -14,26 +14,26 @@ class UMG_API USizeBoxSlot : public UPanelSlot
 
 public:
 	/** The padding area between the slot and the content it contains. */
-	UPROPERTY(EditDefaultsOnly, BlueprintReadOnly, Category="Layout (SizeBox Slot)")
+	UPROPERTY(EditDefaultsOnly, BlueprintReadOnly, Category="Layout|SizeBox Slot")
 	FMargin Padding;
 
 	/** The alignment of the object horizontally. */
-	UPROPERTY(EditDefaultsOnly, BlueprintReadOnly, Category="Layout (SizeBox Slot)")
+	UPROPERTY(EditDefaultsOnly, BlueprintReadOnly, Category="Layout|SizeBox Slot")
 	TEnumAsByte<EHorizontalAlignment> HorizontalAlignment;
 
 	/** The alignment of the object vertically. */
-	UPROPERTY(EditDefaultsOnly, BlueprintReadOnly, Category="Layout (SizeBox Slot)")
+	UPROPERTY(EditDefaultsOnly, BlueprintReadOnly, Category="Layout|SizeBox Slot")
 	TEnumAsByte<EVerticalAlignment> VerticalAlignment;
 
 public:
 
-	UFUNCTION(BlueprintCallable, Category="Layout (SizeBox Slot)")
+	UFUNCTION(BlueprintCallable, Category="Layout|SizeBox Slot")
 	void SetPadding(FMargin InPadding);
 
-	UFUNCTION(BlueprintCallable, Category="Layout (SizeBox Slot)")
+	UFUNCTION(BlueprintCallable, Category="Layout|SizeBox Slot")
 	void SetHorizontalAlignment(EHorizontalAlignment InHorizontalAlignment);
 
-	UFUNCTION(BlueprintCallable, Category="Layout (SizeBox Slot)")
+	UFUNCTION(BlueprintCallable, Category="Layout|SizeBox Slot")
 	void SetVerticalAlignment(EVerticalAlignment InVerticalAlignment);
 
 public:
diff --git a/Engine/Source/Runtime/UMG/Public/Components/UniformGridSlot.h b/Engine/Source/Runtime/UMG/Public/Components/UniformGridSlot.h
--- a/Engine/Source/Runtime/UMG/Public/Components/UniformGridSlot.h
+++ b/Engine/Source/Runtime/UMG/Public/Components/UniformGridSlot.h
@@ -18,37 +18,37 @@ class UMG_API UUniformGridSlot : public UPanelSlot
 public:
 
 	/** The alignment of the object horizontally. */
-	UPROPERTY(EditDefaultsOnly, BlueprintReadOnly, Category="Layout (Uniform Grid Slot)")
+	UPROPERTY(EditDefaultsOnly, BlueprintReadOnly, Category="Layout|Uniform Grid Slot")
 	TEnumAsByte<EHorizontalAlignment> HorizontalAlignment;
 
 	/** The alignment of the object vertically. */
-	UPROPERTY(EditDefaultsOnly, BlueprintReadOnly, Category="Layout (Uniform Grid Slot)")
+	UPROPERTY(EditDefaultsOnly, BlueprintReadOnly, Category="Layout|Uniform Grid Slot")
 	TEnumAsByte<EVerticalAlignment> VerticalAlignment;
 	
 	/** The row index of the cell this slot is in */
-	UPROPERTY(EditDefaultsOnly, BlueprintReadOnly, meta=( UIMin = "0" ), Category="Layout (Uniform Grid Slot)")
+	UPROPERTY(EditDefaultsOnly, BlueprintReadOnly, meta=( UIMin = "0" ), Category="Layout|Uniform Grid Slot")
 	int32 Row;
 	
 	/** The column index of the cell this slot is in */
-	UPROPERTY(EditDefaultsOnly, BlueprintReadOnly, meta=( UIMin = "0" ), Category="Layout (Uniform Grid Slot)")
+	UPROPERTY(EditDefaultsOnly, BlueprintReadOnly, meta=( UIMin = "0" ), Category="Layout|Uniform Grid Slot")
 	int32 Column;
 
 public:
 
 	/** Sets the row index of the slot, this determines what cell the slot is in the panel */
-	UFUNCTION(BlueprintCallable, Category="Layout (Uniform Grid Slot)")
+	UFUNCTION(BlueprintCallable, Category="Layout|Uniform Grid Slot")
 	void SetRow(int32 InRow);
 
 	/** Sets the column index of the slot, this determines what cell the slot is in the panel */
-	UFUNCTION(BlueprintCallable, Category="Layout (Uniform Grid Slot)")
+	UFUNCTION(BlueprintCallable, Category="Layout|Uniform Grid Slot")
 	void SetColumn(int32 InColumn);
 
 	/**  */
-	UFUNCTION(BlueprintCallable, Category="Layout (Uniform Grid Slot)")
+	UFUNCTION(BlueprintCallable, Category="Layout|Uniform Grid Slot")
 	void SetHorizontalAlignment(EHorizontalAlignment InHorizontalAlignment);
 
 	/**  */
-	UFUNCTION(BlueprintCallable, Category="Layout (Uniform Grid Slot)")
+	UFUNCTION(BlueprintCallable, Category="Layout|Uniform Grid Slot")
 	void SetVerticalAlignment(EVerticalAlignment InVerticalAlignment);
 
 public:
diff --git a/Engine/Source/Runtime/UMG/Public/Components/VerticalBoxSlot.h b/Engine/Source/Runtime/UMG/Public/Components/VerticalBoxSlot.h
--- a/Engine/Source/Runtime/UMG/Public/Components/VerticalBoxSlot.h
+++ b/Engine/Source/Runtime/UMG/Public/Components/VerticalBoxSlot.h
@@ -15,33 +15,33 @@ class UMG_API UVerticalBoxSlot : public UPanelSlot
 public:
 	
 	/** The padding area between the slot and the content it contains. */
-	UPROPERTY(EditDefaultsOnly, BlueprintReadOnly, Category="Layout (Vertical Box Slot)")
+	UPROPERTY(EditDefaultsOnly, BlueprintReadOnly, Category="Layout|Vertical Box Slot")
 	FMargin Padding;
 
 	/** How much space this slot should occupy in the direction of the panel. */
-	UPROPERTY(EditDefaultsOnly, BlueprintReadOnly, Category="Layout (Vertical Box Slot)")
+	UPROPERTY(EditDefaultsOnly, BlueprintReadOnly, Category="Layout|Vertical Box Slot")
 	FSlateChildSize Size;
 
 	/** The alignment of the object horizontally. */
-	UPROPERTY(EditDefaultsOnly, BlueprintReadOnly, Category="Layout (Vertical Box Slot)")
+	UPROPERTY(EditDefaultsOnly, BlueprintReadOnly, Category="Layout|Vertical Box Slot")
 	TEnumAsByte<EHorizontalAlignment> HorizontalAlignment;
 
 	/** The alignment of the object vertically. */
-	UPROPERTY(EditDefaultsOnly, BlueprintReadOnly, Category="Layout (Vertical Box Slot)")
+	UPROPERTY(EditDefaultsOnly, BlueprintReadOnly, Category="Layout|Vertical Box Slot")
 	TEnumAsByte<EVerticalAlignment> VerticalAlignment;
 
 public:
 
-	UFUNCTION(BlueprintCallable, Category="Layout (Vertical Box Slot)")
+	UFUNCTION(BlueprintCallable, Category="Layout|Vertical Box Slot")
 	void SetPadding(FMargin InPadding);
 
-	UFUNCTION(BlueprintCallable, Category="Layout (Vertical Box Slot)")
+	UFUNCTION(BlueprintCallable, Category="Layout|Vertical Box Slot")
 	void SetSize(FSlateChildSize InSize);
 
-	UFUNCTION(BlueprintCallable, Category="Layout (Vertical Box Slot)")
+	UFUNCTION(BlueprintCallable, Category="Layout|Vertical Box Slot")
 	void SetHorizontalAlignment(EHorizontalAlignment InHorizontalAlignment);
 
-	UFUNCTION(BlueprintCallable, Category="Layout (Vertical Box Slot)")
+	UFUNCTION(BlueprintCallable, Category="Layout|Vertical Box Slot")
 	void SetVerticalAlignment(EVerticalAlignment InVerticalAlignment);
 
 public:
diff --git a/Engine/Source/Runtime/UMG/Public/Components/WidgetSwitcherSlot.h b/Engine/Source/Runtime/UMG/Public/Components/WidgetSwitcherSlot.h
--- a/Engine/Source/Runtime/UMG/Public/Components/WidgetSwitcherSlot.h
+++ b/Engine/Source/Runtime/UMG/Public/Components/WidgetSwitcherSlot.h
@@ -16,26 +16,26 @@ class UMG_API UWidgetSwitcherSlot : public UPanelSlot
 public:
 	
 	/** The padding area between the slot and the content it contains. */
-	UPROPERTY(EditDefaultsOnly, BlueprintReadOnly, Category="Layout (Widget Switcher Slot)")
+	UPROPERTY(EditDefaultsOnly, BlueprintReadOnly, Category="Layout|Widget Switcher Slot")
 	FMargin Padding;
 
 	/** The alignment of the object horizontally. */
-	UPROPERTY(EditDefaultsOnly, BlueprintReadOnly, Category="Layout (Widget Switcher Slot)")
+	UPROPERTY(EditDefaultsOnly, BlueprintReadOnly, Category="Layout|Widget Switcher Slot")
 	TEnumAsByte<EHorizontalAlignment> HorizontalAlignment;
 
 	/** The alignment of the object vertically. */
-	UPROPERTY(EditDefaultsOnly, BlueprintReadOnly, Category="Layout (Widget Switcher Slot)")
+	UPROPERTY(EditDefaultsOnly, BlueprintReadOnly, Category="Layout|Widget Switcher Slot")
 	TEnumAsByte<EVerticalAlignment> VerticalAlignment;
 
 public:
 
-	UFUNCTION(BlueprintCallable, Category="Layout (Widget Switcher Slot)")
+	UFUNCTION(BlueprintCallable, Category="Layout|Widget Switcher Slot")
 	void SetPadding(FMargin InPadding);
 
-	UFUNCTION(BlueprintCallable, Category="Layout (Widget Switcher Slot)")
+	UFUNCTION(BlueprintCallable, Category="Layout|Widget Switcher Slot")
 	void SetHorizontalAlignment(EHorizontalAlignment InHorizontalAlignment);
 
-	UFUNCTION(BlueprintCallable, Category="Layout (Widget Switcher Slot)")
+	UFUNCTION(BlueprintCallable, Category="Layout|Widget Switcher Slot")
 	void SetVerticalAlignment(EVerticalAlignment InVerticalAlignment);
 
 public:
diff --git a/Engine/Source/Runtime/UMG/Public/Components/WrapBoxSlot.h b/Engine/Source/Runtime/UMG/Public/Components/WrapBoxSlot.h
--- a/Engine/Source/Runtime/UMG/Public/Components/WrapBoxSlot.h
+++ b/Engine/Source/Runtime/UMG/Public/Components/WrapBoxSlot.h
@@ -15,43 +15,43 @@ class UMG_API UWrapBoxSlot : public UPanelSlot
 public:
 
 	/** The padding area between the slot and the content it contains. */
-	UPROPERTY(EditDefaultsOnly, BlueprintReadOnly, Category="Layout (Wrap Box Slot)")
+	UPROPERTY(EditDefaultsOnly, BlueprintReadOnly, Category="Layout|Wrap Box Slot")
 	FMargin Padding;
 
 	/** Should this slot fill the remaining space on the line? */
-	UPROPERTY(EditDefaultsOnly, BlueprintReadOnly, Category="Layout (Wrap Box Slot)")
+	UPROPERTY(EditDefaultsOnly, BlueprintReadOnly, Category="Layout|Wrap Box Slot")
 	bool bFillEmptySpace;
 
 	/**
 	 * If the total available space in the wrap panel drops below this threshold, this slot will attempt to fill an entire line.
 	 * NOTE: A value of 0, denotes no filling will occur.
 	 */
-	UPROPERTY(EditDefaultsOnly, BlueprintReadOnly, Category="Layout (Wrap Box Slot)")
+	UPROPERTY(EditDefaultsOnly, BlueprintReadOnly, Category="Layout|Wrap Box Slot")
 	float FillSpanWhenLessThan;
 
 	/** The alignment of the object horizontally. */
-	UPROPERTY(EditDefaultsOnly, BlueprintReadOnly, Category="Layout (Wrap Box Slot)")
+	UPROPERTY(EditDefaultsOnly, BlueprintReadOnly, Category="Layout|Wrap Box Slot")
 	TEnumAsByte<EHorizontalAlignment> HorizontalAlignment;
 
 	/** The alignment of the object vertically. */
-	UPROPERTY(EditDefaultsOnly, BlueprintReadOnly, Category="Layout (Wrap Box Slot)")
+	UPROPERTY(EditDefaultsOnly, BlueprintReadOnly, Category="Layout|Wrap Box Slot")
 	TEnumAsByte<EVerticalAlignment> VerticalAlignment;
 
 public:
 
-	UFUNCTION(BlueprintCallable, Category="Layout (Wrap Box Slot)")
+	UFUNCTION(BlueprintCallable, Category="Layout|Wrap Box Slot")
 	void SetPadding(FMargin InPadding);
 
-	UFUNCTION(BlueprintCallable, Category="Layout (Wrap Box Slot)")
+	UFUNCTION(BlueprintCallable, Category="Layout|Wrap Box Slot")
 	void SetFillEmptySpace(bool InbFillEmptySpace);
 
-	UFUNCTION(BlueprintCallable, Category="Layout (Wrap Box Slot)")
+	UFUNCTION(BlueprintCallable, Category="Layout|Wrap Box Slot")
 	void SetFillSpanWhenLessThan(float InFillSpanWhenLessThan);
 
-	UFUNCTION(BlueprintCallable, Category="Layout (Wrap Box Slot)")
+	UFUNCTION(BlueprintCallable, Category="Layout|Wrap Box Slot")
 	void SetHorizontalAlignment(EHorizontalAlignment InHorizontalAlignment);
 
-	UFUNCTION(BlueprintCallable, Category="Layout (Wrap Box Slot)")
+	UFUNCTION(BlueprintCallable, Category="Layout|Wrap Box Slot")
 	void SetVerticalAlignment(EVerticalAlignment InVerticalAlignment);
 
 public:
