UMG - Animation fixes.
+ Fix an issue where animations wouldn't always play the first or last frame when stopping.
+ Fix an issue where playing in reverse wouldn't work without looping, or setting the play time.
+ Clamp the start time to between 0 and the length of the animation.

diff --git a/Engine/Source/Runtime/UMG/Private/Animation/UMGSequencePlayer.cpp b/Engine/Source/Runtime/UMG/Private/Animation/UMGSequencePlayer.cpp
--- a/Engine/Source/Runtime/UMG/Private/Animation/UMGSequencePlayer.cpp
+++ b/Engine/Source/Runtime/UMG/Private/Animation/UMGSequencePlayer.cpp
@@ -61,39 +61,50 @@ void UUMGSequencePlayer::Tick(float DeltaTime)
 		float AnimationLength = TimeRange.GetUpperBoundValue();
 		if ( TimeCursorPosition < 0 )
 		{
-			if ( PlayMode == EUMGSequencePlayMode::PingPong )
+			NumLoopsCompleted++;
+			if (NumLoopsToPlay != 0 && NumLoopsCompleted >= NumLoopsToPlay)
 			{
-				bIsPlayingForward = !bIsPlayingForward;
-				TimeCursorPosition = FMath::Abs(TimeCursorPosition);
+				TimeCursorPosition = 0;
+				PlayerStatus = EMovieScenePlayerStatus::Stopped;
+				OnSequenceFinishedPlayingEvent.Broadcast(*this);
 			}
 			else
 			{
-				TimeCursorPosition += AnimationLength;
+				if (PlayMode == EUMGSequencePlayMode::PingPong)
+				{
+					bIsPlayingForward = !bIsPlayingForward;
+					TimeCursorPosition = FMath::Abs(TimeCursorPosition);
+				}
+				else
+				{
+					TimeCursorPosition += AnimationLength;
+				}
 			}
-			NumLoopsCompleted++;
 		}
-		else if ( TimeCursorPosition >= AnimationLength )
+		else if ( TimeCursorPosition > AnimationLength )
 		{
-			if ( PlayMode == EUMGSequencePlayMode::PingPong )
+			NumLoopsCompleted++;
+			if (NumLoopsToPlay != 0 && NumLoopsCompleted >= NumLoopsToPlay)
 			{
-				bIsPlayingForward = !bIsPlayingForward;
-				TimeCursorPosition = AnimationLength - (TimeCursorPosition - AnimationLength);
+				TimeCursorPosition = AnimationLength;
+				PlayerStatus = EMovieScenePlayerStatus::Stopped;
+				OnSequenceFinishedPlayingEvent.Broadcast(*this);
 			}
 			else
 			{
-				TimeCursorPosition -= AnimationLength;
+				if (PlayMode == EUMGSequencePlayMode::PingPong)
+				{
+					bIsPlayingForward = !bIsPlayingForward;
+					TimeCursorPosition = AnimationLength - (TimeCursorPosition - AnimationLength);
+				}
+				else
+				{
+					TimeCursorPosition -= AnimationLength;
+				}
 			}
-			NumLoopsCompleted++;
 		}
-		
-		// Have we looped the requested number of times?
-		if (NumLoopsToPlay != 0 && NumLoopsCompleted >= NumLoopsToPlay)
-		{
-			PlayerStatus = EMovieScenePlayerStatus::Stopped;
-			OnSequenceFinishedPlayingEvent.Broadcast( *this );
-		}
-		
-		if (RootMovieSceneInstance.IsValid() && PlayerStatus != EMovieScenePlayerStatus::Stopped)
+
+		if (RootMovieSceneInstance.IsValid())
 		{
 			RootMovieSceneInstance->Update(TimeCursorPosition, LastTimePosition, *this);
 		}
@@ -103,13 +114,20 @@ void UUMGSequencePlayer::Tick(float DeltaTime)
 void UUMGSequencePlayer::Play(float StartAtTime, int32 InNumLoopsToPlay, EUMGSequencePlayMode::Type InPlayMode)
 {
 	UMovieScene* MovieScene = Animation->MovieScene;
-
 	RootMovieSceneInstance = MakeShareable( new FMovieSceneInstance( *MovieScene ) );
-
 	RootMovieSceneInstance->RefreshInstance( *this );
 
-	TimeCursorPosition = StaticCast<double>(StartAtTime);
 	PlayMode = InPlayMode;
+
+	// Clamp the start time to be between 0 and the upper time bound
+	TimeCursorPosition = StaticCast<double>(FMath::Clamp(StartAtTime, 0.0f, TimeRange.GetUpperBoundValue()));
+
+	if (PlayMode == EUMGSequencePlayMode::Reverse)
+	{
+		// When playing in reverse count substract the start time from the end.
+		TimeCursorPosition = TimeRange.GetUpperBoundValue() - TimeCursorPosition;
+	}
+
 	if ( PlayMode == EUMGSequencePlayMode::PingPong )
 	{
 		// When animating in ping-pong mode double the number of loops to play so that a loop is a complete forward/reverse cycle.
@@ -119,6 +137,7 @@ void UUMGSequencePlayer::Play(float StartAtTime, int32 InNumLoopsToPlay, EUMGSeq
 	{
 		NumLoopsToPlay = InNumLoopsToPlay;
 	}
+
 	NumLoopsCompleted = 0;
 	bIsPlayingForward = InPlayMode != EUMGSequencePlayMode::Reverse;
 
