UMG - Some cleanup on the widget component and exposing a bit more functionality for hit testing to not always assume a 0 size pointer.

diff --git a/Engine/Source/Runtime/UMG/Private/Components/WidgetComponent.cpp b/Engine/Source/Runtime/UMG/Private/Components/WidgetComponent.cpp
--- a/Engine/Source/Runtime/UMG/Private/Components/WidgetComponent.cpp
+++ b/Engine/Source/Runtime/UMG/Private/Components/WidgetComponent.cpp
@@ -179,21 +179,25 @@ public:
 
 		if( World.IsValid() && ensure( World->IsGameWorld() ) )
 		{
-			ULocalPlayer* const TargetPlayer = GEngine->GetLocalPlayerFromControllerId(World.Get(), 0);
-			APlayerController* PlayerController = TargetPlayer->PlayerController;
-
-			if ( UPrimitiveComponent* HitComponent = GetHitResultAtScreenPositionAndCache(TargetPlayer->PlayerController, InGeometry.AbsoluteToLocal(DesktopSpaceCoordinate)) )
+			UWorld* SafeWorld = World.Get();
+			if ( SafeWorld )
 			{
-				if ( UWidgetComponent* WidgetComponent = Cast<UWidgetComponent>(HitComponent) )
+				ULocalPlayer* const TargetPlayer = GEngine->GetLocalPlayerFromControllerId(SafeWorld, 0);
+				APlayerController* PlayerController = TargetPlayer->PlayerController;
+
+				if ( UPrimitiveComponent* HitComponent = GetHitResultAtScreenPositionAndCache(TargetPlayer->PlayerController, InGeometry.AbsoluteToLocal(DesktopSpaceCoordinate)) )
 				{
-					// Make sure the player is interacting with the front of the widget
-					// For widget components, the "front" faces the Z (or Up vector) direction
-					if ( FVector::DotProduct(WidgetComponent->GetUpVector(), CachedHitResult.ImpactPoint - CachedHitResult.TraceStart) < 0.f )
+					if ( UWidgetComponent* WidgetComponent = Cast<UWidgetComponent>(HitComponent) )
 					{
-						// Make sure the player is close enough to the widget to interact with it
-						if ( FVector::DistSquared(CachedHitResult.TraceStart, CachedHitResult.ImpactPoint) <= FMath::Square(WidgetComponent->GetMaxInteractionDistance()) )
+						// Make sure the player is interacting with the front of the widget
+						// For widget components, the "front" faces the Z (or Up vector) direction
+						if ( FVector::DotProduct(WidgetComponent->GetUpVector(), CachedHitResult.ImpactPoint - CachedHitResult.TraceStart) < 0.f )
 						{
-							return WidgetComponent->GetHitWidgetPath(CachedHitResult, bIgnoreEnabledStatus);
+							// Make sure the player is close enough to the widget to interact with it
+							if ( FVector::DistSquared(CachedHitResult.TraceStart, CachedHitResult.ImpactPoint) <= FMath::Square(WidgetComponent->GetMaxInteractionDistance()) )
+							{
+								return WidgetComponent->GetHitWidgetPath(CachedHitResult, bIgnoreEnabledStatus);
+							}
 						}
 					}
 				}
@@ -225,6 +229,8 @@ public:
 			ULocalPlayer* const TargetPlayer = GEngine->GetLocalPlayerFromControllerId(World.Get(), 0);
 			if ( TargetPlayer && TargetPlayer->PlayerController )
 			{
+				FVector2D LocalMouseCoordinate = ViewportGeometry.AbsoluteToLocal(ScreenSpaceMouseCoordinate);
+
 				// Check for a hit against any widget components in the world
 				for ( TWeakObjectPtr<UWidgetComponent> Component : RegisteredComponents )
 				{
@@ -232,9 +238,6 @@ public:
 					// Check if visible;
 					if ( WidgetComponent && WidgetComponent->GetSlateWidget() == ChildWidget )
 					{
-						FVector2D LocalMouseCoordinate = ViewportGeometry.AbsoluteToLocal(ScreenSpaceMouseCoordinate);
-						FVector2D LocalLastMouseCoordinate = ViewportGeometry.AbsoluteToLocal(LastScreenSpaceMouseCoordinate);
-
 						if ( UPrimitiveComponent* HitComponent = GetHitResultAtScreenPositionAndCache(TargetPlayer->PlayerController, LocalMouseCoordinate) )
 						{
 							if ( WidgetComponent == HitComponent )
@@ -878,7 +881,7 @@ UUserWidget* UWidgetComponent::GetUserWidgetObject() const
 	return Widget;
 }
 
-TArray<FWidgetAndPointer> UWidgetComponent::GetHitWidgetPath( const FHitResult& HitResult, bool bIgnoreEnabledStatus  )
+TArray<FWidgetAndPointer> UWidgetComponent::GetHitWidgetPath(const FHitResult& HitResult, bool bIgnoreEnabledStatus, float CursorRadius )
 {
 	FVector2D LocalHitLocation;
 	GetLocalHitLocation(HitResult, LocalHitLocation);
@@ -891,7 +894,7 @@ TArray<FWidgetAndPointer> UWidgetComponent::GetHitWidgetPath( const FHitResult&
 	// Cache the location of the hit
 	LastLocalHitLocation = LocalHitLocation;
 
-	TArray<FWidgetAndPointer> ArrangedWidgets = HitTestGrid->GetBubblePath( LocalHitLocation, 0.0f, bIgnoreEnabledStatus );
+	TArray<FWidgetAndPointer> ArrangedWidgets = HitTestGrid->GetBubblePath(LocalHitLocation, CursorRadius, bIgnoreEnabledStatus);
 
 	for( FWidgetAndPointer& ArrangedWidget : ArrangedWidgets )
 	{
diff --git a/Engine/Source/Runtime/UMG/Public/Components/WidgetComponent.h b/Engine/Source/Runtime/UMG/Public/Components/WidgetComponent.h
--- a/Engine/Source/Runtime/UMG/Public/Components/WidgetComponent.h
+++ b/Engine/Source/Runtime/UMG/Public/Components/WidgetComponent.h
@@ -78,7 +78,7 @@ public:
 	UUserWidget* GetUserWidgetObject() const;
 
 	/** @return List of widgets with their geometry and the cursor position transformed into this Widget component's space. */
-	TArray<FWidgetAndPointer> GetHitWidgetPath( const FHitResult& HitResult, bool bIgnoreEnabledStatus );
+	TArray<FWidgetAndPointer> GetHitWidgetPath( const FHitResult& HitResult, bool bIgnoreEnabledStatus, float CursorRadius = 0.0f );
 
 	/** @return The render target to which the user widget is rendered */
 	UTextureRenderTarget2D* GetRenderTarget() const { return RenderTarget; }
