UMG - Using the widget tree requires you to import the UPanelWidget, so that's now explictly imported by the WidgetTree header.

diff --git a/Engine/Source/Runtime/UMG/Public/Blueprint/WidgetTree.h b/Engine/Source/Runtime/UMG/Public/Blueprint/WidgetTree.h
--- a/Engine/Source/Runtime/UMG/Public/Blueprint/WidgetTree.h
+++ b/Engine/Source/Runtime/UMG/Public/Blueprint/WidgetTree.h
@@ -3,6 +3,7 @@
 #pragma once
 
 #include "UserWidget.h"
+#include "Components/PanelWidget.h"
 
 #include "WidgetTree.generated.h"
 
