UMG - Widgets now have a Destruct event that is called when the slate widget is dereferenced.  All forwarded slate events, Tick, Construct...etc. are no longer BlueprintNativeEvent's, this was done so that we can eventually stop paying the overhead of doing lots of UFunction lookups to parent classes to see if the functions have been implemented.  Additionally users who implemented these functions in C++ were finding their code no longer being called when users subclassing the widget in blueprints forgot to add the call parent node after handling a particular event.  So, all those Event_Implementation are now deprecated.  In C++ if you want to implement logic in Tick, Construct...etc.  Use the NativeTick, NativeConstruct, NativeOnLostFocus...etc.  These will always be called first, and will have control over when the blueprint call actually occurs by controlling when they call Super.  General cleanup in the UserWidget header.

diff --git a/Engine/Source/Runtime/UMG/Private/Slate/SObjectWidget.cpp b/Engine/Source/Runtime/UMG/Private/Slate/SObjectWidget.cpp
--- a/Engine/Source/Runtime/UMG/Private/Slate/SObjectWidget.cpp
+++ b/Engine/Source/Runtime/UMG/Private/Slate/SObjectWidget.cpp
@@ -24,6 +24,11 @@ void SObjectWidget::ResetWidget()
 {
 	if ( UObjectInitialized() && WidgetObject )
 	{
+		if ( CanRouteEvent() )
+		{
+			WidgetObject->NativeDestruct();
+		}
+
 		// NOTE: When the SObjectWidget gets released we know that the User Widget has
 		// been removed from the slate widget hierarchy.  When this occurs, we need to 
 		// immediately release all slate widget widgets to prevent deletion from taking
diff --git a/Engine/Source/Runtime/UMG/Private/UserWidget.cpp b/Engine/Source/Runtime/UMG/Private/UserWidget.cpp
--- a/Engine/Source/Runtime/UMG/Private/UserWidget.cpp
+++ b/Engine/Source/Runtime/UMG/Private/UserWidget.cpp
@@ -201,176 +201,6 @@ void UUserWidget::SetIsDesignTime(bool bInDesignTime)
 	});
 }
 
-void UUserWidget::Construct_Implementation()
-{
-
-}
-
-void UUserWidget::Tick_Implementation(FGeometry MyGeometry, float InDeltaTime)
-{
-
-}
-
-void UUserWidget::OnPaint_Implementation(FPaintContext& Context) const
-{
-
-}
-
-FEventReply UUserWidget::OnFocusReceived_Implementation(FGeometry MyGeometry, FFocusEvent InFocusEvent)
-{
-	return UWidgetBlueprintLibrary::Unhandled();
-}
-
-void UUserWidget::OnFocusLost_Implementation(FFocusEvent InFocusEvent)
-{
-
-}
-
-FEventReply UUserWidget::OnKeyChar_Implementation(FGeometry MyGeometry, FCharacterEvent InCharacterEvent)
-{
-	return UWidgetBlueprintLibrary::Unhandled();
-}
-
-FEventReply UUserWidget::OnPreviewKeyDown_Implementation(FGeometry MyGeometry, FKeyEvent InKeyEvent)
-{
-	return UWidgetBlueprintLibrary::Unhandled();
-}
-
-FEventReply UUserWidget::OnKeyDown_Implementation(FGeometry MyGeometry, FKeyEvent InKeyEvent)
-{
-	return UWidgetBlueprintLibrary::Unhandled();
-}
-
-FEventReply UUserWidget::OnKeyUp_Implementation(FGeometry MyGeometry, FKeyEvent InKeyEvent)
-{
-	return UWidgetBlueprintLibrary::Unhandled();
-}
-
-FEventReply UUserWidget::OnAnalogValueChanged_Implementation(FGeometry MyGeometry, FAnalogInputEvent InAnalogInputEvent)
-{
-	return UWidgetBlueprintLibrary::Unhandled();
-}
-
-FEventReply UUserWidget::OnMouseButtonDown_Implementation(FGeometry MyGeometry, const FPointerEvent& MouseEvent)
-{
-	return UWidgetBlueprintLibrary::Unhandled();
-}
-
-FEventReply UUserWidget::OnPreviewMouseButtonDown_Implementation(FGeometry MyGeometry, const FPointerEvent& MouseEvent)
-{
-	return UWidgetBlueprintLibrary::Unhandled();
-}
-
-FEventReply UUserWidget::OnMouseButtonUp_Implementation(FGeometry MyGeometry, const FPointerEvent& MouseEvent)
-{
-	return UWidgetBlueprintLibrary::Unhandled();
-}
-
-FEventReply UUserWidget::OnMouseMove_Implementation(FGeometry MyGeometry, const FPointerEvent& MouseEvent)
-{
-	return UWidgetBlueprintLibrary::Unhandled();
-}
-
-void UUserWidget::OnMouseEnter_Implementation(FGeometry MyGeometry, const FPointerEvent& MouseEvent)
-{
-
-}
-
-void UUserWidget::OnMouseLeave_Implementation(const FPointerEvent& MouseEvent)
-{
-
-}
-
-FEventReply UUserWidget::OnMouseWheel_Implementation(FGeometry MyGeometry, const FPointerEvent& MouseEvent)
-{
-	return UWidgetBlueprintLibrary::Unhandled();
-}
-
-FEventReply UUserWidget::OnMouseButtonDoubleClick_Implementation(FGeometry InMyGeometry, const FPointerEvent& InMouseEvent)
-{
-	return UWidgetBlueprintLibrary::Unhandled();
-}
-
-void UUserWidget::OnDragDetected_Implementation(FGeometry MyGeometry, const FPointerEvent& PointerEvent, UDragDropOperation*& Operation)
-{
-
-}
-
-void UUserWidget::OnDragCancelled_Implementation(const FPointerEvent& PointerEvent, UDragDropOperation* Operation)
-{
-
-}
-
-void UUserWidget::OnDragEnter_Implementation(FGeometry MyGeometry, FPointerEvent PointerEvent, UDragDropOperation* Operation)
-{
-
-}
-
-void UUserWidget::OnDragLeave_Implementation(FPointerEvent PointerEvent, UDragDropOperation* Operation)
-{
-
-}
-
-bool UUserWidget::OnDragOver_Implementation(FGeometry MyGeometry, FPointerEvent PointerEvent, UDragDropOperation* Operation)
-{
-	return false;
-}
-
-bool UUserWidget::OnDrop_Implementation(FGeometry MyGeometry, FPointerEvent PointerEvent, UDragDropOperation* Operation)
-{
-	return false;
-}
-
-FEventReply UUserWidget::OnControllerButtonPressed_Implementation(FGeometry MyGeometry, FControllerEvent ControllerEvent)
-{
-	return UWidgetBlueprintLibrary::Unhandled();
-}
-
-FEventReply UUserWidget::OnControllerButtonReleased_Implementation(FGeometry MyGeometry, FControllerEvent ControllerEvent)
-{
-	return UWidgetBlueprintLibrary::Unhandled();
-}
-
-FEventReply UUserWidget::OnControllerAnalogValueChanged_Implementation(FGeometry MyGeometry, FControllerEvent ControllerEvent)
-{
-	return UWidgetBlueprintLibrary::Unhandled();
-}
-
-FEventReply UUserWidget::OnTouchGesture_Implementation(FGeometry MyGeometry, const FPointerEvent& GestureEvent)
-{
-	return UWidgetBlueprintLibrary::Unhandled();
-}
-
-FEventReply UUserWidget::OnTouchStarted_Implementation(FGeometry MyGeometry, const FPointerEvent& InTouchEvent)
-{
-	return UWidgetBlueprintLibrary::Unhandled();
-}
-
-FEventReply UUserWidget::OnTouchMoved_Implementation(FGeometry MyGeometry, const FPointerEvent& InTouchEvent)
-{
-	return UWidgetBlueprintLibrary::Unhandled();
-}
-
-FEventReply UUserWidget::OnTouchEnded_Implementation(FGeometry MyGeometry, const FPointerEvent& InTouchEvent)
-{
-	return UWidgetBlueprintLibrary::Unhandled();
-}
-
-FEventReply UUserWidget::OnMotionDetected_Implementation(FGeometry MyGeometry, FMotionEvent InMotionEvent)
-{
-	return UWidgetBlueprintLibrary::Unhandled();
-}
-
-void UUserWidget::OnAnimationStarted_Implementation( const UWidgetAnimation* Animation )
-{
-
-}
-
-void UUserWidget::OnAnimationFinished_Implementation( const UWidgetAnimation* Animation )
-{
-	
-}
-
 void UUserWidget::PlayAnimation( const UWidgetAnimation* InAnimation, float StartAtTime, int32 NumberOfLoops, EUMGSequencePlayMode::Type PlayMode)
 {
 	if( InAnimation )
@@ -465,7 +295,7 @@ void UUserWidget::OnWidgetRebuilt()
 	if (!IsDesignTime())
 	{
 		// Notify the widget that it has been constructed.
-		Construct();
+		NativeConstruct();
 	}
 }
 
@@ -558,38 +388,6 @@ void UUserWidget::SetContentForSlot(FName SlotName, UWidget* Content)
 	}
 }
 
-void UUserWidget::NativeTick(const FGeometry& MyGeometry, float InDeltaTime )
-{
-	GInitRunaway();
-
-	// Update active movie scenes
-	for( UUMGSequencePlayer* Player : ActiveSequencePlayers )
-	{
-		Player->Tick( InDeltaTime );
-	}
-
-	// The process of ticking the players above can stop them so we remove them after all players have ticked
-	for( UUMGSequencePlayer* StoppedPlayer : StoppedSequencePlayers )
-	{
-		ActiveSequencePlayers.Remove( StoppedPlayer );	
-	}
-
-	StoppedSequencePlayers.Empty();
-
-	if ( !bDesignTime )
-	{
-		UWorld* World = GetWorld();
-		if ( World )
-		{
-			// Update any latent actions we have for this actor
-			FLatentActionManager& LatentActionManager = World->GetLatentActionManager();
-			LatentActionManager.ProcessLatentActions(this, InDeltaTime);
-		}
-	}
-
-	Tick( MyGeometry, InDeltaTime );
-}
-
 TSharedRef<SWidget> UUserWidget::MakeViewportWidget(TSharedPtr<SWidget>& UserSlateWidget)
 {
 	UserSlateWidget = TakeWidget();
@@ -647,20 +445,32 @@ void UUserWidget::AddToScreen(ULocalPlayer* Player, int32 ZOrder)
 		{
 			if ( UGameViewportClient* ViewportClient = World->GetGameViewport() )
 			{
-				ZOrder = 10 + ZOrder;
 				if ( Player )
 				{
 					ViewportClient->AddViewportWidgetForPlayer(Player, RootWidget, ZOrder);
 				}
 				else
 				{
-					ViewportClient->AddViewportWidgetContent(RootWidget, ZOrder);
+					// We add 10 to the zorder when adding to the viewport to avoid 
+					// displaying below any built-in controls, like the virtual joysticks on mobile builds.
+					ViewportClient->AddViewportWidgetContent(RootWidget, ZOrder + 10);
 				}
+
+				FWorldDelegates::LevelRemovedFromWorld.AddUObject(this, &UUserWidget::OnLevelRemovedFromWorld);
 			}
 		}
 	}
 }
 
+void UUserWidget::OnLevelRemovedFromWorld(ULevel* InLevel, UWorld* InWorld)
+{
+	if ( InLevel == nullptr && InWorld == GetWorld() )
+	{
+		RemoveFromParent();
+		MarkPendingKill();
+	}
+}
+
 void UUserWidget::RemoveFromViewport()
 {
 	RemoveFromParent();
@@ -686,6 +496,8 @@ void UUserWidget::RemoveFromParent()
 				{
 					ViewportClient->RemoveViewportWidgetForPlayer(LocalPlayer, WidgetHostRef);
 				}
+
+				FWorldDelegates::LevelRemovedFromWorld.RemoveAll(this);
 			}
 		}
 	}
@@ -840,8 +652,60 @@ const FText UUserWidget::GetPaletteCategory()
 
 #endif
 
+void UUserWidget::OnAnimationStarted_Implementation(const UWidgetAnimation* Animation)
+{
+
+}
+
+void UUserWidget::OnAnimationFinished_Implementation(const UWidgetAnimation* Animation)
+{
+
+}
+
 // Native handling for SObjectWidget
 
+void UUserWidget::NativeConstruct()
+{
+	Construct();
+}
+
+void UUserWidget::NativeDestruct()
+{
+	Destruct();
+}
+
+void UUserWidget::NativeTick(const FGeometry& MyGeometry, float InDeltaTime)
+{
+	GInitRunaway();
+
+	// Update active movie scenes
+	for ( UUMGSequencePlayer* Player : ActiveSequencePlayers )
+	{
+		Player->Tick(InDeltaTime);
+	}
+
+	// The process of ticking the players above can stop them so we remove them after all players have ticked
+	for ( UUMGSequencePlayer* StoppedPlayer : StoppedSequencePlayers )
+	{
+		ActiveSequencePlayers.Remove(StoppedPlayer);
+	}
+
+	StoppedSequencePlayers.Empty();
+
+	if ( !bDesignTime )
+	{
+		UWorld* World = GetWorld();
+		if ( World )
+		{
+			// Update any latent actions we have for this actor
+			FLatentActionManager& LatentActionManager = World->GetLatentActionManager();
+			LatentActionManager.ProcessLatentActions(this, InDeltaTime);
+		}
+	}
+
+	Tick(MyGeometry, InDeltaTime);
+}
+
 void UUserWidget::NativePaint( FPaintContext& InContext ) const 
 {
 	OnPaint( InContext );
@@ -977,6 +841,157 @@ FReply UUserWidget::NativeOnMotionDetected( const FGeometry& InGeometry, const F
 	return OnMotionDetected( InGeometry, InMotionEvent ).NativeReply;
 }
 
+/////////////////////////////////////////////////////
+//
+// DEPRECATED
+//
+/////////////////////////////////////////////////////
+
+void UUserWidget::Construct_Implementation()
+{
+
+}
+
+void UUserWidget::Tick_Implementation(FGeometry MyGeometry, float InDeltaTime)
+{
+
+}
+
+void UUserWidget::OnPaint_Implementation(FPaintContext& Context) const
+{
+
+}
+
+FEventReply UUserWidget::OnFocusReceived_Implementation(FGeometry MyGeometry, FFocusEvent InFocusEvent)
+{
+	return UWidgetBlueprintLibrary::Unhandled();
+}
+
+void UUserWidget::OnFocusLost_Implementation(FFocusEvent InFocusEvent)
+{
+
+}
+
+FEventReply UUserWidget::OnKeyChar_Implementation(FGeometry MyGeometry, FCharacterEvent InCharacterEvent)
+{
+	return UWidgetBlueprintLibrary::Unhandled();
+}
+
+FEventReply UUserWidget::OnPreviewKeyDown_Implementation(FGeometry MyGeometry, FKeyEvent InKeyEvent)
+{
+	return UWidgetBlueprintLibrary::Unhandled();
+}
+
+FEventReply UUserWidget::OnKeyDown_Implementation(FGeometry MyGeometry, FKeyEvent InKeyEvent)
+{
+	return UWidgetBlueprintLibrary::Unhandled();
+}
+
+FEventReply UUserWidget::OnKeyUp_Implementation(FGeometry MyGeometry, FKeyEvent InKeyEvent)
+{
+	return UWidgetBlueprintLibrary::Unhandled();
+}
+
+FEventReply UUserWidget::OnAnalogValueChanged_Implementation(FGeometry MyGeometry, FAnalogInputEvent InAnalogInputEvent)
+{
+	return UWidgetBlueprintLibrary::Unhandled();
+}
+
+FEventReply UUserWidget::OnMouseButtonDown_Implementation(FGeometry MyGeometry, const FPointerEvent& MouseEvent)
+{
+	return UWidgetBlueprintLibrary::Unhandled();
+}
+
+FEventReply UUserWidget::OnPreviewMouseButtonDown_Implementation(FGeometry MyGeometry, const FPointerEvent& MouseEvent)
+{
+	return UWidgetBlueprintLibrary::Unhandled();
+}
+
+FEventReply UUserWidget::OnMouseButtonUp_Implementation(FGeometry MyGeometry, const FPointerEvent& MouseEvent)
+{
+	return UWidgetBlueprintLibrary::Unhandled();
+}
+
+FEventReply UUserWidget::OnMouseMove_Implementation(FGeometry MyGeometry, const FPointerEvent& MouseEvent)
+{
+	return UWidgetBlueprintLibrary::Unhandled();
+}
+
+void UUserWidget::OnMouseEnter_Implementation(FGeometry MyGeometry, const FPointerEvent& MouseEvent)
+{
+
+}
+
+void UUserWidget::OnMouseLeave_Implementation(const FPointerEvent& MouseEvent)
+{
+
+}
+
+FEventReply UUserWidget::OnMouseWheel_Implementation(FGeometry MyGeometry, const FPointerEvent& MouseEvent)
+{
+	return UWidgetBlueprintLibrary::Unhandled();
+}
+
+FEventReply UUserWidget::OnMouseButtonDoubleClick_Implementation(FGeometry InMyGeometry, const FPointerEvent& InMouseEvent)
+{
+	return UWidgetBlueprintLibrary::Unhandled();
+}
+
+void UUserWidget::OnDragDetected_Implementation(FGeometry MyGeometry, const FPointerEvent& PointerEvent, UDragDropOperation*& Operation)
+{
+
+}
+
+void UUserWidget::OnDragCancelled_Implementation(const FPointerEvent& PointerEvent, UDragDropOperation* Operation)
+{
+
+}
+
+void UUserWidget::OnDragEnter_Implementation(FGeometry MyGeometry, FPointerEvent PointerEvent, UDragDropOperation* Operation)
+{
+
+}
+
+void UUserWidget::OnDragLeave_Implementation(FPointerEvent PointerEvent, UDragDropOperation* Operation)
+{
+
+}
+
+bool UUserWidget::OnDragOver_Implementation(FGeometry MyGeometry, FPointerEvent PointerEvent, UDragDropOperation* Operation)
+{
+	return false;
+}
+
+bool UUserWidget::OnDrop_Implementation(FGeometry MyGeometry, FPointerEvent PointerEvent, UDragDropOperation* Operation)
+{
+	return false;
+}
+
+FEventReply UUserWidget::OnTouchGesture_Implementation(FGeometry MyGeometry, const FPointerEvent& GestureEvent)
+{
+	return UWidgetBlueprintLibrary::Unhandled();
+}
+
+FEventReply UUserWidget::OnTouchStarted_Implementation(FGeometry MyGeometry, const FPointerEvent& InTouchEvent)
+{
+	return UWidgetBlueprintLibrary::Unhandled();
+}
+
+FEventReply UUserWidget::OnTouchMoved_Implementation(FGeometry MyGeometry, const FPointerEvent& InTouchEvent)
+{
+	return UWidgetBlueprintLibrary::Unhandled();
+}
+
+FEventReply UUserWidget::OnTouchEnded_Implementation(FGeometry MyGeometry, const FPointerEvent& InTouchEvent)
+{
+	return UWidgetBlueprintLibrary::Unhandled();
+}
+
+FEventReply UUserWidget::OnMotionDetected_Implementation(FGeometry MyGeometry, FMotionEvent InMotionEvent)
+{
+	return UWidgetBlueprintLibrary::Unhandled();
+}
+
 /////////////////////////////////////////////////////
 
 #undef LOCTEXT_NAMESPACE
diff --git a/Engine/Source/Runtime/UMG/Public/Blueprint/UserWidget.h b/Engine/Source/Runtime/UMG/Public/Blueprint/UserWidget.h
--- a/Engine/Source/Runtime/UMG/Public/Blueprint/UserWidget.h
+++ b/Engine/Source/Runtime/UMG/Public/Blueprint/UserWidget.h
@@ -126,9 +126,11 @@ DECLARE_DYNAMIC_MULTICAST_DELEGATE(FOnConstructEvent);
 UCLASS(Abstract, editinlinenew, BlueprintType, Blueprintable, meta=( Category="User Controls" ))
 class UMG_API UUserWidget : public UWidget, public INamedSlotInterface
 {
-	GENERATED_UCLASS_BODY()
+	GENERATED_BODY()
 
 public:
+	UUserWidget(const FObjectInitializer& ObjectInitializer);
+
 	//UObject interface
 	virtual void PostInitProperties() override;
 	virtual class UWorld* GetWorld() const override;
@@ -251,21 +253,37 @@ public:
 	 * Called after the underlying slate widget is constructed.  Depending on how the slate object is used
 	 * this event may be called multiple times due to adding and removing from the hierarchy.
 	 */
-	UFUNCTION(BlueprintNativeEvent, BlueprintCosmetic, Category="User Interface", meta=( Keywords="Begin Play" ))
+	UFUNCTION(BlueprintImplementableEvent, BlueprintCosmetic, Category="User Interface", meta=( Keywords="Begin Play" ))
 	void Construct();
 
+	DEPRECATED(4.8, "Use NativeConstruct")
+	virtual void Construct_Implementation();
+
+	/**
+	 * Called when a widget is no longer referenced causing the slate resource to destroyed.  Just like
+	 * Construct this event can be called multiple times.
+	 */
+	UFUNCTION(BlueprintImplementableEvent, BlueprintCosmetic, Category="User Interface", meta=( Keywords="End Play, Destroy" ))
+	void Destruct();
+
 	/**
 	 * Ticks this widget.  Override in derived classes, but always call the parent implementation.
 	 *
 	 * @param  MyGeometry The space allotted for this widget
 	 * @param  InDeltaTime  Real time passed since last tick
 	 */
-	UFUNCTION(BlueprintNativeEvent, BlueprintCosmetic, Category="User Interface")
+	UFUNCTION(BlueprintImplementableEvent, BlueprintCosmetic, Category="User Interface")
 	void Tick(FGeometry MyGeometry, float InDeltaTime);
 
-	UFUNCTION(BlueprintNativeEvent, BlueprintCosmetic, Category="User Interface | Painting")
+	DEPRECATED(4.8, "Use NativeTick")
+	virtual void Tick_Implementation(FGeometry MyGeometry, float InDeltaTime);
+
+	UFUNCTION(BlueprintImplementableEvent, BlueprintCosmetic, Category="User Interface | Painting")
 	void OnPaint(UPARAM(ref) FPaintContext& Context) const;
 
+	DEPRECATED(4.8, "Use NativePaint")
+	virtual void OnPaint_Implementation(FPaintContext& Context) const;
+
 	/**
 	 * Called when keyboard focus is given to this widget.  This event does not bubble.
 	 *
@@ -273,17 +291,23 @@ public:
 	 * @param InFocusEvent  FocusEvent
 	 * @return  Returns whether the event was handled, along with other possible actions
 	 */
-	UFUNCTION(BlueprintNativeEvent, BlueprintCosmetic, Category="Input")
+	UFUNCTION(BlueprintImplementableEvent, BlueprintCosmetic, Category="Input")
 	FEventReply OnFocusReceived(FGeometry MyGeometry, FFocusEvent InFocusEvent);
 
+	DEPRECATED(4.8, "Use NativeOnFocusReceived")
+	FEventReply OnFocusReceived_Implementation(FGeometry MyGeometry, FFocusEvent InFocusEvent);
+
 	/**
 	 * Called when this widget loses focus.  This event does not bubble.
 	 *
 	 * @param  InFocusEvent  FocusEvent
 	 */
-	UFUNCTION(BlueprintNativeEvent, BlueprintCosmetic, Category="Input")
+	UFUNCTION(BlueprintImplementableEvent, BlueprintCosmetic, Category="Input")
 	void OnFocusLost(FFocusEvent InFocusEvent);
 
+	DEPRECATED(4.8, "Use NativeOnFocusLost")
+	virtual void OnFocusLost_Implementation(FFocusEvent InFocusEvent);
+
 	/**
 	 * Called after a character is entered while this widget has focus
 	 *
@@ -291,9 +315,12 @@ public:
 	 * @param  InCharacterEvent  Character event
 	 * @return  Returns whether the event was handled, along with other possible actions
 	 */
-	UFUNCTION(BlueprintNativeEvent, BlueprintCosmetic, Category="Input")
+	UFUNCTION(BlueprintImplementableEvent, BlueprintCosmetic, Category="Input")
 	FEventReply OnKeyChar(FGeometry MyGeometry, FCharacterEvent InCharacterEvent);
 
+	DEPRECATED(4.8, "Use NativeOnKeyChar")
+	virtual FEventReply OnKeyChar_Implementation(FGeometry MyGeometry, FCharacterEvent InCharacterEvent);
+
 	/**
 	 * Called after a key (keyboard, controller, ...) is pressed when this widget or a child of this widget has focus
 	 * If a widget handles this event, OnKeyDown will *not* be passed to the focused widget.
@@ -305,9 +332,12 @@ public:
 	 * @param  InKeyEvent  Key event
 	 * @return  Returns whether the event was handled, along with other possible actions
 	 */
-	UFUNCTION(BlueprintNativeEvent, Category="Input")
+	UFUNCTION(BlueprintImplementableEvent, Category="Input")
 	FEventReply OnPreviewKeyDown(FGeometry MyGeometry, FKeyEvent InKeyEvent);
 
+	DEPRECATED(4.8, "Use NativeOnPreviewKeyDown")
+	virtual FEventReply OnPreviewKeyDown_Implementation(FGeometry MyGeometry, FKeyEvent InKeyEvent);
+
 	/**
 	 * Called after a key (keyboard, controller, ...) is pressed when this widget has focus (this event bubbles if not handled)
 	 *
@@ -315,9 +345,12 @@ public:
 	 * @param  InKeyEvent  Key event
 	 * @return  Returns whether the event was handled, along with other possible actions
 	 */
-	UFUNCTION(BlueprintNativeEvent, BlueprintCosmetic, Category="Input")
+	UFUNCTION(BlueprintImplementableEvent, BlueprintCosmetic, Category="Input")
 	FEventReply OnKeyDown(FGeometry MyGeometry, FKeyEvent InKeyEvent);
 
+	DEPRECATED(4.8, "Use NativeOnKeyDown")
+	virtual FEventReply OnKeyDown_Implementation(FGeometry MyGeometry, FKeyEvent InKeyEvent);
+
 	/**
 	 * Called after a key (keyboard, controller, ...) is released when this widget has focus
 	 *
@@ -325,9 +358,12 @@ public:
 	 * @param  InKeyEvent  Key event
 	 * @return  Returns whether the event was handled, along with other possible actions
 	 */
-	UFUNCTION(BlueprintNativeEvent, BlueprintCosmetic, Category="Input")
+	UFUNCTION(BlueprintImplementableEvent, BlueprintCosmetic, Category="Input")
 	FEventReply OnKeyUp(FGeometry MyGeometry, FKeyEvent InKeyEvent);
 
+	DEPRECATED(4.8, "Use NativeOnKeyUp")
+	virtual FEventReply OnKeyUp_Implementation(FGeometry MyGeometry, FKeyEvent InKeyEvent);
+
 	/**
 	* Called when an analog value changes on a button that supports analog
 	*
@@ -335,9 +371,12 @@ public:
 	* @param  InAnalogInputEvent  Analog Event
 	* @return  Returns whether the event was handled, along with other possible actions
 	*/
-	UFUNCTION(BlueprintNativeEvent, Category = "Input")
+	UFUNCTION(BlueprintImplementableEvent, Category = "Input")
 	FEventReply OnAnalogValueChanged(FGeometry MyGeometry, FAnalogInputEvent InAnalogInputEvent);
 
+	DEPRECATED(4.8, "Use NativeOnAnalogValueChanged")
+	virtual FEventReply OnAnalogValueChanged_Implementation(FGeometry MyGeometry, FAnalogInputEvent InAnalogInputEvent);
+
 	/**
 	 * The system calls this method to notify the widget that a mouse button was pressed within it. This event is bubbled.
 	 *
@@ -345,9 +384,12 @@ public:
 	 * @param MouseEvent Information about the input event
 	 * @return Whether the event was handled along with possible requests for the system to take action.
 	 */
-	UFUNCTION(BlueprintNativeEvent, BlueprintCosmetic, Category="Mouse")
+	UFUNCTION(BlueprintImplementableEvent, BlueprintCosmetic, Category="Mouse")
 	FEventReply OnMouseButtonDown(FGeometry MyGeometry, const FPointerEvent& MouseEvent);
 
+	DEPRECATED(4.8, "Use NativeOnMouseButtonDown")
+	virtual FEventReply OnMouseButtonDown_Implementation(FGeometry MyGeometry, const FPointerEvent& MouseEvent);
+
 	/**
 	 * Just like OnMouseButtonDown, but tunnels instead of bubbling.
 	 * If this even is handled, OnMouseButtonDown will not be sent.
@@ -359,9 +401,12 @@ public:
 	 * @param MouseEvent Information about the input event
 	 * @return Whether the event was handled along with possible requests for the system to take action.
 	 */
-	UFUNCTION(BlueprintNativeEvent, BlueprintCosmetic, Category="Mouse")
+	UFUNCTION(BlueprintImplementableEvent, BlueprintCosmetic, Category="Mouse")
 	FEventReply OnPreviewMouseButtonDown(FGeometry MyGeometry, const FPointerEvent& MouseEvent);
 
+	DEPRECATED(4.8, "Use NativeOnPreviewMouseButtonDown")
+	virtual FEventReply OnPreviewMouseButtonDown_Implementation(FGeometry MyGeometry, const FPointerEvent& MouseEvent);
+
 	/**
 	 * The system calls this method to notify the widget that a mouse button was release within it. This event is bubbled.
 	 *
@@ -369,9 +414,12 @@ public:
 	 * @param MouseEvent Information about the input event
 	 * @return Whether the event was handled along with possible requests for the system to take action.
 	 */
-	UFUNCTION(BlueprintNativeEvent, BlueprintCosmetic, Category="Mouse")
+	UFUNCTION(BlueprintImplementableEvent, BlueprintCosmetic, Category="Mouse")
 	FEventReply OnMouseButtonUp(FGeometry MyGeometry, const FPointerEvent& MouseEvent);
 
+	DEPRECATED(4.8, "Use NativeOnMouseButtonUp")
+	virtual FEventReply OnMouseButtonUp_Implementation(FGeometry MyGeometry, const FPointerEvent& MouseEvent);
+
 	/**
 	 * The system calls this method to notify the widget that a mouse moved within it. This event is bubbled.
 	 *
@@ -379,35 +427,47 @@ public:
 	 * @param MouseEvent Information about the input event
 	 * @return Whether the event was handled along with possible requests for the system to take action.
 	 */
-	UFUNCTION(BlueprintNativeEvent, BlueprintCosmetic, Category="Mouse")
+	UFUNCTION(BlueprintImplementableEvent, BlueprintCosmetic, Category="Mouse")
 	FEventReply OnMouseMove(FGeometry MyGeometry, const FPointerEvent& MouseEvent);
 
+	DEPRECATED(4.8, "Use NativeOnMouseMove")
+	virtual FEventReply OnMouseMove_Implementation(FGeometry MyGeometry, const FPointerEvent& MouseEvent);
+
 	/**
 	 * The system will use this event to notify a widget that the cursor has entered it. This event is NOT bubbled.
 	 *
 	 * @param MyGeometry The Geometry of the widget receiving the event
 	 * @param MouseEvent Information about the input event
 	 */
-	UFUNCTION(BlueprintNativeEvent, BlueprintCosmetic, Category="Mouse")
+	UFUNCTION(BlueprintImplementableEvent, BlueprintCosmetic, Category="Mouse")
 	void OnMouseEnter(FGeometry MyGeometry, const FPointerEvent& MouseEvent);
 
+	DEPRECATED(4.8, "Use NativeOnMouseEnter")
+	virtual void OnMouseEnter_Implementation(FGeometry MyGeometry, const FPointerEvent& MouseEvent);
+
 	/**
 	 * The system will use this event to notify a widget that the cursor has left it. This event is NOT bubbled.
 	 *
 	 * @param MouseEvent Information about the input event
 	 */
-	UFUNCTION(BlueprintNativeEvent, BlueprintCosmetic, Category="Mouse")
+	UFUNCTION(BlueprintImplementableEvent, BlueprintCosmetic, Category="Mouse")
 	void OnMouseLeave(const FPointerEvent& MouseEvent);
 
+	DEPRECATED(4.8, "Use NativeOnMouseLeave")
+	virtual void OnMouseLeave_Implementation(const FPointerEvent& MouseEvent);
+
 	/**
 	 * Called when the mouse wheel is spun. This event is bubbled.
 	 *
 	 * @param  MouseEvent  Mouse event
 	 * @return  Returns whether the event was handled, along with other possible actions
 	 */
-	UFUNCTION(BlueprintNativeEvent, BlueprintCosmetic, Category="Mouse")
+	UFUNCTION(BlueprintImplementableEvent, BlueprintCosmetic, Category="Mouse")
 	FEventReply OnMouseWheel(FGeometry MyGeometry, const FPointerEvent& MouseEvent);
 
+	DEPRECATED(4.8, "Use NativeOnMouseWheel")
+	virtual FEventReply OnMouseWheel_Implementation(FGeometry MyGeometry, const FPointerEvent& MouseEvent);
+
 	/**
 	 * Called when a mouse button is double clicked.  Override this in derived classes.
 	 *
@@ -415,11 +475,14 @@ public:
 	 * @param  InMouseEvent  Mouse button event
 	 * @return  Returns whether the event was handled, along with other possible actions
 	 */
-	UFUNCTION(BlueprintNativeEvent, BlueprintCosmetic, Category="Mouse")
+	UFUNCTION(BlueprintImplementableEvent, BlueprintCosmetic, Category="Mouse")
 	FEventReply OnMouseButtonDoubleClick(FGeometry InMyGeometry, const FPointerEvent& InMouseEvent);
 
+	DEPRECATED(4.8, "Use NativeOnMouseButtonDoubleClick")
+	virtual FEventReply OnMouseButtonDoubleClick_Implementation(FGeometry InMyGeometry, const FPointerEvent& InMouseEvent);
+
 	// TODO
-	//UFUNCTION(BlueprintNativeEvent, Category="Mouse")
+	//UFUNCTION(BlueprintImplementableEvent, Category="Mouse")
 	//FCursorReply OnCursorQuery(FGeometry MyGeometry, const FPointerEvent& CursorEvent) const;
 
 	// TODO
@@ -432,9 +495,12 @@ public:
 	 * @param  PointerEvent  MouseMove that triggered the drag
 	 * @param  Operation     The drag operation that was detected.
 	 */
-	UFUNCTION(BlueprintNativeEvent, BlueprintCosmetic, Category="Drag and Drop")
+	UFUNCTION(BlueprintImplementableEvent, BlueprintCosmetic, Category="Drag and Drop")
 	void OnDragDetected(FGeometry MyGeometry, const FPointerEvent& PointerEvent, UDragDropOperation*& Operation);
 
+	DEPRECATED(4.8, "Use NativeOnDragDetected")
+	virtual void OnDragDetected_Implementation(FGeometry MyGeometry, const FPointerEvent& PointerEvent, UDragDropOperation*& Operation);
+
 	/**
 	 * Called when the user cancels the drag operation, typically when they simply release the mouse button after
 	 * beginning a drag operation, but failing to complete the drag.
@@ -442,8 +508,11 @@ public:
 	 * @param  PointerEvent  Last mouse event from when the drag was canceled.
 	 * @param  Operation     The drag operation that was canceled.
 	 */
-	UFUNCTION(BlueprintNativeEvent, BlueprintCosmetic, Category="Drag and Drop")
+	UFUNCTION(BlueprintImplementableEvent, BlueprintCosmetic, Category="Drag and Drop")
 	void OnDragCancelled(const FPointerEvent& PointerEvent, UDragDropOperation* Operation);
+
+	DEPRECATED(4.8, "Use NativeOnDragCancelled")
+	virtual void OnDragCancelled_Implementation(const FPointerEvent& PointerEvent, UDragDropOperation* Operation);
 	
 	/**
 	 * Called during drag and drop when the drag enters the widget.
@@ -452,18 +521,24 @@ public:
 	 * @param PointerEvent   The mouse event from when the drag entered the widget.
 	 * @param Operation      The drag operation that entered the widget.
 	 */
-	UFUNCTION(BlueprintNativeEvent, BlueprintCosmetic, Category="Drag and Drop")
+	UFUNCTION(BlueprintImplementableEvent, BlueprintCosmetic, Category="Drag and Drop")
 	void OnDragEnter(FGeometry MyGeometry, FPointerEvent PointerEvent, UDragDropOperation* Operation);
 
+	DEPRECATED(4.8, "Use NativeOnDragEnter")
+	virtual void OnDragEnter_Implementation(FGeometry MyGeometry, FPointerEvent PointerEvent, UDragDropOperation* Operation);
+
 	/**
 	 * Called during drag and drop when the drag leaves the widget.
 	 *
 	 * @param PointerEvent   The mouse event from when the drag left the widget.
 	 * @param Operation      The drag operation that entered the widget.
 	 */
-	UFUNCTION(BlueprintNativeEvent, BlueprintCosmetic, Category="Drag and Drop")
+	UFUNCTION(BlueprintImplementableEvent, BlueprintCosmetic, Category="Drag and Drop")
 	void OnDragLeave(FPointerEvent PointerEvent, UDragDropOperation* Operation);
 
+	DEPRECATED(4.8, "Use NativeOnDragLeave")
+	void OnDragLeave_Implementation(FPointerEvent PointerEvent, UDragDropOperation* Operation);
+
 	/**
 	 * Called during drag and drop when the the mouse is being dragged over a widget.
 	 *
@@ -473,9 +548,12 @@ public:
 	 *
 	 * @return 'true' to indicate that you handled the drag over operation.  Returning 'false' will cause the operation to continue to bubble up.
 	 */
-	UFUNCTION(BlueprintNativeEvent, BlueprintCosmetic, Category="Drag and Drop")
+	UFUNCTION(BlueprintImplementableEvent, BlueprintCosmetic, Category="Drag and Drop")
 	bool OnDragOver(FGeometry MyGeometry, FPointerEvent PointerEvent, UDragDropOperation* Operation);
 
+	DEPRECATED(4.8, "Use NativeOnDragOver")
+	virtual bool OnDragOver_Implementation(FGeometry MyGeometry, FPointerEvent PointerEvent, UDragDropOperation* Operation);
+
 	/**
 	 * Called when the user is dropping something onto a widget.  Ends the drag and drop operation, even if no widget handles this.
 	 *
@@ -485,17 +563,19 @@ public:
 	 * 
 	 * @return 'true' to indicate you handled the drop operation.
 	 */
-	UFUNCTION(BlueprintNativeEvent, BlueprintCosmetic, Category="Drag and Drop")
+	UFUNCTION(BlueprintImplementableEvent, BlueprintCosmetic, Category="Drag and Drop")
 	bool OnDrop(FGeometry MyGeometry, FPointerEvent PointerEvent, UDragDropOperation* Operation);
 
+	DEPRECATED(4.8, "Use NativeOnDrop")
+	virtual bool OnDrop_Implementation(FGeometry MyGeometry, FPointerEvent PointerEvent, UDragDropOperation* Operation);
 
-	UFUNCTION(BlueprintNativeEvent, meta = (DeprecatedFunction, DeprecationMessage = "Use OnKeyDown() instead"), Category = "Gamepad Input")
+	UFUNCTION(BlueprintImplementableEvent, meta = ( DeprecatedFunction, DeprecationMessage = "Use OnKeyDown() instead" ), Category = "Gamepad Input")
 	FEventReply OnControllerButtonPressed(FGeometry MyGeometry, FControllerEvent ControllerEvent);
 
-	UFUNCTION(BlueprintNativeEvent, meta = (DeprecatedFunction, DeprecationMessage = "Use OnKeyUp() instead"), Category = "Gamepad Input")
+	UFUNCTION(BlueprintImplementableEvent, meta = ( DeprecatedFunction, DeprecationMessage = "Use OnKeyUp() instead" ), Category = "Gamepad Input")
 	FEventReply OnControllerButtonReleased(FGeometry MyGeometry, FControllerEvent ControllerEvent);
 
-	UFUNCTION(BlueprintNativeEvent, meta = (DeprecatedFunction, DeprecationMessage = "Use OnAnalogValueChanged() instead"), Category = "Gamepad Input")
+	UFUNCTION(BlueprintImplementableEvent, meta = ( DeprecatedFunction, DeprecationMessage = "Use OnAnalogValueChanged() instead" ), Category = "Gamepad Input")
 	FEventReply OnControllerAnalogValueChanged(FGeometry MyGeometry, FControllerEvent ControllerEvent);
 
 	/**
@@ -505,36 +585,48 @@ public:
 	 * @param  GestureEvent  gesture event
 	 * @return  Returns whether the event was handled, along with other possible actions
 	 */
-	UFUNCTION(BlueprintNativeEvent, BlueprintCosmetic, Category="Touch Input")
+	UFUNCTION(BlueprintImplementableEvent, BlueprintCosmetic, Category="Touch Input")
 	FEventReply OnTouchGesture(FGeometry MyGeometry, const FPointerEvent& GestureEvent);
 
+	DEPRECATED(4.8, "Use NativeOnTouchGesture")
+	virtual FEventReply OnTouchGesture_Implementation(FGeometry MyGeometry, const FPointerEvent& GestureEvent);
+
 	/**
 	 * Called when a touchpad touch is started (finger down)
 	 * 
 	 * @param MyGeometry    The geometry of the widget receiving the event.
 	 * @param InTouchEvent	The touch event generated
 	 */
-	UFUNCTION(BlueprintNativeEvent, BlueprintCosmetic, Category="Touch Input")
+	UFUNCTION(BlueprintImplementableEvent, BlueprintCosmetic, Category="Touch Input")
 	FEventReply OnTouchStarted(FGeometry MyGeometry, const FPointerEvent& InTouchEvent);
 
+	DEPRECATED(4.8, "Use NativeOnTouchStarted")
+	virtual FEventReply OnTouchStarted_Implementation(FGeometry MyGeometry, const FPointerEvent& InTouchEvent);
+	
 	/**
 	 * Called when a touchpad touch is moved  (finger moved)
 	 * 
 	 * @param MyGeometry    The geometry of the widget receiving the event.
 	 * @param InTouchEvent	The touch event generated
 	 */
-	UFUNCTION(BlueprintNativeEvent, BlueprintCosmetic, Category="Touch Input")
+	UFUNCTION(BlueprintImplementableEvent, BlueprintCosmetic, Category="Touch Input")
 	FEventReply OnTouchMoved(FGeometry MyGeometry, const FPointerEvent& InTouchEvent);
 
+	DEPRECATED(4.8, "Use NativeOnTouchMoved")
+	virtual FEventReply OnTouchMoved_Implementation(FGeometry MyGeometry, const FPointerEvent& InTouchEvent);
+
 	/**
 	 * Called when a touchpad touch is ended (finger lifted)
 	 * 
 	 * @param MyGeometry    The geometry of the widget receiving the event.
 	 * @param InTouchEvent	The touch event generated
 	 */
-	UFUNCTION(BlueprintNativeEvent, BlueprintCosmetic, Category="Touch Input")
+	UFUNCTION(BlueprintImplementableEvent, BlueprintCosmetic, Category="Touch Input")
 	FEventReply OnTouchEnded(FGeometry MyGeometry, const FPointerEvent& InTouchEvent);
 
+	DEPRECATED(4.8, "Use NativeOnTouchEnded")
+	virtual FEventReply OnTouchEnded_Implementation(FGeometry MyGeometry, const FPointerEvent& InTouchEvent);
+	
 	/**
 	 * Called when motion is detected (controller or device)
 	 * e.g. Someone tilts or shakes their controller.
@@ -542,9 +634,14 @@ public:
 	 * @param MyGeometry    The geometry of the widget receiving the event.
 	 * @param MotionEvent	The motion event generated
 	 */
-	UFUNCTION(BlueprintNativeEvent, BlueprintCosmetic, Category="Touch Input")
+	UFUNCTION(BlueprintImplementableEvent, BlueprintCosmetic, Category="Touch Input")
 	FEventReply OnMotionDetected(FGeometry MyGeometry, FMotionEvent InMotionEvent);
 
+	DEPRECATED(4.8, "Use NativeOnMotionDetected")
+	virtual FEventReply OnMotionDetected_Implementation(FGeometry MyGeometry, FMotionEvent InMotionEvent);
+
+public:
+
 	/**
 	 * Called when an animation is started.
 	 *
@@ -553,6 +650,8 @@ public:
 	UFUNCTION( BlueprintNativeEvent, BlueprintCosmetic, Category = "Animation" )
 	void OnAnimationStarted( const UWidgetAnimation* Animation );
 
+	virtual void OnAnimationStarted_Implementation(const UWidgetAnimation* Animation);
+
 	/**
 	 * Called when an animation has either played all the way through or is stopped
 	 *
@@ -561,6 +660,8 @@ public:
 	UFUNCTION( BlueprintNativeEvent, BlueprintCosmetic, Category = "Animation" )
 	void OnAnimationFinished( const UWidgetAnimation* Animation );
 
+	virtual void OnAnimationFinished_Implementation(const UWidgetAnimation* Animation);
+
 public:
 
 	/**
@@ -625,9 +726,6 @@ public:
 	/** @returns The uobject widget corresponding to a given name */
 	UWidget* GetWidgetFromName(const FName& Name) const;
 
-	/** Ticks this widget and forwards to the underlying widget to tick */
-	virtual void NativeTick(const FGeometry& MyGeometry, float InDeltaTime );
-
 	// Begin UObject interface
 	virtual void PreSave() override;
 	// End UObject interface
@@ -701,6 +799,12 @@ protected:
 	/** Adds the widget to the screen, either to the viewport or to the player's screen depending on if the LocalPlayer is null. */
 	void AddToScreen(ULocalPlayer* LocalPlayer, int32 ZOrder);
 
+	/**
+	 * Called when a top level widget is in the viewport and the world is potentially coming to and end. When this occurs, 
+	 * it's not save to keep widgets on the screen.  We automatically remove them when this happens and mark them for pending kill.
+	 */
+	void OnLevelRemovedFromWorld(ULevel* InLevel, UWorld* InWorld);
+
 	virtual TSharedRef<SWidget> RebuildWidget() override;
 	virtual void OnWidgetRebuilt() override;
 
@@ -711,6 +815,9 @@ protected:
 	//native SObjectWidget methods
 	friend class SObjectWidget;
 
+	virtual void NativeConstruct();
+	virtual void NativeDestruct();
+	virtual void NativeTick(const FGeometry& MyGeometry, float InDeltaTime);
 	virtual void NativePaint( FPaintContext& InContext ) const;
 	virtual FReply NativeOnFocusReceived( const FGeometry& InGeometry, const FFocusEvent& InFocusEvent );
 	virtual void NativeOnFocusLost( const FFocusEvent& InFocusEvent );
