Blueprint Call Function nodes will now display localized function names.

diff --git a/Engine/Source/Runtime/UMG/Public/Components/MultiLineEditableTextBox.h b/Engine/Source/Runtime/UMG/Public/Components/MultiLineEditableTextBox.h
--- a/Engine/Source/Runtime/UMG/Public/Components/MultiLineEditableTextBox.h
+++ b/Engine/Source/Runtime/UMG/Public/Components/MultiLineEditableTextBox.h
@@ -78,14 +78,14 @@ public:
 public:
 
 	/**  */
-	UFUNCTION(BlueprintCallable, Category="Widget", meta=(FriendlyName="GetText (Multi-Line Text Box)"))
+	UFUNCTION(BlueprintCallable, Category="Widget", meta=(DisplayName="GetText (Multi-Line Text Box)"))
 	FText GetText() const;
 
 	/**  */
-	UFUNCTION(BlueprintCallable, Category="Widget", meta=(FriendlyName="SetText (Multi-Line Text Box)"))
+	UFUNCTION(BlueprintCallable, Category="Widget", meta=(DisplayName="SetText (Multi-Line Text Box)"))
 	void SetText(FText InText);
 
-	UFUNCTION(BlueprintCallable, Category="Widget", meta=(FriendlyName="SetError (Multi-Line Text Box)"))
+	UFUNCTION(BlueprintCallable, Category="Widget", meta=(DisplayName="SetError (Multi-Line Text Box)"))
 	void SetError(FText InError);
 
 	//TODO UMG Add Set ReadOnlyForegroundColor
diff --git a/Engine/Source/Runtime/UMG/Public/Components/TextBlock.h b/Engine/Source/Runtime/UMG/Public/Components/TextBlock.h
--- a/Engine/Source/Runtime/UMG/Public/Components/TextBlock.h
+++ b/Engine/Source/Runtime/UMG/Public/Components/TextBlock.h
@@ -117,7 +117,7 @@ public:
 	 * Gets the widget text
 	 * @return The widget text
 	 */
-	UFUNCTION(BlueprintCallable, Category="Widget", meta=(FriendlyName="GetText (Text)"))
+	UFUNCTION(BlueprintCallable, Category="Widget", meta=(DisplayName="GetText (Text)"))
 	FText GetText() const;
 
 	/**
@@ -125,7 +125,7 @@ public:
 	 * Warning: This will wipe any binding created for the Text property!
 	 * @param InText The text to assign to the widget
 	 */
-	UFUNCTION(BlueprintCallable, Category="Widget", meta=(FriendlyName="SetText (Text)"))
+	UFUNCTION(BlueprintCallable, Category="Widget", meta=(DisplayName="SetText (Text)"))
 	void SetText(FText InText);
 
 	// UWidget interface
