UMG - UserWidgets that have been added to the root viewport are automatically removed from the viewport and marked for pending kill when a full level transitions occurs now.

diff --git a/Engine/Source/Runtime/UMG/Private/UserWidget.cpp b/Engine/Source/Runtime/UMG/Private/UserWidget.cpp
--- a/Engine/Source/Runtime/UMG/Private/UserWidget.cpp
+++ b/Engine/Source/Runtime/UMG/Private/UserWidget.cpp
@@ -456,6 +456,7 @@ void UUserWidget::AddToScreen(ULocalPlayer* Player, int32 ZOrder)
 					ViewportClient->AddViewportWidgetContent(RootWidget, ZOrder + 10);
 				}
 
+				// Widgets added to the viewport are automatically removed if the persistant level is unloaded.
 				FWorldDelegates::LevelRemovedFromWorld.AddUObject(this, &UUserWidget::OnLevelRemovedFromWorld);
 			}
 		}
@@ -464,6 +465,9 @@ void UUserWidget::AddToScreen(ULocalPlayer* Player, int32 ZOrder)
 
 void UUserWidget::OnLevelRemovedFromWorld(ULevel* InLevel, UWorld* InWorld)
 {
+	// If the InLevel is null, it's a signal that the entire world is about to disappear, so
+	// go ahead and remove this widget from the viewport, it could be holding onto too many
+	// dangerous actor references that won't carry over into the next world.
 	if ( InLevel == nullptr && InWorld == GetWorld() )
 	{
 		RemoveFromParent();
