UMG - Simplifying the utility function for going from world space to widget space.

diff --git a/Engine/Source/Runtime/UMG/Private/Components/WidgetComponent.cpp b/Engine/Source/Runtime/UMG/Private/Components/WidgetComponent.cpp
--- a/Engine/Source/Runtime/UMG/Private/Components/WidgetComponent.cpp
+++ b/Engine/Source/Runtime/UMG/Private/Components/WidgetComponent.cpp
@@ -587,16 +587,16 @@ void UWidgetComponent::TickComponent(float DeltaTime, enum ELevelTick TickType,
 			{
 				FVector WorldLocation = GetComponentLocation();
 
-				FVector2D ScreenPosition, ScreenPositionInvDPI;
+				FVector2D ScreenPosition;
 				const bool bProjected = UWidgetLayoutLibrary::ProjectWorldLocationToWidgetPosition(
-					PlayerController, WorldLocation, ScreenPosition, ScreenPositionInvDPI);
+					PlayerController, WorldLocation, ScreenPosition);
 
 				if ( bProjected )
 				{
 					Widget->SetVisibility(ESlateVisibility::SelfHitTestInvisible);
 					
 					Widget->SetDesiredSizeInViewport(DrawSize);
-					Widget->SetPositionInViewport(ScreenPositionInvDPI, false);
+					Widget->SetPositionInViewport(ScreenPosition, false);
 					Widget->SetAlignmentInViewport(Pivot);
 				}
 				else
diff --git a/Engine/Source/Runtime/UMG/Private/WidgetLayoutLibrary.cpp b/Engine/Source/Runtime/UMG/Private/WidgetLayoutLibrary.cpp
--- a/Engine/Source/Runtime/UMG/Private/WidgetLayoutLibrary.cpp
+++ b/Engine/Source/Runtime/UMG/Private/WidgetLayoutLibrary.cpp
@@ -16,7 +16,7 @@ UWidgetLayoutLibrary::UWidgetLayoutLibrary(const FObjectInitializer& ObjectIniti
 {
 }
 
-bool UWidgetLayoutLibrary::ProjectWorldLocationToWidgetPosition(APlayerController* PlayerController, FVector WorldLocation, FVector2D& ScreenPosition, FVector2D& ScreenPositionInvDPI)
+bool UWidgetLayoutLibrary::ProjectWorldLocationToWidgetPosition(APlayerController* PlayerController, FVector WorldLocation, FVector2D& ScreenPosition)
 {
 	if ( PlayerController )
 	{
@@ -39,8 +39,7 @@ bool UWidgetLayoutLibrary::ProjectWorldLocationToWidgetPosition(APlayerControlle
 			const float Scale = UWidgetLayoutLibrary::GetViewportScale(PlayerController);
 
 			// Apply inverse DPI scale so that the widget ends up in the expected position
-			ScreenPositionInvDPI.X = ScreenPosition.X / Scale;
-			ScreenPositionInvDPI.Y = ScreenPosition.Y / Scale;
+			ScreenPosition = ScreenPosition / Scale;
 
 			return true;
 		}
diff --git a/Engine/Source/Runtime/UMG/Public/Blueprint/WidgetLayoutLibrary.h b/Engine/Source/Runtime/UMG/Public/Blueprint/WidgetLayoutLibrary.h
--- a/Engine/Source/Runtime/UMG/Public/Blueprint/WidgetLayoutLibrary.h
+++ b/Engine/Source/Runtime/UMG/Public/Blueprint/WidgetLayoutLibrary.h
@@ -16,12 +16,11 @@ public:
 	 * position, which takes into account any quality scaling.
 	 * @param PlayerController The player controller to project the position in the world to their screen.
 	 * @param WorldLocation The world location to project from.
-	 * @param ScreenPosition The position in the viewport without the inverse DPI being taken into account.
-	 * @param ScreenPositionInvDPI The psotion in the viewport with the inverse DPI applied to it.  This is useful for positioning child widgets of widgets in the viewport where you need to remove DPI before applying it.
+	 * @param ScreenPosition The position in the viewport with quality scale removed and DPI scale remove.
 	 * @return true if the position projects onto the screen.
 	 */
 	UFUNCTION(BlueprintPure, BlueprintCosmetic, Category="Viewport")
-	static bool ProjectWorldLocationToWidgetPosition(APlayerController* PlayerController, FVector WorldLocation, FVector2D& ScreenPosition, FVector2D& ScreenPositionInvDPI);
+	static bool ProjectWorldLocationToWidgetPosition(APlayerController* PlayerController, FVector WorldLocation, FVector2D& ScreenPosition);
 
 	/**
 	 * Gets the current DPI Scale being applied to the viewport and all the Widgets.
