UMG - Adding the ability for bindings to also work on TWeakObjectPtr.

diff --git a/Engine/Source/Editor/UMGEditor/Private/Details/SPropertyBinding.cpp b/Engine/Source/Editor/UMGEditor/Private/Details/SPropertyBinding.cpp
--- a/Engine/Source/Editor/UMGEditor/Private/Details/SPropertyBinding.cpp
+++ b/Engine/Source/Editor/UMGEditor/Private/Details/SPropertyBinding.cpp
@@ -356,17 +356,32 @@ void SPropertyBinding::FillPropertyMenu(FMenuBuilder& MenuBuilder, TSharedRef<IP
 					if ( Property->HasAllPropertyFlags(CPF_BlueprintVisible) )
 					{
 						UObjectProperty* ObjectProperty = Cast<UObjectProperty>(Property);
+						UWeakObjectProperty* WeakObjectProperty = Cast<UWeakObjectProperty>(Property);
 						UStructProperty* StructProperty = Cast<UStructProperty>(Property);
 
-						if ( ObjectProperty || StructProperty )
+						UStruct* Struct = nullptr;
+						UClass* Class = nullptr;
+
+						if ( ObjectProperty )
+						{
+							Struct = Class = ObjectProperty->PropertyClass;
+						}
+						else if ( WeakObjectProperty )
 						{
-							UStruct* Struct = ObjectProperty ? (UStruct*)ObjectProperty->PropertyClass : (UStruct*)StructProperty->Struct;
+							Struct = Class = WeakObjectProperty->PropertyClass;
+						}
+						else if ( StructProperty )
+						{
+							Struct = StructProperty->Struct;
+						}
 
-							// Ignore any properties that are widgets, we don't want users binding widgets to other widgets.
-							// also ignore any class that is explicitly on the black list.
-							if ( ObjectProperty )
+						if ( Struct )
+						{
+							if ( Class )
 							{
-								if ( IsClassBlackListed(ObjectProperty->PropertyClass) || ObjectProperty->PropertyClass->IsChildOf(UWidget::StaticClass()) )
+								// Ignore any properties that are widgets, we don't want users binding widgets to other widgets.
+								// also ignore any class that is explicitly on the black list.
+								if ( IsClassBlackListed(Class) || Class->IsChildOf(UWidget::StaticClass()) )
 								{
 									continue;
 								}
diff --git a/Engine/Source/Runtime/UMG/Public/Binding/DynamicPropertyPath.h b/Engine/Source/Runtime/UMG/Public/Binding/DynamicPropertyPath.h
--- a/Engine/Source/Runtime/UMG/Public/Binding/DynamicPropertyPath.h
+++ b/Engine/Source/Runtime/UMG/Public/Binding/DynamicPropertyPath.h
@@ -204,7 +204,10 @@ public:
 
 private:
 
-	/** */
+	/**
+	 * Evaluates the dynamic property path, and gets the value or calls the function at the end of the evaluation if possible.
+	 * Otherwise returns false.
+	 */
 	template<typename T>
 	bool GetValueRecursive(UStruct* InStruct, void* InContainer, int32 ArrayIndex, int32 SegmentIndex, T& OutValue, UProperty*& OutProperty) const
 	{
@@ -217,7 +220,7 @@ private:
 			{
 				if ( SegmentIndex < ( Segments.Num() - 1 ) )
 				{
-					// Check first to see if this is a simple object (i.e. not an array of objects)
+					// Check first to see if this is a simple object (eg. not an array of objects)
 					if ( UObjectProperty* ObjectProperty = Cast<UObjectProperty>(Property) )
 					{
 						// If it's an object we need to get the value of the property in the container first before we 
@@ -227,7 +230,19 @@ private:
 							return GetValueRecursive(CurrentObject->GetClass(), CurrentObject, ArrayIndex, SegmentIndex + 1, OutValue, OutProperty);
 						}
 					}
-					// Check to see if this is a simple structure (i.e. not an array of structures)
+					// Check to see if this is a simple weak object property (eg. not an array of weak objects).
+					if ( UWeakObjectProperty* WeakObjectProperty = Cast<UWeakObjectProperty>(Property) )
+					{
+						FWeakObjectPtr WeakObject = WeakObjectProperty->GetPropertyValue_InContainer(InContainer);
+
+						// If it's an object we need to get the value of the property in the container first before we 
+						// can continue, if the object is null we safely stop processing the chain of properties.
+						if ( UObject* CurrentObject = WeakObject.Get() )
+						{
+							return GetValueRecursive(CurrentObject->GetClass(), CurrentObject, ArrayIndex, SegmentIndex + 1, OutValue, OutProperty);
+						}
+					}
+					// Check to see if this is a simple structure (eg. not an array of structures)
 					else if ( UStructProperty* StructProp = Cast<UStructProperty>(Property) )
 					{
 						// Recursively call back into this function with the structure property and container value
