UMG - Cleaning up some headers.

diff --git a/Engine/Source/Runtime/UMG/Public/Components/ComboBoxString.h b/Engine/Source/Runtime/UMG/Public/Components/ComboBoxString.h
--- a/Engine/Source/Runtime/UMG/Public/Components/ComboBoxString.h
+++ b/Engine/Source/Runtime/UMG/Public/Components/ComboBoxString.h
@@ -9,7 +9,7 @@
 /**
  * The combobox allows you to display a list of options to the user in a dropdown menu for them to select one.
  */
-UCLASS( meta=( DisplayName="ComboBox (String)"), ClassGroup=UserInterface)
+UCLASS(meta=( DisplayName="ComboBox (String)"))
 class UMG_API UComboBoxString : public UWidget
 {
 	GENERATED_UCLASS_BODY()
diff --git a/Engine/Source/Runtime/UMG/Public/Components/HorizontalBox.h b/Engine/Source/Runtime/UMG/Public/Components/HorizontalBox.h
--- a/Engine/Source/Runtime/UMG/Public/Components/HorizontalBox.h
+++ b/Engine/Source/Runtime/UMG/Public/Components/HorizontalBox.h
@@ -12,7 +12,7 @@ class UHorizontalBoxSlot;
  * ● Many Children
  * ● Flow Horizontal
  */
-UCLASS(ClassGroup=UserInterface)
+UCLASS()
 class UMG_API UHorizontalBox : public UPanelWidget
 {
 	GENERATED_UCLASS_BODY()
diff --git a/Engine/Source/Runtime/UMG/Public/Components/ListView.h b/Engine/Source/Runtime/UMG/Public/Components/ListView.h
--- a/Engine/Source/Runtime/UMG/Public/Components/ListView.h
+++ b/Engine/Source/Runtime/UMG/Public/Components/ListView.h
@@ -7,7 +7,7 @@
 /**
  * Allows thousands of items to be displayed in a list.  Generates widgets dynamically for each item.
  */
-UCLASS(Experimental, ClassGroup=UserInterface)
+UCLASS(Experimental)
 class UMG_API UListView : public UTableViewBase
 {
 	GENERATED_UCLASS_BODY()
diff --git a/Engine/Source/Runtime/UMG/Public/Components/ProgressBar.h b/Engine/Source/Runtime/UMG/Public/Components/ProgressBar.h
--- a/Engine/Source/Runtime/UMG/Public/Components/ProgressBar.h
+++ b/Engine/Source/Runtime/UMG/Public/Components/ProgressBar.h
@@ -11,7 +11,7 @@
  *
  * ● No Children
  */
-UCLASS(ClassGroup=UserInterface)
+UCLASS()
 class UMG_API UProgressBar : public UWidget
 {
 	GENERATED_UCLASS_BODY()
diff --git a/Engine/Source/Runtime/UMG/Public/Components/Throbber.h b/Engine/Source/Runtime/UMG/Public/Components/Throbber.h
--- a/Engine/Source/Runtime/UMG/Public/Components/Throbber.h
+++ b/Engine/Source/Runtime/UMG/Public/Components/Throbber.h
@@ -10,7 +10,7 @@ class USlateBrushAsset;
 /**
  * A Throbber widget that shows several zooming circles in a row.
  */
-UCLASS(ClassGroup=UserInterface)
+UCLASS()
 class UMG_API UThrobber : public UWidget
 {
 	GENERATED_UCLASS_BODY()
diff --git a/Engine/Source/Runtime/UMG/Public/Components/TileView.h b/Engine/Source/Runtime/UMG/Public/Components/TileView.h
--- a/Engine/Source/Runtime/UMG/Public/Components/TileView.h
+++ b/Engine/Source/Runtime/UMG/Public/Components/TileView.h
@@ -7,7 +7,7 @@
 /**
  * A flow panel that presents the contents as a set of tiles all uniformly sized.
  */
-UCLASS(Experimental, ClassGroup=UserInterface)
+UCLASS(Experimental)
 class UMG_API UTileView : public UTableViewBase
 {
 	GENERATED_UCLASS_BODY()
diff --git a/Engine/Source/Runtime/UMG/Public/Components/Viewport.h b/Engine/Source/Runtime/UMG/Public/Components/Viewport.h
--- a/Engine/Source/Runtime/UMG/Public/Components/Viewport.h
+++ b/Engine/Source/Runtime/UMG/Public/Components/Viewport.h
@@ -216,7 +216,7 @@ protected:
 /**
  * 
  */
-UCLASS(Experimental, ClassGroup=UserInterface)
+UCLASS(Experimental)
 class UMG_API UViewport : public UContentWidget
 {
 	GENERATED_UCLASS_BODY()
