UMG - Adding an experimental wrapper for SRichTextBlock.  A lot of work remains for supporting links and images but this will serve as a nice starting point.

diff --git a/Engine/Source/Runtime/UMG/Private/Components/RichTextBlock.cpp b/Engine/Source/Runtime/UMG/Private/Components/RichTextBlock.cpp
new file mode 100644
--- /dev/null
+++ b/Engine/Source/Runtime/UMG/Private/Components/RichTextBlock.cpp
@@ -0,0 +1,88 @@
+// Copyright 1998-2015 Epic Games, Inc. All Rights Reserved.
+
+#include "UMGPrivatePCH.h"
+
+#include "RichTextBlockDecorator.h"
+#include "RichTextBlock.h"
+
+#define LOCTEXT_NAMESPACE "UMG"
+
+/////////////////////////////////////////////////////
+// URichTextBlock
+
+URichTextBlock::URichTextBlock(const FObjectInitializer& ObjectInitializer)
+	: Super(ObjectInitializer)
+{
+	static ConstructorHelpers::FObjectFinder<UFont> RobotoFontObj(TEXT("/Engine/EngineFonts/Roboto"));
+	Font = FSlateFontInfo(RobotoFontObj.Object, 12, FName("Regular"));
+	Color = FLinearColor::White;
+	LineHeightPercentage = 1;
+
+	Decorators.Add(ObjectInitializer.CreateOptionalDefaultSubobject<URichTextBlockDecorator>(this, FName("DefaultDecorator")));
+}
+
+void URichTextBlock::ReleaseSlateResources(bool bReleaseChildren)
+{
+	Super::ReleaseSlateResources(bReleaseChildren);
+
+	MyRichTextBlock.Reset();
+}
+
+TSharedRef<SWidget> URichTextBlock::RebuildWidget()
+{
+	//+ OnHyperlinkClicked = FSlateHyperlinkRun::FOnClick::CreateStatic(&RichTextHelper::OnBrowserLinkClicked, AsShared());
+	//+ FHyperlinkDecorator::Create(TEXT("browser"), OnHyperlinkClicked))
+	//+MakeShareable(new FDefaultRichTextDecorator(Font, Color));
+
+	DefaultStyle.SetFont(Font);
+	DefaultStyle.SetColorAndOpacity(Color);
+
+	TArray< TSharedRef< class ITextDecorator > > CreatedDecorators;
+
+	for ( URichTextBlockDecorator* Decorator : Decorators )
+	{
+		if ( Decorator )
+		{
+			CreatedDecorators.Add(Decorator->CreateDecorator(Font, Color));
+		}
+	}
+
+	MyRichTextBlock =
+		SNew(SRichTextBlock)
+		.Justification(Justification)
+		.AutoWrapText(AutoWrapText)
+		.WrapTextAt(WrapTextAt)
+		.Margin(Margin)
+		.LineHeightPercentage(LineHeightPercentage)
+		.TextStyle(&DefaultStyle)
+		.Decorators(CreatedDecorators);
+	
+	return MyRichTextBlock.ToSharedRef();
+}
+
+void URichTextBlock::SynchronizeProperties()
+{
+	Super::SynchronizeProperties();
+
+	TAttribute<FText> TextBinding = OPTIONAL_BINDING(FText, Text);
+
+	MyRichTextBlock->SetText(TextBinding);
+}
+
+#if WITH_EDITOR
+
+const FSlateBrush* URichTextBlock::GetEditorIcon()
+{
+	return FUMGStyle::Get().GetBrush("Widget.RichTextBlock");
+}
+
+const FText URichTextBlock::GetPaletteCategory()
+{
+	return LOCTEXT("Common", "Common");
+}
+
+#endif
+
+/////////////////////////////////////////////////////
+
+#undef LOCTEXT_NAMESPACE
diff --git a/Engine/Source/Runtime/UMG/Private/Components/RichTextBlockDecorator.cpp b/Engine/Source/Runtime/UMG/Private/Components/RichTextBlockDecorator.cpp
new file mode 100644
--- /dev/null
+++ b/Engine/Source/Runtime/UMG/Private/Components/RichTextBlockDecorator.cpp
@@ -0,0 +1,204 @@
+// Copyright 1998-2015 Epic Games, Inc. All Rights Reserved.
+
+#include "UMGPrivatePCH.h"
+
+#include "RichTextBlock.h"
+
+#define LOCTEXT_NAMESPACE "UMG"
+
+class FDefaultRichTextRun : public FSlateTextRun
+{
+public:
+	FDefaultRichTextRun(URichTextBlockDecorator* Decorator, const TSharedRef<FTextLayout>& TextLayout, const FRunInfo& InRunInfo, const TSharedRef< const FString >& InText, const FTextBlockStyle& InStyle, const FTextRange& InRange)
+		: FSlateTextRun(InRunInfo, InText, InStyle, InRange)
+		, Decorator(Decorator)
+		, TextLayout(TextLayout)
+	{
+	}
+
+	virtual int32 OnPaint(const FPaintArgs& Args,
+		const FTextLayout::FLineView& Line,
+		const TSharedRef< ILayoutBlock >& Block,
+		const FTextBlockStyle& DefaultStyle,
+		const FGeometry& AllottedGeometry,
+		const FSlateRect& ClippingRect,
+		FSlateWindowElementList& OutDrawElements,
+		int32 LayerId,
+		const FWidgetStyle& InWidgetStyle,
+		bool bParentEnabled) const override
+	{
+		URichTextBlockDecorator* RichBlock = Decorator.Get();
+		if ( RichBlock && RichBlock->bReveal )
+		{
+			const TArray< FTextLayout::FLineView >& Views = TextLayout->GetLineViews();
+
+			int32 AbsoluteBeginIndex = Line.Range.BeginIndex;
+			for ( int32 i = 0; i < Views.Num(); i++ )
+			{
+				if ( Views[i].ModelIndex == Line.ModelIndex )
+				{
+					break;
+				}
+
+				AbsoluteBeginIndex = Views[i].Range.Len();
+			}
+
+			if ( AbsoluteBeginIndex < RichBlock->RevealedIndex )
+			{
+				// compute the resulting color
+				FLinearColor TextColor = InWidgetStyle.GetColorAndOpacityTint() * Style.ColorAndOpacity.GetColor(InWidgetStyle);
+
+				FVector2D BlockOffset = Block->GetLocationOffset();
+
+				const FTextRange BlockRange = Block->GetTextRange();
+				const float InverseScale = Inverse(AllottedGeometry.Scale);
+				const ESlateDrawEffect::Type DrawEffects = bParentEnabled ? ESlateDrawEffect::None : ESlateDrawEffect::DisabledEffect;
+
+				int32 EndIndex = FMath::Clamp(RichBlock->RevealedIndex - AbsoluteBeginIndex, BlockRange.BeginIndex, BlockRange.EndIndex);
+
+				// Draw the text itself
+				FSlateDrawElement::MakeText(
+					OutDrawElements,
+					++LayerId,
+					AllottedGeometry.ToPaintGeometry(TransformVector(InverseScale, Block->GetSize()), FSlateLayoutTransform(TransformPoint(InverseScale, BlockOffset))),
+					Text.Get(),
+					BlockRange.BeginIndex,
+					EndIndex,
+					Style.Font,
+					ClippingRect,
+					DrawEffects,
+					TextColor
+					);
+			}
+		}
+		else
+		{
+			FSlateTextRun::OnPaint(Args, Line, Block, DefaultStyle, AllottedGeometry, ClippingRect, OutDrawElements, LayerId, InWidgetStyle, bParentEnabled);
+		}
+
+		return LayerId;
+	}
+
+private:
+	const TSharedRef<FTextLayout> TextLayout;
+	TWeakObjectPtr<URichTextBlockDecorator> Decorator;
+};
+
+/////////////////////////////////////////////////////
+// FDefaultRichTextDecorator
+
+FDefaultRichTextDecorator::FDefaultRichTextDecorator(URichTextBlockDecorator* Decorator, const FSlateFontInfo& DefaultFont, const FLinearColor& DefaultColor)
+	: DefaultFont(DefaultFont)
+	, DefaultColor(DefaultColor)
+	, Decorator(Decorator)
+{
+}
+
+FDefaultRichTextDecorator::~FDefaultRichTextDecorator()
+{
+}
+
+bool FDefaultRichTextDecorator::Supports(const FTextRunParseResults& RunParseResult, const FString& Text) const
+{
+	return ( RunParseResult.Name == TEXT("span") );
+}
+
+TSharedRef<ISlateRun> FDefaultRichTextDecorator::Create(const TSharedRef<FTextLayout>& TextLayout, const FTextRunParseResults& RunParseResult, const FString& OriginalText, const TSharedRef<FString>& InOutModelText, const ISlateStyle* Style)
+{
+	FRunInfo RunInfo(RunParseResult.Name);
+	for ( const TPair<FString, FTextRange>& Pair : RunParseResult.MetaData )
+	{
+		RunInfo.MetaData.Add(Pair.Key, OriginalText.Mid(Pair.Value.BeginIndex, Pair.Value.EndIndex - Pair.Value.BeginIndex));
+	}
+
+	FTextRange ModelRange;
+	ModelRange.BeginIndex = InOutModelText->Len();
+	*InOutModelText += OriginalText.Mid(RunParseResult.ContentRange.BeginIndex, RunParseResult.ContentRange.EndIndex - RunParseResult.ContentRange.BeginIndex);
+	ModelRange.EndIndex = InOutModelText->Len();
+
+	return CreateRun(TextLayout, RunInfo, InOutModelText, CreateTextBlockStyle(RunInfo), ModelRange);
+}
+
+TSharedRef<ISlateRun> FDefaultRichTextDecorator::CreateRun(const TSharedRef<FTextLayout>& TextLayout, const FRunInfo& InRunInfo, const TSharedRef< const FString >& InText, const FTextBlockStyle& Style, const FTextRange& InRange)
+{
+	return MakeShareable(new FDefaultRichTextRun(Decorator.Get(), TextLayout, InRunInfo, InText, Style, InRange));
+}
+
+FTextBlockStyle FDefaultRichTextDecorator::CreateTextBlockStyle(const FRunInfo& InRunInfo) const
+{
+	FSlateFontInfo Font;
+	FLinearColor FontColor;
+	ExplodeRunInfo(InRunInfo, Font, FontColor);
+
+	FTextBlockStyle TextBlockStyle;
+	TextBlockStyle.SetFont(Font);
+	TextBlockStyle.SetColorAndOpacity(FontColor);
+
+	return TextBlockStyle;
+}
+
+void FDefaultRichTextDecorator::ExplodeRunInfo(const FRunInfo& InRunInfo, FSlateFontInfo& OutFont, FLinearColor& OutFontColor) const
+{
+	OutFont = DefaultFont;
+
+	const FString* const FontFamilyString = InRunInfo.MetaData.Find(TEXT("font"));
+	const FString* const FontSizeString = InRunInfo.MetaData.Find(TEXT("size"));
+	const FString* const FontStyleString = InRunInfo.MetaData.Find(TEXT("style"));
+	const FString* const FontColorString = InRunInfo.MetaData.Find(TEXT("color"));
+
+	if ( FontFamilyString )
+	{
+		FStringAssetReference Font(**FontFamilyString);
+		if ( UObject* FontAsset = Font.TryLoad() )
+		{
+			OutFont.FontObject = FontAsset;
+		}
+	}
+
+	if ( FontSizeString )
+	{
+		OutFont.Size = static_cast<uint8>( FPlatformString::Atoi(**FontSizeString) );
+	}
+
+	if ( FontStyleString )
+	{
+		OutFont.TypefaceFontName = FName(**FontStyleString);
+	}
+
+	OutFontColor = DefaultColor;
+	if ( FontColorString )
+	{
+		const FString& FontColorStringRef = *FontColorString;
+
+		// Is Hex color?
+		if ( ( !FontColorStringRef.IsEmpty() && FontColorStringRef[0] == TCHAR('#') ) )
+		{
+			OutFontColor = FLinearColor(FColor::FromHex(FontColorStringRef));
+		}
+		else if ( OutFontColor.InitFromString(*FontColorString) )
+		{
+			// Initialized
+		}
+		else
+		{
+			OutFontColor = DefaultColor;
+		}
+	}
+}
+
+/////////////////////////////////////////////////////
+// URichTextBlockDecorator
+
+URichTextBlockDecorator::URichTextBlockDecorator(const FObjectInitializer& ObjectInitializer)
+	: Super(ObjectInitializer)
+{
+}
+
+TSharedRef<ITextDecorator> URichTextBlockDecorator::CreateDecorator(const FSlateFontInfo& DefaultFont, const FLinearColor& DefaultColor)
+{
+	return MakeShareable(new FDefaultRichTextDecorator(this, DefaultFont, DefaultColor));
+}
+
+/////////////////////////////////////////////////////
+
+#undef LOCTEXT_NAMESPACE
diff --git a/Engine/Source/Runtime/UMG/Public/Components/RichTextBlock.h b/Engine/Source/Runtime/UMG/Public/Components/RichTextBlock.h
new file mode 100644
--- /dev/null
+++ b/Engine/Source/Runtime/UMG/Public/Components/RichTextBlock.h
@@ -0,0 +1,84 @@
+﻿// Copyright 1998-2015 Epic Games, Inc. All Rights Reserved.
+
+#pragma once
+
+#include "RichTextBlock.generated.h"
+
+/**
+ * The rich text block
+ *
+ * ● Fancy Text
+ * ● No Children
+ */
+UCLASS(Experimental)
+class UMG_API URichTextBlock : public UWidget
+{
+	GENERATED_UCLASS_BODY()
+
+public:
+	
+	// UWidget interface
+	virtual void SynchronizeProperties() override;
+	// End of UWidget interface
+
+	// UVisual interface
+	virtual void ReleaseSlateResources(bool bReleaseChildren) override;
+	// End of UVisual interface
+
+#if WITH_EDITOR
+	// UWidget interface
+	virtual const FSlateBrush* GetEditorIcon() override;
+	virtual const FText GetPaletteCategory() override;
+	// End UWidget interface
+#endif
+
+protected:
+	/** The text to display */
+	UPROPERTY(EditDefaultsOnly, Category=Content, meta=( MultiLine="true" ))
+	FText Text;
+
+	/** A bindable delegate to allow logic to drive the text of the widget */
+	UPROPERTY()
+	FGetText TextDelegate;
+
+	/** The default font for the text. */
+	UPROPERTY(EditDefaultsOnly, BlueprintReadOnly, Category=Appearance)
+	FSlateFontInfo Font;
+
+	/** The default color for the text. */
+	UPROPERTY(EditDefaultsOnly, BlueprintReadOnly, Category=Appearance)
+	FLinearColor Color;
+
+	/** How the text should be aligned with the margin. */
+	UPROPERTY(EditDefaultsOnly, BlueprintReadOnly, Category=Appearance)
+	TEnumAsByte<ETextJustify::Type> Justification;
+
+	/** True if we're wrapping text automatically based on the computed horizontal space for this widget */
+	UPROPERTY(EditDefaultsOnly, BlueprintReadOnly, Category=Appearance)
+	bool AutoWrapText;
+
+	/** Whether text wraps onto a new line when it's length exceeds this width; if this value is zero or negative, no wrapping occurs. */
+	UPROPERTY(EditDefaultsOnly, BlueprintReadOnly, Category=Appearance, AdvancedDisplay)
+	float WrapTextAt;
+
+	/** The amount of blank space left around the edges of text area. */
+	UPROPERTY(EditDefaultsOnly, BlueprintReadOnly, Category=Appearance, AdvancedDisplay)
+	FMargin Margin;
+
+	/** The amount to scale each lines height by. */
+	UPROPERTY(EditDefaultsOnly, BlueprintReadOnly, Category=Appearance, AdvancedDisplay)
+	float LineHeightPercentage;
+
+	UPROPERTY(EditDefaultsOnly, Instanced, Category=Decorators)
+	TArray<class URichTextBlockDecorator*> Decorators;
+
+protected:
+	FTextBlockStyle DefaultStyle;
+
+	/** Native Slate Widget */
+	TSharedPtr<SRichTextBlock> MyRichTextBlock;
+
+	// UWidget interface
+	virtual TSharedRef<SWidget> RebuildWidget() override;
+	// End of UWidget interface
+};
diff --git a/Engine/Source/Runtime/UMG/Public/Components/RichTextBlockDecorator.h b/Engine/Source/Runtime/UMG/Public/Components/RichTextBlockDecorator.h
new file mode 100644
--- /dev/null
+++ b/Engine/Source/Runtime/UMG/Public/Components/RichTextBlockDecorator.h
@@ -0,0 +1,54 @@
+﻿// Copyright 1998-2015 Epic Games, Inc. All Rights Reserved.
+
+#pragma once
+
+#include "RichTextBlockDecorator.generated.h"
+
+class URichTextBlockDecorator;
+
+class FDefaultRichTextDecorator : public ITextDecorator
+{
+public:
+	FDefaultRichTextDecorator(URichTextBlockDecorator* Decorator, const FSlateFontInfo& DefaultFont, const FLinearColor& DefaultColor);
+
+	virtual ~FDefaultRichTextDecorator();
+
+	virtual bool Supports(const FTextRunParseResults& RunParseResult, const FString& Text) const override;
+
+	virtual TSharedRef<ISlateRun> Create(const TSharedRef<FTextLayout>& TextLayout, const FTextRunParseResults& RunParseResult, const FString& OriginalText, const TSharedRef<FString>& InOutModelText, const ISlateStyle* Style) override;
+
+protected:
+
+	virtual TSharedRef<ISlateRun> CreateRun(const TSharedRef<FTextLayout>& TextLayout, const FRunInfo& InRunInfo, const TSharedRef< const FString >& InText, const FTextBlockStyle& Style, const FTextRange& InRange);
+
+	FTextBlockStyle CreateTextBlockStyle(const FRunInfo& InRunInfo) const;
+
+	void ExplodeRunInfo(const FRunInfo& InRunInfo, FSlateFontInfo& OutFont, FLinearColor& OutFontColor) const;
+
+protected:
+
+	FSlateFontInfo DefaultFont;
+	FLinearColor DefaultColor;
+
+private:
+	TWeakObjectPtr<URichTextBlockDecorator> Decorator;
+};
+
+UCLASS(editinlinenew)
+class UMG_API URichTextBlockDecorator : public UObject
+{
+	GENERATED_BODY()
+
+public:
+	
+	UPROPERTY(EditAnywhere, BlueprintReadWrite, Category=Animation)
+	bool bReveal;
+
+	UPROPERTY(EditAnywhere, BlueprintReadWrite, Category=Animation)
+	int32 RevealedIndex;
+
+public:
+	URichTextBlockDecorator(const FObjectInitializer& ObjectInitializer);
+
+	virtual TSharedRef<ITextDecorator> CreateDecorator(const FSlateFontInfo& DefaultFont, const FLinearColor& DefaultColor);
+};
