UMG - Using the new GENERATE_BODY() macro for RichTextBlock.

diff --git a/Engine/Source/Runtime/UMG/Public/Components/RichTextBlock.h b/Engine/Source/Runtime/UMG/Public/Components/RichTextBlock.h
--- a/Engine/Source/Runtime/UMG/Public/Components/RichTextBlock.h
+++ b/Engine/Source/Runtime/UMG/Public/Components/RichTextBlock.h
@@ -13,9 +13,10 @@
 UCLASS(Experimental)
 class UMG_API URichTextBlock : public UWidget
 {
-	GENERATED_UCLASS_BODY()
+	GENERATED_BODY()
 
 public:
+	URichTextBlock(const FObjectInitializer& ObjectInitializer);
 	
 	// UWidget interface
 	virtual void SynchronizeProperties() override;
