UMG - UE-13066 - Properly ignore any widget that is not of the class provided.

diff --git a/Engine/Source/Runtime/UMG/Private/WidgetBlueprintLibrary.cpp b/Engine/Source/Runtime/UMG/Private/WidgetBlueprintLibrary.cpp
--- a/Engine/Source/Runtime/UMG/Private/WidgetBlueprintLibrary.cpp
+++ b/Engine/Source/Runtime/UMG/Private/WidgetBlueprintLibrary.cpp
@@ -370,10 +370,17 @@ void UWidgetBlueprintLibrary::GetAllWidgetsOfClass(UObject* WorldContextObject,
 	{
 		UUserWidget* LiveWidget = *Itr;
 
+		// Skip any widget that's not in the current world context.
 		if ( LiveWidget->GetWorld() != World )
 		{
 			continue;
 		}
+
+		// Skip any widget that is not a child of the class specified.
+		if ( !LiveWidget->IsChildOf(WidgetClass) )
+		{
+			continue;
+		}
 		
 		if ( TopLevelOnly )
 		{
