Build - Fixing the build.

diff --git a/Engine/Source/Runtime/UMG/Private/WidgetBlueprintLibrary.cpp b/Engine/Source/Runtime/UMG/Private/WidgetBlueprintLibrary.cpp
--- a/Engine/Source/Runtime/UMG/Private/WidgetBlueprintLibrary.cpp
+++ b/Engine/Source/Runtime/UMG/Private/WidgetBlueprintLibrary.cpp
@@ -377,7 +377,7 @@ void UWidgetBlueprintLibrary::GetAllWidgetsOfClass(UObject* WorldContextObject,
 		}
 
 		// Skip any widget that is not a child of the class specified.
-		if ( !LiveWidget->IsChildOf(WidgetClass) )
+		if ( !LiveWidget->GetClass()->IsChildOf(WidgetClass) )
 		{
 			continue;
 		}
