Fix PS4 compile errors due to no editor only data

diff --git a/Engine/Source/Runtime/UMG/Private/UserWidget.cpp b/Engine/Source/Runtime/UMG/Private/UserWidget.cpp
--- a/Engine/Source/Runtime/UMG/Private/UserWidget.cpp
+++ b/Engine/Source/Runtime/UMG/Private/UserWidget.cpp
@@ -860,6 +860,8 @@ void UUserWidget::PostLoad()
 {
 	Super::PostLoad();
 
+#if WITH_EDITORONLY_DATA
+
 	if ( GetLinkerUE4Version() < VER_UE4_DEPRECATE_USER_WIDGET_DESIGN_SIZE )
 	{
 		if ( bUseDesignTimeSize_DEPRECATED )
@@ -871,6 +873,8 @@ void UUserWidget::PostLoad()
 			DesignSizeMode = EDesignPreviewSizeMode::Desired;
 		}
 	}
+
+#endif
 }
 
 /////////////////////////////////////////////////////
