UMG - NamedSlots, ScaleBox, and SizeBox, were never intended to be Visible by default.  Their intended visibility is SelfHitTestInvisible.  If anyone wanted to have theirs set to visible, this change will affect them, their widgets will all suddenly no longer be hit testable.  However, this will fix several problems with invisible widgets preventing hit testing on widgets below them.

diff --git a/Engine/Source/Runtime/UMG/Private/Components/NamedSlot.cpp b/Engine/Source/Runtime/UMG/Private/Components/NamedSlot.cpp
--- a/Engine/Source/Runtime/UMG/Private/Components/NamedSlot.cpp
+++ b/Engine/Source/Runtime/UMG/Private/Components/NamedSlot.cpp
@@ -11,6 +11,7 @@ UNamedSlot::UNamedSlot(const FObjectInitializer& ObjectInitializer)
 	: Super(ObjectInitializer)
 {
 	bIsVariable = true;
+	Visibility = ESlateVisibility::SelfHitTestInvisible;
 }
 
 void UNamedSlot::ReleaseSlateResources(bool bReleaseChildren)
diff --git a/Engine/Source/Runtime/UMG/Private/Components/ScaleBox.cpp b/Engine/Source/Runtime/UMG/Private/Components/ScaleBox.cpp
--- a/Engine/Source/Runtime/UMG/Private/Components/ScaleBox.cpp
+++ b/Engine/Source/Runtime/UMG/Private/Components/ScaleBox.cpp
@@ -11,6 +11,7 @@ UScaleBox::UScaleBox(const FObjectInitializer& ObjectInitializer)
 	: Super(ObjectInitializer)
 {
 	bIsVariable = false;
+	Visibility = ESlateVisibility::SelfHitTestInvisible;
 
 	StretchDirection = EStretchDirection::Both;
 	Stretch = EStretch::ScaleToFit;
diff --git a/Engine/Source/Runtime/UMG/Private/Components/SizeBox.cpp b/Engine/Source/Runtime/UMG/Private/Components/SizeBox.cpp
--- a/Engine/Source/Runtime/UMG/Private/Components/SizeBox.cpp
+++ b/Engine/Source/Runtime/UMG/Private/Components/SizeBox.cpp
@@ -12,6 +12,7 @@ USizeBox::USizeBox(const FObjectInitializer& ObjectInitializer)
 	: Super(ObjectInitializer)
 {
 	bIsVariable = false;
+	Visibility = ESlateVisibility::SelfHitTestInvisible;
 }
 
 void USizeBox::ReleaseSlateResources(bool bReleaseChildren)
