PR #1007: added GetAllWidgetsWithInterface function (Contributed by [REDACTED_PERSON])

diff --git a/Engine/Source/Runtime/UMG/Private/WidgetBlueprintLibrary.cpp b/Engine/Source/Runtime/UMG/Private/WidgetBlueprintLibrary.cpp
--- a/Engine/Source/Runtime/UMG/Private/WidgetBlueprintLibrary.cpp
+++ b/Engine/Source/Runtime/UMG/Private/WidgetBlueprintLibrary.cpp
@@ -396,4 +396,41 @@ void UWidgetBlueprintLibrary::GetAllWidgetsOfClass(UObject* WorldContextObject,
 	}
 }
 
+void UWidgetBlueprintLibrary::GetAllWidgetsWithInterface(UObject* WorldContextObject, TSubclassOf<UInterface> Interface, TArray<UUserWidget*>& FoundWidgets, bool TopLevelOnly)
+{
+	FoundWidgets.Empty();
+
+	if (!Interface || !WorldContextObject)
+	{
+		return;
+	}
+
+	const UWorld* World = GEngine->GetWorldFromContextObject(WorldContextObject);
+	if (!World)
+	{
+		return;
+	}
+
+	for (TObjectIterator<UUserWidget> Itr; Itr; ++Itr)
+	{
+		UUserWidget* LiveWidget = *Itr;
+
+		if (LiveWidget->GetWorld() != World)
+		{
+			continue;
+		}
+
+		if (TopLevelOnly)
+		{
+			if (LiveWidget->IsInViewport() && LiveWidget->GetClass()->ImplementsInterface(Interface))
+			{
+				FoundWidgets.Add(LiveWidget);
+			}
+		}
+		else if (LiveWidget->GetClass()->ImplementsInterface(Interface))
+		{
+			FoundWidgets.Add(LiveWidget);
+		}
+	}
+}
 #undef LOCTEXT_NAMESPACE
\ No newline at end of file
diff --git a/Engine/Source/Runtime/UMG/Public/Blueprint/WidgetBlueprintLibrary.h b/Engine/Source/Runtime/UMG/Public/Blueprint/WidgetBlueprintLibrary.h
--- a/Engine/Source/Runtime/UMG/Public/Blueprint/WidgetBlueprintLibrary.h
+++ b/Engine/Source/Runtime/UMG/Public/Blueprint/WidgetBlueprintLibrary.h
@@ -193,4 +193,14 @@ public:
 	 */
 	UFUNCTION(BlueprintCallable, BlueprintCosmetic, Category="Widget", meta=( WorldContext="WorldContextObject" ))
 	static void GetAllWidgetsOfClass(UObject* WorldContextObject, TArray<UUserWidget*>& FoundWidgets, TSubclassOf<UUserWidget> WidgetClass, bool TopLevelOnly = true);
+
+	/**
+	* Find all widgets in the world with the specified interface.
+	* This is a slow operation, use with caution e.g. do not use every frame.
+	* @param Interface The interface to find. Must be specified or result array will be empty.
+	* @param FoundWidgets Output array of widgets that implement the specified interface.
+	* @param TopLevelOnly Only the widgets that are direct children of the viewport will be returned.
+	*/
+	UFUNCTION(BlueprintCallable, BlueprintCosmetic, Category = "Widget", meta = (WorldContext = "WorldContextObject"))
+	static void GetAllWidgetsWithInterface(UObject* WorldContextObject, TSubclassOf<UInterface> Interface, TArray<UUserWidget*>& FoundWidgets, bool TopLevelOnly);
 };
