UMG - Allowing access to the widget animations start and end time from blueprints.

diff --git a/Engine/Source/Runtime/UMG/Private/Animation/WidgetAnimation.cpp b/Engine/Source/Runtime/UMG/Private/Animation/WidgetAnimation.cpp
--- a/Engine/Source/Runtime/UMG/Private/Animation/WidgetAnimation.cpp
+++ b/Engine/Source/Runtime/UMG/Private/Animation/WidgetAnimation.cpp
@@ -43,4 +43,14 @@ UWidgetAnimation* UWidgetAnimation::GetNullAnimation()
 	return NullAnimation;
 }
 
-#endif
\ No newline at end of file
+#endif
+
+float UWidgetAnimation::GetStartTime() const
+{
+	return MovieScene->GetTimeRange().GetLowerBoundValue();
+}
+
+float UWidgetAnimation::GetEndTime() const
+{
+	return MovieScene->GetTimeRange().GetUpperBoundValue();
+}
diff --git a/Engine/Source/Runtime/UMG/Public/Animation/WidgetAnimation.h b/Engine/Source/Runtime/UMG/Public/Animation/WidgetAnimation.h
--- a/Engine/Source/Runtime/UMG/Public/Animation/WidgetAnimation.h
+++ b/Engine/Source/Runtime/UMG/Public/Animation/WidgetAnimation.h
@@ -56,6 +56,12 @@ public:
 	static UMG_API UWidgetAnimation* GetNullAnimation();
 #endif
 
+	UFUNCTION(BlueprintCallable, Category="Animation")
+	UMG_API float GetStartTime() const;
+
+	UFUNCTION(BlueprintCallable, Category="Animation")
+	UMG_API float GetEndTime() const;
+
 public:
 	UPROPERTY()
 	UMovieScene* MovieScene;
