UMG - Writing a better tooltip for RemoveFromParent, now notes the the fact that it removes it from the viewport and player screens.

diff --git a/Engine/Source/Runtime/UMG/Public/Components/Widget.h b/Engine/Source/Runtime/UMG/Public/Components/Widget.h
--- a/Engine/Source/Runtime/UMG/Public/Components/Widget.h
+++ b/Engine/Source/Runtime/UMG/Public/Components/Widget.h
@@ -314,7 +314,10 @@ public:
 	UFUNCTION(BlueprintCallable, Category="Widget")
 	class UPanelWidget* GetParent() const;
 
-	/** Removes the widget from it's parent widget */
+	/**
+	 * Removes the widget from its parent widget.  If this widget was added to the player's screen or the viewport
+	 * it will also be removed from those containers.
+	 */
 	UFUNCTION(BlueprintCallable, Category="Widget")
 	virtual void RemoveFromParent();
 
