UMG - Now using IsBindableEvent="True" for widgets that expose bindable single cast delegates.

diff --git a/Engine/Source/Runtime/UMG/Public/Components/Border.h b/Engine/Source/Runtime/UMG/Public/Components/Border.h
--- a/Engine/Source/Runtime/UMG/Public/Components/Border.h
+++ b/Engine/Source/Runtime/UMG/Public/Components/Border.h
@@ -62,16 +62,16 @@ public:
 
 public:
 
-	UPROPERTY(EditAnywhere, Category=Events)
+	UPROPERTY(EditAnywhere, Category=Events, meta=( IsBindableEvent="True" ))
 	FOnPointerEvent OnMouseButtonDownEvent;
 
-	UPROPERTY(EditAnywhere, Category=Events)
+	UPROPERTY(EditAnywhere, Category=Events, meta=( IsBindableEvent="True" ))
 	FOnPointerEvent OnMouseButtonUpEvent;
 
-	UPROPERTY(EditAnywhere, Category=Events)
+	UPROPERTY(EditAnywhere, Category=Events, meta=( IsBindableEvent="True" ))
 	FOnPointerEvent OnMouseMoveEvent;
 
-	UPROPERTY(EditAnywhere, Category=Events)
+	UPROPERTY(EditAnywhere, Category=Events, meta=( IsBindableEvent="True" ))
 	FOnPointerEvent OnMouseDoubleClickEvent;
 
 public:
diff --git a/Engine/Source/Runtime/UMG/Public/Components/ComboBox.h b/Engine/Source/Runtime/UMG/Public/Components/ComboBox.h
--- a/Engine/Source/Runtime/UMG/Public/Components/ComboBox.h
+++ b/Engine/Source/Runtime/UMG/Public/Components/ComboBox.h
@@ -19,7 +19,7 @@ public:
 	TArray<UObject*> Items;
 
 	/** Called when the widget is needed for the item. */
-	UPROPERTY(EditAnywhere, Category=Events)
+	UPROPERTY(EditAnywhere, Category=Events, meta=( IsBindableEvent="True" ))
 	FGenerateWidgetForObject OnGenerateWidgetEvent;
 
 	virtual void ReleaseSlateResources(bool bReleaseChildren) override;
diff --git a/Engine/Source/Runtime/UMG/Public/Components/ComboBoxString.h b/Engine/Source/Runtime/UMG/Public/Components/ComboBoxString.h
--- a/Engine/Source/Runtime/UMG/Public/Components/ComboBoxString.h
+++ b/Engine/Source/Runtime/UMG/Public/Components/ComboBoxString.h
@@ -50,7 +50,7 @@ public:
 public: // Events
 
 	/** Called when the widget is needed for the item. */
-	UPROPERTY(EditAnywhere, Category=Events)
+	UPROPERTY(EditAnywhere, Category=Events, meta=( IsBindableEvent="True" ))
 	FGenerateWidgetForString OnGenerateWidgetEvent;
 
 	/** Called when a new item is selected in the combobox. */
diff --git a/Engine/Source/Runtime/UMG/Public/Components/Image.h b/Engine/Source/Runtime/UMG/Public/Components/Image.h
--- a/Engine/Source/Runtime/UMG/Public/Components/Image.h
+++ b/Engine/Source/Runtime/UMG/Public/Components/Image.h
@@ -40,7 +40,7 @@ public:
 
 public:
 
-	UPROPERTY(EditAnywhere, Category=Events)
+	UPROPERTY(EditAnywhere, Category=Events, meta=( IsBindableEvent="True" ))
 	FOnPointerEvent OnMouseButtonDownEvent;
 
 public:
diff --git a/Engine/Source/Runtime/UMG/Public/Components/ListView.h b/Engine/Source/Runtime/UMG/Public/Components/ListView.h
--- a/Engine/Source/Runtime/UMG/Public/Components/ListView.h
+++ b/Engine/Source/Runtime/UMG/Public/Components/ListView.h
@@ -27,7 +27,7 @@ public:
 	TEnumAsByte<ESelectionMode::Type> SelectionMode;
 
 	/** Called when a widget needs to be generated */
-	UPROPERTY(EditAnywhere, Category=Events)
+	UPROPERTY(EditAnywhere, Category=Events, meta=( IsBindableEvent="True" ))
 	FOnGenerateRowUObject OnGenerateRowEvent;
 
 #if WITH_EDITOR
diff --git a/Engine/Source/Runtime/UMG/Public/Components/TileView.h b/Engine/Source/Runtime/UMG/Public/Components/TileView.h
--- a/Engine/Source/Runtime/UMG/Public/Components/TileView.h
+++ b/Engine/Source/Runtime/UMG/Public/Components/TileView.h
@@ -31,7 +31,7 @@ public:
 	TEnumAsByte<ESelectionMode::Type> SelectionMode;
 
 	/**  */
-	UPROPERTY(EditAnywhere, Category=Events)
+	UPROPERTY(EditAnywhere, Category=Events, meta=( IsBindableEvent="True" ))
 	FOnGenerateRowUObject OnGenerateTileEvent;
 
 public:
