UMG - Updating the comment on UserWidget's version of RemoveFromParent to be more accurate.

diff --git a/Engine/Source/Runtime/UMG/Public/Blueprint/UserWidget.h b/Engine/Source/Runtime/UMG/Public/Blueprint/UserWidget.h
--- a/Engine/Source/Runtime/UMG/Public/Blueprint/UserWidget.h
+++ b/Engine/Source/Runtime/UMG/Public/Blueprint/UserWidget.h
@@ -197,7 +197,10 @@ public:
 	UFUNCTION(BlueprintCallable, BlueprintCosmetic, Category="User Interface|Viewport", meta=( DeprecatedFunction, DeprecationMessage="Use RemoveFromParent instead" ))
 	void RemoveFromViewport();
 
-	/** Removes the widget from it's parent widget, including the viewport if it was added to the viewport. */
+	/**
+	 * Removes the widget from its parent widget.  If this widget was added to the player's screen or the viewport
+	 * it will also be removed from those containers.
+	 */
 	virtual void RemoveFromParent() override;
 
 	/**
