UMG - Exporting 3 blueprint function libraries UMG exposes to be linkable from C++ DLLs.

diff --git a/Engine/Source/Runtime/UMG/Public/Blueprint/SlateBlueprintLibrary.h b/Engine/Source/Runtime/UMG/Public/Blueprint/SlateBlueprintLibrary.h
--- a/Engine/Source/Runtime/UMG/Public/Blueprint/SlateBlueprintLibrary.h
+++ b/Engine/Source/Runtime/UMG/Public/Blueprint/SlateBlueprintLibrary.h
@@ -4,8 +4,8 @@
 
 #include "SlateBlueprintLibrary.generated.h"
 
-UCLASS(MinimalAPI)
-class USlateBlueprintLibrary : public UBlueprintFunctionLibrary
+UCLASS()
+class UMG_API USlateBlueprintLibrary : public UBlueprintFunctionLibrary
 {
 	GENERATED_UCLASS_BODY()
 
diff --git a/Engine/Source/Runtime/UMG/Public/Blueprint/WidgetBlueprintLibrary.h b/Engine/Source/Runtime/UMG/Public/Blueprint/WidgetBlueprintLibrary.h
--- a/Engine/Source/Runtime/UMG/Public/Blueprint/WidgetBlueprintLibrary.h
+++ b/Engine/Source/Runtime/UMG/Public/Blueprint/WidgetBlueprintLibrary.h
@@ -10,8 +10,8 @@ class USlateBrushAsset;
 
 struct FPaintContext;
 
-UCLASS(MinimalAPI)
-class UWidgetBlueprintLibrary : public UBlueprintFunctionLibrary
+UCLASS()
+class UMG_API UWidgetBlueprintLibrary : public UBlueprintFunctionLibrary
 {
 	GENERATED_UCLASS_BODY()
 
diff --git a/Engine/Source/Runtime/UMG/Public/Blueprint/WidgetLayoutLibrary.h b/Engine/Source/Runtime/UMG/Public/Blueprint/WidgetLayoutLibrary.h
--- a/Engine/Source/Runtime/UMG/Public/Blueprint/WidgetLayoutLibrary.h
+++ b/Engine/Source/Runtime/UMG/Public/Blueprint/WidgetLayoutLibrary.h
@@ -4,8 +4,8 @@
 
 #include "WidgetLayoutLibrary.generated.h"
 
-UCLASS(MinimalAPI)
-class UWidgetLayoutLibrary : public UBlueprintFunctionLibrary
+UCLASS()
+class UMG_API UWidgetLayoutLibrary : public UBlueprintFunctionLibrary
 {
 	GENERATED_UCLASS_BODY()
 
