UMG - Fixing SetBrush to show up correctly as an input, and not an output in blueprints.

diff --git a/Engine/Source/Runtime/UMG/Private/Components/Border.cpp b/Engine/Source/Runtime/UMG/Private/Components/Border.cpp
--- a/Engine/Source/Runtime/UMG/Private/Components/Border.cpp
+++ b/Engine/Source/Runtime/UMG/Private/Components/Border.cpp
@@ -181,7 +181,7 @@ FReply UBorder::HandleMouseDoubleClick(const FGeometry& Geometry, const FPointer
 	return FReply::Unhandled();
 }
 
-void UBorder::SetBrush(FSlateBrush& Brush)
+void UBorder::SetBrush(const FSlateBrush& Brush)
 {
 	Background = Brush;
 
diff --git a/Engine/Source/Runtime/UMG/Private/Components/Image.cpp b/Engine/Source/Runtime/UMG/Private/Components/Image.cpp
--- a/Engine/Source/Runtime/UMG/Private/Components/Image.cpp
+++ b/Engine/Source/Runtime/UMG/Private/Components/Image.cpp
@@ -76,7 +76,7 @@ const FSlateBrush* UImage::ConvertImage(TAttribute<FSlateBrush> InImageAsset) co
 	return &Brush;
 }
 
-void UImage::SetBrush(FSlateBrush& InBrush)
+void UImage::SetBrush(const FSlateBrush& InBrush)
 {
 	Brush = InBrush;
 
diff --git a/Engine/Source/Runtime/UMG/Public/Components/Border.h b/Engine/Source/Runtime/UMG/Public/Components/Border.h
--- a/Engine/Source/Runtime/UMG/Public/Components/Border.h
+++ b/Engine/Source/Runtime/UMG/Public/Components/Border.h
@@ -95,7 +95,7 @@ public:
 
 	/**  */
 	UFUNCTION(BlueprintCallable, Category="Appearance")
-	void SetBrush(FSlateBrush& InBrush);
+	void SetBrush(const FSlateBrush& InBrush);
 
 	/**  */
 	UFUNCTION(BlueprintCallable, Category="Appearance")
diff --git a/Engine/Source/Runtime/UMG/Public/Components/Image.h b/Engine/Source/Runtime/UMG/Public/Components/Image.h
--- a/Engine/Source/Runtime/UMG/Public/Components/Image.h
+++ b/Engine/Source/Runtime/UMG/Public/Components/Image.h
@@ -55,7 +55,7 @@ public:
 
 	/**  */
 	UFUNCTION(BlueprintCallable, Category="Appearance")
-	void SetBrush(FSlateBrush& InBrush);
+	void SetBrush(const FSlateBrush& InBrush);
 
 	/**  */
 	UFUNCTION(BlueprintCallable, Category="Appearance")
