UMG - Adding a way to make a custom user widget interactable for use with virtual cursors.  Also making it possible to pause an animation, which returns the current time of the animation to use when resuming it with Play later.

diff --git a/Engine/Source/Runtime/UMG/Private/Slate/SObjectWidget.cpp b/Engine/Source/Runtime/UMG/Private/Slate/SObjectWidget.cpp
--- a/Engine/Source/Runtime/UMG/Private/Slate/SObjectWidget.cpp
+++ b/Engine/Source/Runtime/UMG/Private/Slate/SObjectWidget.cpp
@@ -74,11 +74,21 @@ int32 SObjectWidget::OnPaint(const FPaintArgs& Args, const FGeometry& AllottedGe
 	return MaxLayer;
 }
 
+bool SObjectWidget::IsInteractable() const
+{
+	if ( CanRouteEvent() )
+	{
+		return WidgetObject->NativeIsInteractable();
+	}
+
+	return false;
+}
+
 bool SObjectWidget::SupportsKeyboardFocus() const
 {
 	if ( CanRouteEvent() )
 	{
-		return WidgetObject->bSupportsKeyboardFocus;
+		return WidgetObject->NativeSupportsKeyboardFocus();
 	}
 
 	return false;
diff --git a/Engine/Source/Runtime/UMG/Private/UserWidget.cpp b/Engine/Source/Runtime/UMG/Private/UserWidget.cpp
--- a/Engine/Source/Runtime/UMG/Private/UserWidget.cpp
+++ b/Engine/Source/Runtime/UMG/Private/UserWidget.cpp
@@ -244,6 +244,23 @@ void UUserWidget::StopAnimation(const UWidgetAnimation* InAnimation)
 	}
 }
 
+float UUserWidget::PauseAnimation(const UWidgetAnimation* InAnimation)
+{
+	if ( InAnimation )
+	{
+		// @todo UMG sequencer - Restart animations which have had Play called on them?
+		UUMGSequencePlayer** FoundPlayer = ActiveSequencePlayers.FindByPredicate([&] (const UUMGSequencePlayer* Player) { return Player->GetAnimation() == InAnimation; });
+
+		if ( FoundPlayer )
+		{
+			( *FoundPlayer )->Pause();
+			return (float)( *FoundPlayer )->GetTimeCursorPosition();
+		}
+	}
+
+	return 0;
+}
+
 void UUserWidget::OnAnimationFinishedPlaying( UUMGSequencePlayer& Player )
 {
 	OnAnimationFinished( Player.GetAnimation() );
@@ -730,6 +747,16 @@ void UUserWidget::NativePaint( FPaintContext& InContext ) const
 	OnPaint( InContext );
 }
 
+bool UUserWidget::NativeIsInteractable() const
+{
+	return IsInteractable();
+}
+
+bool UUserWidget::NativeSupportsKeyboardFocus() const
+{
+	return bSupportsKeyboardFocus;
+}
+
 FReply UUserWidget::NativeOnFocusReceived( const FGeometry& InGeometry, const FFocusEvent& InFocusEvent )
 {
 	return OnFocusReceived( InGeometry, InFocusEvent ).NativeReply;
diff --git a/Engine/Source/Runtime/UMG/Public/Animation/UMGSequencePlayer.h b/Engine/Source/Runtime/UMG/Public/Animation/UMGSequencePlayer.h
--- a/Engine/Source/Runtime/UMG/Public/Animation/UMGSequencePlayer.h
+++ b/Engine/Source/Runtime/UMG/Public/Animation/UMGSequencePlayer.h
@@ -28,6 +28,9 @@ public:
 	/** Pauses a running animation */
 	void Pause();
 
+	/** Gets the current time position in the player (in seconds). */
+	double GetTimeCursorPosition() const { return TimeCursorPosition; }
+
 	/** @return The current animation being played */
 	const UWidgetAnimation* GetAnimation() const { return Animation; }
 
diff --git a/Engine/Source/Runtime/UMG/Public/Blueprint/UserWidget.h b/Engine/Source/Runtime/UMG/Public/Blueprint/UserWidget.h
--- a/Engine/Source/Runtime/UMG/Public/Blueprint/UserWidget.h
+++ b/Engine/Source/Runtime/UMG/Public/Blueprint/UserWidget.h
@@ -302,6 +302,12 @@ public:
 	DEPRECATED(4.8, "Use NativePaint")
 	virtual void OnPaint_Implementation(FPaintContext& Context) const;
 
+	/**
+	 * Gets a value indicating if the widget is interactive.
+	 */
+	UFUNCTION(BlueprintImplementableEvent, BlueprintCosmetic, Category="User Interface | Interaction")
+	virtual bool IsInteractable() const;
+
 	/**
 	 * Called when keyboard focus is given to this widget.  This event does not bubble.
 	 *
@@ -718,6 +724,15 @@ public:
 	UFUNCTION(BlueprintCallable, BlueprintCosmetic, Category="User Interface|Animation")
 	void StopAnimation(const UWidgetAnimation* InAnimation);
 
+	/**
+	 * Pauses an already running animation in this widget
+	 * 
+	 * @param The name of the animation to pause
+	 * @return the time point the animation was at when it was paused.
+	 */
+	UFUNCTION(BlueprintCallable, BlueprintCosmetic, Category="User Interface|Animation")
+	float PauseAnimation(const UWidgetAnimation* InAnimation);
+
 	/** Called when a sequence player is finished playing an animation */
 	void OnAnimationFinishedPlaying(UUMGSequencePlayer& Player );
 
@@ -845,8 +860,13 @@ protected:
 
 	virtual void NativeConstruct();
 	virtual void NativeDestruct();
+
 	virtual void NativeTick(const FGeometry& MyGeometry, float InDeltaTime);
 	virtual void NativePaint( FPaintContext& InContext ) const;
+
+	virtual bool NativeIsInteractable() const;
+	virtual bool NativeSupportsKeyboardFocus() const;
+
 	virtual FReply NativeOnFocusReceived( const FGeometry& InGeometry, const FFocusEvent& InFocusEvent );
 	virtual void NativeOnFocusLost( const FFocusEvent& InFocusEvent );
 	virtual FReply NativeOnKeyChar( const FGeometry& InGeometry, const FCharacterEvent& InCharEvent );
diff --git a/Engine/Source/Runtime/UMG/Public/Slate/SObjectWidget.h b/Engine/Source/Runtime/UMG/Public/Slate/SObjectWidget.h
--- a/Engine/Source/Runtime/UMG/Public/Slate/SObjectWidget.h
+++ b/Engine/Source/Runtime/UMG/Public/Slate/SObjectWidget.h
@@ -33,6 +33,7 @@ class SObjectWidget : public SCompoundWidget, public FGCObject
 	virtual void Tick( const FGeometry& AllottedGeometry, const double InCurrentTime, const float InDeltaTime ) override;
 	virtual int32 OnPaint(const FPaintArgs& Args, const FGeometry& AllottedGeometry, const FSlateRect& MyClippingRect, FSlateWindowElementList& OutDrawElements, int32 LayerId, const FWidgetStyle& InWidgetStyle, bool bParentEnabled) const override;
 
+	virtual bool IsInteractable() const override;
 	virtual bool SupportsKeyboardFocus() const override;
 
 	virtual FReply OnFocusReceived(const FGeometry& MyGeometry, const FFocusEvent& InFocusEvent) override;
