IsInteractable doesn't need to be virtual

diff --git a/Engine/Source/Runtime/UMG/Public/Blueprint/UserWidget.h b/Engine/Source/Runtime/UMG/Public/Blueprint/UserWidget.h
--- a/Engine/Source/Runtime/UMG/Public/Blueprint/UserWidget.h
+++ b/Engine/Source/Runtime/UMG/Public/Blueprint/UserWidget.h
@@ -306,7 +306,7 @@ public:
 	 * Gets a value indicating if the widget is interactive.
 	 */
 	UFUNCTION(BlueprintImplementableEvent, BlueprintCosmetic, Category="User Interface | Interaction")
-	virtual bool IsInteractable() const;
+	bool IsInteractable() const;
 
 	/**
 	 * Called when keyboard focus is given to this widget.  This event does not bubble.
