UMG - Now initializing getter binding temp values on the stack to deal with blueprint functions that do not properly return a value.

diff --git a/Engine/Source/Runtime/UMG/Private/Binding/BoolBinding.cpp b/Engine/Source/Runtime/UMG/Private/Binding/BoolBinding.cpp
--- a/Engine/Source/Runtime/UMG/Private/Binding/BoolBinding.cpp
+++ b/Engine/Source/Runtime/UMG/Private/Binding/BoolBinding.cpp
@@ -25,7 +25,7 @@ bool UBoolBinding::GetValue() const
 
 	if ( UObject* Source = SourceObject.Get() )
 	{
-		bool Value;
+		bool Value = false;
 		if ( SourcePath.GetValue<bool>(Source, Value) )
 		{
 			return Value;
diff --git a/Engine/Source/Runtime/UMG/Private/Binding/CheckedStateBinding.cpp b/Engine/Source/Runtime/UMG/Private/Binding/CheckedStateBinding.cpp
--- a/Engine/Source/Runtime/UMG/Private/Binding/CheckedStateBinding.cpp
+++ b/Engine/Source/Runtime/UMG/Private/Binding/CheckedStateBinding.cpp
@@ -37,7 +37,7 @@ ECheckBoxState UCheckedStateBinding::GetValue() const
 	{
 		if ( bConversion.Get(EConversion::None) == EConversion::None )
 		{
-			uint8 Value;
+			uint8 Value = 0;
 			if ( SourcePath.GetValue<uint8>(Source, Value) )
 			{
 				bConversion = EConversion::None;
@@ -47,7 +47,7 @@ ECheckBoxState UCheckedStateBinding::GetValue() const
 
 		if ( bConversion.Get(EConversion::Bool) == EConversion::Bool )
 		{
-			bool Value;
+			bool Value = false;
 			if ( SourcePath.GetValue<bool>(Source, Value) )
 			{
 				bConversion = EConversion::Bool;
diff --git a/Engine/Source/Runtime/UMG/Private/Binding/ColorBinding.cpp b/Engine/Source/Runtime/UMG/Private/Binding/ColorBinding.cpp
--- a/Engine/Source/Runtime/UMG/Private/Binding/ColorBinding.cpp
+++ b/Engine/Source/Runtime/UMG/Private/Binding/ColorBinding.cpp
@@ -53,7 +53,7 @@ FSlateColor UColorBinding::GetSlateValue() const
 
 		if ( bNeedsConversion.Get(true) )
 		{
-			FLinearColor LinearValue;
+			FLinearColor LinearValue(ForceInitToZero);
 			if ( SourcePath.GetValue<FLinearColor>(Source, LinearValue) )
 			{
 				bNeedsConversion = true;
@@ -73,7 +73,7 @@ FLinearColor UColorBinding::GetLinearValue() const
 	{
 		if ( !bNeedsConversion.Get(false) )
 		{
-			FLinearColor LinearValue;
+			FLinearColor LinearValue(ForceInitToZero);
 			if ( SourcePath.GetValue<FLinearColor>(Source, LinearValue) )
 			{
 				bNeedsConversion = false;
diff --git a/Engine/Source/Runtime/UMG/Private/Binding/FloatBinding.cpp b/Engine/Source/Runtime/UMG/Private/Binding/FloatBinding.cpp
--- a/Engine/Source/Runtime/UMG/Private/Binding/FloatBinding.cpp
+++ b/Engine/Source/Runtime/UMG/Private/Binding/FloatBinding.cpp
@@ -25,7 +25,7 @@ float UFloatBinding::GetValue() const
 
 	if ( UObject* Source = SourceObject.Get() )
 	{
-		float Value;
+		float Value = 0;
 		if ( SourcePath.GetValue<float>(Source, Value) )
 		{
 			return Value;
diff --git a/Engine/Source/Runtime/UMG/Private/Binding/Int32Binding.cpp b/Engine/Source/Runtime/UMG/Private/Binding/Int32Binding.cpp
--- a/Engine/Source/Runtime/UMG/Private/Binding/Int32Binding.cpp
+++ b/Engine/Source/Runtime/UMG/Private/Binding/Int32Binding.cpp
@@ -25,7 +25,7 @@ int32 UInt32Binding::GetValue() const
 
 	if ( UObject* Source = SourceObject.Get() )
 	{
-		int32 Value;
+		int32 Value = 0;
 		if ( SourcePath.GetValue<int32>(Source, Value) )
 		{
 			return Value;
diff --git a/Engine/Source/Runtime/UMG/Private/Binding/MouseCursorBinding.cpp b/Engine/Source/Runtime/UMG/Private/Binding/MouseCursorBinding.cpp
--- a/Engine/Source/Runtime/UMG/Private/Binding/MouseCursorBinding.cpp
+++ b/Engine/Source/Runtime/UMG/Private/Binding/MouseCursorBinding.cpp
@@ -35,7 +35,7 @@ EMouseCursor::Type UMouseCursorBinding::GetValue() const
 
 	if ( UObject* Source = SourceObject.Get() )
 	{
-		uint8 Value;
+		uint8 Value = 0;
 		if ( SourcePath.GetValue<uint8>(Source, Value) )
 		{
 			return static_cast<EMouseCursor::Type>( Value );
diff --git a/Engine/Source/Runtime/UMG/Private/Binding/TextBinding.cpp b/Engine/Source/Runtime/UMG/Private/Binding/TextBinding.cpp
--- a/Engine/Source/Runtime/UMG/Private/Binding/TextBinding.cpp
+++ b/Engine/Source/Runtime/UMG/Private/Binding/TextBinding.cpp
@@ -43,7 +43,7 @@ FText UTextBinding::GetTextValue() const
 	{
 		if ( !bNeedsConversion.Get(false) )
 		{
-			FText TextValue;
+			FText TextValue = FText::GetEmpty();
 			if ( SourcePath.GetValue<FText>(Source, TextValue) )
 			{
 				bNeedsConversion = false;
@@ -83,7 +83,7 @@ FString UTextBinding::GetStringValue() const
 
 		if ( bNeedsConversion.Get(true) )
 		{
-			FText TextValue;
+			FText TextValue = FText::GetEmpty();
 			if ( SourcePath.GetValue<FText>(Source, TextValue) )
 			{
 				bNeedsConversion = true;
diff --git a/Engine/Source/Runtime/UMG/Private/Binding/VisibilityBinding.cpp b/Engine/Source/Runtime/UMG/Private/Binding/VisibilityBinding.cpp
--- a/Engine/Source/Runtime/UMG/Private/Binding/VisibilityBinding.cpp
+++ b/Engine/Source/Runtime/UMG/Private/Binding/VisibilityBinding.cpp
@@ -35,7 +35,7 @@ ESlateVisibility UVisibilityBinding::GetValue() const
 
 	if ( UObject* Source = SourceObject.Get() )
 	{
-		uint8 Value;
+		uint8 Value = 0;
 		if ( SourcePath.GetValue<uint8>(Source, Value) )
 		{
 			return static_cast<ESlateVisibility>( Value );
diff --git a/Engine/Source/Runtime/UMG/Private/Binding/WidgetBinding.cpp b/Engine/Source/Runtime/UMG/Private/Binding/WidgetBinding.cpp
--- a/Engine/Source/Runtime/UMG/Private/Binding/WidgetBinding.cpp
+++ b/Engine/Source/Runtime/UMG/Private/Binding/WidgetBinding.cpp
@@ -34,7 +34,7 @@ UWidget* UWidgetBinding::GetValue() const
 
 	if ( UObject* Source = SourceObject.Get() )
 	{
-		UObject* Value;
+		UObject* Value = nullptr;
 		if ( SourcePath.GetValue<UObject*>(Source, Value) )
 		{
 			if ( UWidget* Widget = Cast<UWidget>(Value) )
