UMG - Stubing in nav reply for SObjectWidget.

diff --git a/Engine/Source/Runtime/UMG/Private/Slate/SObjectWidget.cpp b/Engine/Source/Runtime/UMG/Private/Slate/SObjectWidget.cpp
--- a/Engine/Source/Runtime/UMG/Private/Slate/SObjectWidget.cpp
+++ b/Engine/Source/Runtime/UMG/Private/Slate/SObjectWidget.cpp
@@ -397,3 +397,12 @@ FReply SObjectWidget::OnMotionDetected(const FGeometry& MyGeometry, const FMotio
 
 	return FReply::Unhandled();
 }
+
+FNavigationReply SObjectWidget::OnNavigation(const FGeometry& MyGeometry, const FNavigationEvent& InNavigationEvent)
+{
+	FNavigationReply Reply = SCompoundWidget::OnNavigation(MyGeometry, InNavigationEvent);
+
+	// TODO Notify blueprint
+
+	return Reply;
+}
\ No newline at end of file
diff --git a/Engine/Source/Runtime/UMG/Public/Slate/SObjectWidget.h b/Engine/Source/Runtime/UMG/Public/Slate/SObjectWidget.h
--- a/Engine/Source/Runtime/UMG/Public/Slate/SObjectWidget.h
+++ b/Engine/Source/Runtime/UMG/Public/Slate/SObjectWidget.h
@@ -70,6 +70,8 @@ class SObjectWidget : public SCompoundWidget, public FGCObject
 	virtual FReply OnTouchEnded(const FGeometry& MyGeometry, const FPointerEvent& InTouchEvent) override;
 	virtual FReply OnMotionDetected(const FGeometry& MyGeometry, const FMotionEvent& InMotionEvent) override;
 
+	virtual FNavigationReply OnNavigation(const FGeometry& MyGeometry, const FNavigationEvent& InNavigationEvent) override;
+
 private:
 
 	/** The UWidget that created this SObjectWidget who needs to be kept alive. */
