UMG - Tweaking some viewport things, still not sure why it doesn't work in shipping.

diff --git a/Engine/Source/Runtime/UMG/Private/Components/Viewport.cpp b/Engine/Source/Runtime/UMG/Private/Components/Viewport.cpp
--- a/Engine/Source/Runtime/UMG/Private/Components/Viewport.cpp
+++ b/Engine/Source/Runtime/UMG/Private/Components/Viewport.cpp
@@ -94,12 +94,6 @@ FUMGViewportClient::FUMGViewportClient(FPreviewScene* InPreviewScene)
 
 FUMGViewportClient::~FUMGViewportClient()
 {
-
-}
-
-void FUMGViewportClient::AddReferencedObjects(FReferenceCollector & Collector)
-{
-
 }
 
 void FUMGViewportClient::Tick(float InDeltaTime)
@@ -183,6 +177,8 @@ void FUMGViewportClient::Draw(FViewport* InViewport, FCanvas* Canvas)
 
 	Canvas->Clear(BackgroundColor);
 
+	// workaround for hacky renderer code that uses GFrameNumber to decide whether to resize render targets
+	--GFrameNumber;
 	GetRendererModule().BeginRenderingViewFamily(Canvas, &ViewFamily);
 
 	// Remove temporary debug lines.
@@ -197,13 +193,6 @@ void FUMGViewportClient::Draw(FViewport* InViewport, FCanvas* Canvas)
 		World->ForegroundLineBatcher->Flush();
 	}
 	
-	//FCanvas* DebugCanvas = Viewport->GetDebugCanvas();
-
-	//UDebugDrawService::Draw(ViewFamily.EngineShowFlags, Viewport, View, DebugCanvas);
-
-	//
-	//FlushRenderingCommands();
-
 	Viewport = ViewportBackup;
 }
 
@@ -338,7 +327,7 @@ class SAutoRefreshViewport : public SViewport
 	virtual void Tick( const FGeometry& AllottedGeometry, const double InCurrentTime, const float InDeltaTime ) override
 	{
 		Viewport->Invalidate();
-		Viewport->InvalidateDisplay();
+		//Viewport->InvalidateDisplay();
 
 		Viewport->Tick(AllottedGeometry, InCurrentTime, InDeltaTime);
 		ViewportClient->Tick(InDeltaTime);
@@ -481,7 +470,10 @@ AActor* UViewport::Spawn(TSubclassOf<AActor> ActorClass)
 	if ( ViewportWidget.IsValid() )
 	{
 		UWorld* World = GetViewportWorld();
-		return World->SpawnActor<AActor>(ActorClass, FVector(0, 0, 0), FRotator());
+		FActorSpawnParameters SpawnParameters;
+		SpawnParameters.bNoCollisionFail = true;
+
+		return World->SpawnActor<AActor>(ActorClass, FVector(0, 0, 0), FRotator(), SpawnParameters);
 	}
 
 	// TODO UMG Report spawning actor error before the world is ready.
diff --git a/Engine/Source/Runtime/UMG/Public/Components/Viewport.h b/Engine/Source/Runtime/UMG/Public/Components/Viewport.h
--- a/Engine/Source/Runtime/UMG/Public/Components/Viewport.h
+++ b/Engine/Source/Runtime/UMG/Public/Components/Viewport.h
@@ -88,7 +88,7 @@ private:
 };
 
 
-class UMG_API FUMGViewportClient : public FCommonViewportClient, public FViewElementDrawer, public FGCObject
+class UMG_API FUMGViewportClient : public FCommonViewportClient, public FViewElementDrawer
 {
 public:
 	FUMGViewportClient(FPreviewScene* InPreviewScene = nullptr);
@@ -96,9 +96,6 @@ public:
 
 	using FViewElementDrawer::Draw;
 
-	// FGCObject interface
-	virtual void AddReferencedObjects(FReferenceCollector & Collector) override;
-
 	// FViewportClient interface
 	virtual UWorld* GetWorld() const override;
 	virtual void Draw(FViewport* InViewport, FCanvas* Canvas) override;
