UMG - SObjectWidget now allows SWidget to handle Key / Analog events if the user's native or blueprint code does not handle it.

diff --git a/Engine/Source/Runtime/UMG/Private/Slate/SObjectWidget.cpp b/Engine/Source/Runtime/UMG/Private/Slate/SObjectWidget.cpp
--- a/Engine/Source/Runtime/UMG/Private/Slate/SObjectWidget.cpp
+++ b/Engine/Source/Runtime/UMG/Private/Slate/SObjectWidget.cpp
@@ -141,7 +141,11 @@ FReply SObjectWidget::OnKeyDown(const FGeometry& MyGeometry, const FKeyEvent& In
 {
 	if ( CanRouteEvent() )
 	{
-		return WidgetObject->NativeOnKeyDown( MyGeometry, InKeyEvent );
+		FReply Result = WidgetObject->NativeOnKeyDown( MyGeometry, InKeyEvent );
+		if ( !Result.IsEventHandled() )
+		{
+			return SCompoundWidget::OnKeyDown(MyGeometry, InKeyEvent);
+		}
 	}
 
 	return FReply::Unhandled();
@@ -151,7 +155,11 @@ FReply SObjectWidget::OnKeyUp(const FGeometry& MyGeometry, const FKeyEvent& InKe
 {
 	if ( CanRouteEvent() )
 	{
-		return WidgetObject->NativeOnKeyUp( MyGeometry, InKeyEvent );
+		FReply Result = WidgetObject->NativeOnKeyUp(MyGeometry, InKeyEvent);
+		if ( !Result.IsEventHandled() )
+		{
+			return SCompoundWidget::OnKeyUp(MyGeometry, InKeyEvent);
+		}
 	}
 
 	return FReply::Unhandled();
@@ -161,7 +169,11 @@ FReply SObjectWidget::OnAnalogValueChanged(const FGeometry& MyGeometry, const FA
 {
 	if ( CanRouteEvent() )
 	{
-		return WidgetObject->NativeOnAnalogValueChanged( MyGeometry, InAnalogInputEvent );
+		FReply Result = WidgetObject->NativeOnAnalogValueChanged(MyGeometry, InAnalogInputEvent);
+		if ( !Result.IsEventHandled() )
+		{
+			return SCompoundWidget::OnAnalogValueChanged(MyGeometry, InAnalogInputEvent);
+		}
 	}
 
 	return FReply::Unhandled();
@@ -169,6 +181,7 @@ FReply SObjectWidget::OnAnalogValueChanged(const FGeometry& MyGeometry, const FA
 
 FReply SObjectWidget::OnMouseButtonDown(const FGeometry& MyGeometry, const FPointerEvent& MouseEvent)
 {
+	// NOTE: Done so that IsHovered() works
 	SCompoundWidget::OnMouseButtonDown(MyGeometry, MouseEvent);
 
 	if ( CanRouteEvent() )
@@ -181,6 +194,7 @@ FReply SObjectWidget::OnMouseButtonDown(const FGeometry& MyGeometry, const FPoin
 
 FReply SObjectWidget::OnPreviewMouseButtonDown(const FGeometry& MyGeometry, const FPointerEvent& MouseEvent)
 {
+	// NOTE: Done so that IsHovered() works
 	SCompoundWidget::OnPreviewMouseButtonDown(MyGeometry, MouseEvent);
 
 	if ( CanRouteEvent() )
@@ -193,6 +207,7 @@ FReply SObjectWidget::OnPreviewMouseButtonDown(const FGeometry& MyGeometry, cons
 
 FReply SObjectWidget::OnMouseButtonUp(const FGeometry& MyGeometry, const FPointerEvent& MouseEvent)
 {
+	// NOTE: Done so that IsHovered() works
 	SCompoundWidget::OnMouseButtonUp(MyGeometry, MouseEvent);
 
 	if ( CanRouteEvent() )
