UMG - Making the tooltip widget property readonly from blueprints.

diff --git a/Engine/Source/Runtime/UMG/Public/Components/Widget.h b/Engine/Source/Runtime/UMG/Public/Components/Widget.h
--- a/Engine/Source/Runtime/UMG/Public/Components/Widget.h
+++ b/Engine/Source/Runtime/UMG/Public/Components/Widget.h
@@ -135,7 +135,7 @@ public:
 	FGetText ToolTipTextDelegate;
 
 	/** Tooltip widget to show when the user hovers over the widget with the mouse */
-	UPROPERTY(EditAnywhere, Category="Behavior", AdvancedDisplay)
+	UPROPERTY(EditAnywhere, BlueprintReadOnly, Category="Behavior", AdvancedDisplay)
 	UWidget* ToolTipWidget;
 
 	/** A bindable delegate for ToolTipWidget */
