Fixed problem with detecting edited widget names in the Hierarchy
[UE-9480] - Default names for Custom Widget variables in the Hierarchy and My Blueprint have ugly unsanitized names

#change Fixed UWidget::IsGeneratedName() to match all possible naming schemes.

diff --git a/Engine/Source/Runtime/UMG/Private/Components/Widget.cpp b/Engine/Source/Runtime/UMG/Private/Components/Widget.cpp
--- a/Engine/Source/Runtime/UMG/Private/Components/Widget.cpp
+++ b/Engine/Source/Runtime/UMG/Private/Components/Widget.cpp
@@ -509,12 +509,21 @@ TSharedRef<SWidget> UWidget::BuildDesignTimeWidget(TSharedRef<SWidget> WrapWidge
 bool UWidget::IsGeneratedName() const
 {
 	FString Name = GetName();
-	FString BaseName = GetClass()->GetName() + TEXT("_");
 
-	if ( Name.StartsWith(BaseName) )
+	if (Name == GetClass()->GetName() || Name.StartsWith(GetClass()->GetName() + TEXT("_")))
 	{
 		return true;
 	}
+	else if (GetClass()->ClassGeneratedBy != nullptr)
+	{
+		FString BaseNameForBP = GetClass()->GetName();
+		BaseNameForBP.RemoveFromEnd(TEXT("_C"), ESearchCase::CaseSensitive);
+
+		if (Name == BaseNameForBP || Name.StartsWith(BaseNameForBP + TEXT("_")))
+		{
+			return true;
+		}
+	}
 
 	return false;
 }
