UMG - Tweaking the categoires of a few properties in UButton.

diff --git a/Engine/Source/Runtime/UMG/Public/Components/Button.h b/Engine/Source/Runtime/UMG/Public/Components/Button.h
--- a/Engine/Source/Runtime/UMG/Public/Components/Button.h
+++ b/Engine/Source/Runtime/UMG/Public/Components/Button.h
@@ -29,15 +29,15 @@ public:
 	USlateWidgetStyleAsset* Style_DEPRECATED;
 
 	/** The button style used at runtime */
-	UPROPERTY(EditAnywhere, BlueprintReadWrite, Category="Style", meta=( DisplayName="Style" ))
+	UPROPERTY(EditAnywhere, BlueprintReadWrite, Category="Appearance", meta=( DisplayName="Style" ))
 	FButtonStyle WidgetStyle;
 	
 	/** The color multiplier for the button content */
-	UPROPERTY(EditAnywhere, BlueprintReadOnly, Category=Appearance)
+	UPROPERTY(EditAnywhere, BlueprintReadOnly, Category="Appearance")
 	FLinearColor ColorAndOpacity;
 	
 	/** The color multiplier for the button background */
-	UPROPERTY(EditAnywhere, BlueprintReadOnly, Category=Appearance)
+	UPROPERTY(EditAnywhere, BlueprintReadOnly, Category="Appearance")
 	FLinearColor BackgroundColor;
 
 	/** The type of mouse action required by the user to trigger the buttons 'Click' */
@@ -49,7 +49,7 @@ public:
 	TEnumAsByte<EButtonTouchMethod::Type> TouchMethod;
 
 	/** Sometimes a button should only be mouse-clickable and never keyboard focusable. */
-	UPROPERTY(EditAnywhere, BlueprintReadOnly, Category="Interaction", AdvancedDisplay)
+	UPROPERTY(EditAnywhere, BlueprintReadOnly, Category="Interaction")
 	bool IsFocusable;
 
 public:
