UMG - Tweaking some categories on UserWidget.  Also making Focusability more visibile, and renaming it from "Supports Keyboard Focus" to "Is Focusable" just like UButton.

diff --git a/Engine/Source/Runtime/UMG/Public/Blueprint/UserWidget.h b/Engine/Source/Runtime/UMG/Public/Blueprint/UserWidget.h
--- a/Engine/Source/Runtime/UMG/Public/Blueprint/UserWidget.h
+++ b/Engine/Source/Runtime/UMG/Public/Blueprint/UserWidget.h
@@ -728,7 +728,7 @@ public:
 	 * Pauses an already running animation in this widget
 	 * 
 	 * @param The name of the animation to pause
-	 * @return the time point the animation was at when it was paused.
+	 * @return the time point the animation was at when it was paused.  Use this as the StartAtTime when you trigger PlayAnimation.
 	 */
 	UFUNCTION(BlueprintCallable, BlueprintCosmetic, Category="User Interface|Animation")
 	float PauseAnimation(const UWidgetAnimation* InAnimation);
@@ -741,7 +741,7 @@ public:
 	 *
 	 * @param The sound to play
 	 */
-	UFUNCTION(BlueprintCallable, BlueprintCosmetic, Category="Sound", meta=( DeprecatedFunction, DeprecationMessage="Use the global function PlaySound2D instead." ))
+	UFUNCTION(BlueprintCallable, BlueprintCosmetic, Category="Sound", meta=( DeprecatedFunction, DeprecationMessage="Use the UGameplayStatics::PlaySound2D instead." ))
 	void PlaySound(class USoundBase* SoundToPlay);
 
 	/** @returns The UObject wrapper for a given SWidget */
@@ -772,7 +772,7 @@ public:
 
 public:
 	/** The color and opacity of this widget.  Tints all child widgets. */
-	UPROPERTY(EditAnywhere, BlueprintReadOnly, Category="Style")
+	UPROPERTY(EditAnywhere, BlueprintReadOnly, Category="Appearance")
 	FLinearColor ColorAndOpacity;
 
 	UPROPERTY()
@@ -782,14 +782,14 @@ public:
 	 * The foreground color of the widget, this is inherited by sub widgets.  Any color property
 	 * that is marked as inherit will use this color.
 	 */
-	UPROPERTY(EditAnywhere, BlueprintReadOnly, Category="Style")
+	UPROPERTY(EditAnywhere, BlueprintReadOnly, Category="Appearance")
 	FSlateColor ForegroundColor;
 
 	UPROPERTY()
 	FGetSlateColor ForegroundColorDelegate;
 
 	/** Setting this flag to true, allows this widget to accept focus when clicked, or when navigated to. */
-	UPROPERTY(EditAnywhere, BlueprintReadWrite, AdvancedDisplay, Category="Behavior")
+	UPROPERTY(EditAnywhere, BlueprintReadWrite, Category="Interaction", meta=(DisplayName="Is Focusable"))
 	bool bSupportsKeyboardFocus;
 
 	/** The widget tree contained inside this user widget initialized by the blueprint */
