PR #1111: Fixed potential nullptr access in FWidget3DHitTester. (Contributed by [REDACTED_PERSON])

diff --git a/Engine/Source/Runtime/UMG/Private/Components/WidgetComponent.cpp b/Engine/Source/Runtime/UMG/Private/Components/WidgetComponent.cpp
--- a/Engine/Source/Runtime/UMG/Private/Components/WidgetComponent.cpp
+++ b/Engine/Source/Runtime/UMG/Private/Components/WidgetComponent.cpp
@@ -183,20 +183,22 @@ public:
 			if ( SafeWorld )
 			{
 				ULocalPlayer* const TargetPlayer = GEngine->GetLocalPlayerFromControllerId(SafeWorld, 0);
-				APlayerController* PlayerController = TargetPlayer->PlayerController;
 
-				if ( UPrimitiveComponent* HitComponent = GetHitResultAtScreenPositionAndCache(TargetPlayer->PlayerController, InGeometry.AbsoluteToLocal(DesktopSpaceCoordinate)) )
+				if( TargetPlayer && TargetPlayer->PlayerController )
 				{
-					if ( UWidgetComponent* WidgetComponent = Cast<UWidgetComponent>(HitComponent) )
+					if ( UPrimitiveComponent* HitComponent = GetHitResultAtScreenPositionAndCache(TargetPlayer->PlayerController, InGeometry.AbsoluteToLocal(DesktopSpaceCoordinate)) )
 					{
-						// Make sure the player is interacting with the front of the widget
-						// For widget components, the "front" faces the Z (or Up vector) direction
-						if ( FVector::DotProduct(WidgetComponent->GetUpVector(), CachedHitResult.ImpactPoint - CachedHitResult.TraceStart) < 0.f )
+						if ( UWidgetComponent* WidgetComponent = Cast<UWidgetComponent>(HitComponent) )
 						{
-							// Make sure the player is close enough to the widget to interact with it
-							if ( FVector::DistSquared(CachedHitResult.TraceStart, CachedHitResult.ImpactPoint) <= FMath::Square(WidgetComponent->GetMaxInteractionDistance()) )
+							// Make sure the player is interacting with the front of the widget
+							// For widget components, the "front" faces the Z (or Up vector) direction
+							if ( FVector::DotProduct(WidgetComponent->GetUpVector(), CachedHitResult.ImpactPoint - CachedHitResult.TraceStart) < 0.f )
 							{
-								return WidgetComponent->GetHitWidgetPath(CachedHitResult, bIgnoreEnabledStatus);
+								// Make sure the player is close enough to the widget to interact with it
+								if ( FVector::DistSquared(CachedHitResult.TraceStart, CachedHitResult.ImpactPoint) <= FMath::Square(WidgetComponent->GetMaxInteractionDistance()) )
+								{
+									return WidgetComponent->GetHitWidgetPath(CachedHitResult, bIgnoreEnabledStatus);
+								}
 							}
 						}
 					}
