UMG - UE-14967 - Fixing an edge case where if the user has set the ActiveWidgetIndex to an invalid index and then they select/deselect the widget it would trigger setting an invalid index on the slate widget leading to a crash.  Changing it so that it properly filters out invalid indexes for this case, and changing the post edit change code to rewrite the users value to the correct range after setting it.

diff --git a/Engine/Source/Runtime/UMG/Private/Components/WidgetSwitcher.cpp b/Engine/Source/Runtime/UMG/Private/Components/WidgetSwitcher.cpp
--- a/Engine/Source/Runtime/UMG/Private/Components/WidgetSwitcher.cpp
+++ b/Engine/Source/Runtime/UMG/Private/Components/WidgetSwitcher.cpp
@@ -109,6 +109,13 @@ TSharedRef<SWidget> UWidgetSwitcher::RebuildWidget()
 	return BuildDesignTimeWidget( MyWidgetSwitcher.ToSharedRef() );
 }
 
+void UWidgetSwitcher::PostEditChangeProperty(struct FPropertyChangedEvent& PropertyChangedEvent)
+{
+	ActiveWidgetIndex = FMath::Clamp(ActiveWidgetIndex, 0, FMath::Max(0, Slots.Num() - 1));
+
+	Super::PostEditChangeProperty(PropertyChangedEvent);
+}
+
 void UWidgetSwitcher::SynchronizeProperties()
 {
 	Super::SynchronizeProperties();
@@ -145,10 +152,7 @@ void UWidgetSwitcher::OnDescendantSelected(UWidget* DescendantWidget)
 
 void UWidgetSwitcher::OnDescendantDeselected(UWidget* DescendantWidget)
 {
-	if ( MyWidgetSwitcher.IsValid() )
-	{
-		MyWidgetSwitcher->SetActiveWidgetIndex(ActiveWidgetIndex);
-	}
+	SetActiveWidgetIndex(ActiveWidgetIndex);
 }
 
 #endif
diff --git a/Engine/Source/Runtime/UMG/Public/Components/WidgetSwitcher.h b/Engine/Source/Runtime/UMG/Public/Components/WidgetSwitcher.h
--- a/Engine/Source/Runtime/UMG/Public/Components/WidgetSwitcher.h
+++ b/Engine/Source/Runtime/UMG/Public/Components/WidgetSwitcher.h
@@ -34,10 +34,15 @@ public:
 	/** Activates the widget and makes it the active index. */
 	UFUNCTION(BlueprintCallable, Category="Switcher")
 	void SetActiveWidget(UWidget* Widget);
+
+	// UObject interface
+	virtual void PostEditChangeProperty(struct FPropertyChangedEvent& PropertyChangedEvent) override;
+	// End of UObject interface
 	
+	// UWidget interface
 	virtual void SynchronizeProperties() override;
-
 	virtual void ReleaseSlateResources(bool bReleaseChildren) override;
+	// End of UWidget interface
 
 #if WITH_EDITOR
 	virtual const FSlateBrush* GetEditorIcon() override;
