UMG - Fixing a bug with Render Transforms's IsIdentity calculation, the wrong identity comparison was being performed meaning essentially all widgets were adding a render transform to their stack, needlessly adding more calcs to all Slate layouts.  This also means you can use a scale of 0 now in your render transform when doing animations.

diff --git a/Engine/Source/Runtime/UMG/Public/Slate/WidgetTransform.h b/Engine/Source/Runtime/UMG/Public/Slate/WidgetTransform.h
--- a/Engine/Source/Runtime/UMG/Public/Slate/WidgetTransform.h
+++ b/Engine/Source/Runtime/UMG/Public/Slate/WidgetTransform.h
@@ -30,9 +30,9 @@ struct FWidgetTransform
 	
 
 	FWidgetTransform()
-		: Translation(0, 0)
-		, Scale(1, 1)
-		, Shear(0, 0)
+		: Translation(FVector2D::ZeroVector)
+		, Scale(FVector2D::UnitVector)
+		, Shear(FVector2D::ZeroVector)
 		, Angle(0)
 	{
 	}
@@ -47,11 +47,9 @@ struct FWidgetTransform
 
 	bool IsIdentity() const
 	{
-		return
-			Scale == FVector2D::ZeroVector
-			&& Shear == FVector2D::ZeroVector
-			&& Angle == 0.0f
-			&& Translation == FVector2D::ZeroVector;
+		const static FWidgetTransform Identity;
+
+		return Identity == *this;
 	}
 
 	bool operator==( const FWidgetTransform &Other ) const
