UMG - The designer now has seperate flags for Design Time, Show Outline.  So if you want to be in design mode but don't want to show panel outlines that's now possible.  There's a global flag for controlling it in the Widget Designer settings in Editor Preferences.

diff --git a/Engine/Source/Editor/UMGEditor/Private/Designer/SDesignerView.cpp b/Engine/Source/Editor/UMGEditor/Private/Designer/SDesignerView.cpp
--- a/Engine/Source/Editor/UMGEditor/Private/Designer/SDesignerView.cpp
+++ b/Engine/Source/Editor/UMGEditor/Private/Designer/SDesignerView.cpp
@@ -1710,7 +1710,7 @@ UWidget* SDesignerView::ProcessDropAndAddWidget(const FGeometry& MyGeometry, con
 
 			// Add it to the root if there are no other widgets to add it to.
 			UWidget* Widget = TemplateDragDropOp->Template->Create(BP->WidgetTree);
-			Widget->SetIsDesignTime(true);
+			Widget->SetDesignerFlags(BlueprintEditor.Pin()->GetCurrentDesignerFlags());
 
 			BP->WidgetTree->RootWidget = Widget;
 
@@ -1742,7 +1742,7 @@ UWidget* SDesignerView::ProcessDropAndAddWidget(const FGeometry& MyGeometry, con
 
 			// Construct the widget and mark it for design time rendering.
 			UWidget* Widget = TemplateDragDropOp->Template->Create(BP->WidgetTree);
-			Widget->SetIsDesignTime(true);
+			Widget->SetDesignerFlags(BlueprintEditor.Pin()->GetCurrentDesignerFlags());
 
 			// Determine local position inside the parent widget and add the widget to the slot.
 			FVector2D LocalPosition = WidgetUnderCursorGeometry.AbsoluteToLocal(DragDropEvent.GetScreenSpacePosition());
diff --git a/Engine/Source/Editor/UMGEditor/Private/Settings/WidgetDesignerSettings.cpp b/Engine/Source/Editor/UMGEditor/Private/Settings/WidgetDesignerSettings.cpp
--- a/Engine/Source/Editor/UMGEditor/Private/Settings/WidgetDesignerSettings.cpp
+++ b/Engine/Source/Editor/UMGEditor/Private/Settings/WidgetDesignerSettings.cpp
@@ -8,4 +8,5 @@ UWidgetDesignerSettings::UWidgetDesignerSettings(const FObjectInitializer& Objec
 {
 	GridSnapEnabled = true;
 	GridSnapSize = 4;
+	bShowOutlines = true;
 }
diff --git a/Engine/Source/Editor/UMGEditor/Private/WidgetBlueprintEditor.cpp b/Engine/Source/Editor/UMGEditor/Private/WidgetBlueprintEditor.cpp
--- a/Engine/Source/Editor/UMGEditor/Private/WidgetBlueprintEditor.cpp
+++ b/Engine/Source/Editor/UMGEditor/Private/WidgetBlueprintEditor.cpp
@@ -66,7 +66,7 @@ void FWidgetBlueprintEditor::InitWidgetBlueprintEditor(const EToolkitMode::Type
 	if ( Blueprint->WidgetTree->RootWidget == nullptr )
 	{
 		UWidget* RootWidget = Blueprint->WidgetTree->ConstructWidget<UCanvasPanel>(UCanvasPanel::StaticClass());
-		RootWidget->SetIsDesignTime(true);
+		RootWidget->SetDesignerFlags(GetCurrentDesignerFlags());
 		Blueprint->WidgetTree->RootWidget = RootWidget;
 	}
 
@@ -634,7 +634,7 @@ void FWidgetBlueprintEditor::UpdatePreview(UBlueprint* InBlueprint, bool bInForc
 		PreviewActor->SetFlags(RF_Transactional);
 		
 		// Configure all the widgets to be set to design time.
-		PreviewActor->SetIsDesignTime(true);
+		PreviewActor->SetDesignerFlags(GetCurrentDesignerFlags());
 
 		// Store a reference to the preview actor.
 		PreviewWidgetPtr = PreviewActor;
@@ -685,4 +685,16 @@ TArray< TFunction<void()> >& FWidgetBlueprintEditor::GetQueuedDesignerActions()
 	return QueuedDesignerActions;
 }
 
+EWidgetDesignFlags::Type FWidgetBlueprintEditor::GetCurrentDesignerFlags() const
+{
+	EWidgetDesignFlags::Type Flags = EWidgetDesignFlags::Designing;
+	
+	if ( GetDefault<UWidgetDesignerSettings>()->bShowOutlines )
+	{
+		Flags = ( EWidgetDesignFlags::Type )(Flags | EWidgetDesignFlags::ShowOutline);
+	}
+
+	return Flags;
+}
+
 #undef LOCTEXT_NAMESPACE
diff --git a/Engine/Source/Editor/UMGEditor/Private/WidgetBlueprintEditor.h b/Engine/Source/Editor/UMGEditor/Private/WidgetBlueprintEditor.h
--- a/Engine/Source/Editor/UMGEditor/Private/WidgetBlueprintEditor.h
+++ b/Engine/Source/Editor/UMGEditor/Private/WidgetBlueprintEditor.h
@@ -121,6 +121,9 @@ public:
 
 	TArray< TFunction<void()> >& GetQueuedDesignerActions();
 
+	/** Get the current designer flags that are in effect for the current user widget we're editing. */
+	EWidgetDesignFlags::Type GetCurrentDesignerFlags() const;
+
 public:
 	/** Fires whenever a new widget is being hovered over */
 	FOnHoveredWidgetSet OnHoveredWidgetSet;
diff --git a/Engine/Source/Editor/UMGEditor/Public/Settings/WidgetDesignerSettings.h b/Engine/Source/Editor/UMGEditor/Public/Settings/WidgetDesignerSettings.h
--- a/Engine/Source/Editor/UMGEditor/Public/Settings/WidgetDesignerSettings.h
+++ b/Engine/Source/Editor/UMGEditor/Public/Settings/WidgetDesignerSettings.h
@@ -21,4 +21,7 @@ class UMGEDITOR_API UWidgetDesignerSettings : public UObject
 
 	UPROPERTY(EditAnywhere, config, Category=Dragging)
 	bool bLockToPanelOnDragByDefault;
+
+	UPROPERTY(EditAnywhere, config, Category=Visuals)
+	bool bShowOutlines;
 };
diff --git a/Engine/Source/Runtime/UMG/Private/Components/PanelWidget.cpp b/Engine/Source/Runtime/UMG/Private/Components/PanelWidget.cpp
--- a/Engine/Source/Runtime/UMG/Private/Components/PanelWidget.cpp
+++ b/Engine/Source/Runtime/UMG/Private/Components/PanelWidget.cpp
@@ -176,6 +176,21 @@ void UPanelWidget::ShiftChild(int32 Index, UWidget* Child)
 	Slots.Insert(Child->Slot, FMath::Clamp(Index, 0, Slots.Num()));
 }
 
+void UPanelWidget::SetDesignerFlags(EWidgetDesignFlags::Type NewFlags)
+{
+	Super::SetDesignerFlags(NewFlags);
+
+	// Also mark all children as design time widgets.
+	int32 Children = GetChildrenCount();
+	for ( int32 SlotIndex = 0; SlotIndex < Children; SlotIndex++ )
+	{
+		if ( Slots[SlotIndex]->Content != nullptr )
+		{
+			Slots[SlotIndex]->Content->SetDesignerFlags(NewFlags);
+		}
+	}
+}
+
 #endif
 
 bool UPanelWidget::RemoveChild(UWidget* Content)
@@ -203,21 +218,6 @@ void UPanelWidget::ClearChildren()
 	}
 }
 
-void UPanelWidget::SetIsDesignTime(bool bInDesignTime)
-{
-	Super::SetIsDesignTime(bInDesignTime);
-
-	// Also mark all children as design time widgets.
-	int32 Children = GetChildrenCount();
-	for ( int32 SlotIndex = 0; SlotIndex < Children; SlotIndex++ )
-	{
-		if ( Slots[SlotIndex]->Content != nullptr )
-		{
-			Slots[SlotIndex]->Content->SetIsDesignTime(bInDesignTime);
-		}
-	}
-}
-
 void UPanelWidget::PostLoad()
 {
 	Super::PostLoad();
diff --git a/Engine/Source/Runtime/UMG/Private/Components/Widget.cpp b/Engine/Source/Runtime/UMG/Private/Components/Widget.cpp
--- a/Engine/Source/Runtime/UMG/Private/Components/Widget.cpp
+++ b/Engine/Source/Runtime/UMG/Private/Components/Widget.cpp
@@ -100,7 +100,7 @@ UWidget::UWidget(const FObjectInitializer& ObjectInitializer)
 {
 	bIsEnabled = true;
 	bIsVariable = true;
-	bDesignTime = false;
+	DesignerFlags = EWidgetDesignFlags::None;
 	Visiblity_DEPRECATED = Visibility = ESlateVisibility::Visible;	
 	RenderTransformPivot = FVector2D(0.5f, 0.5f);
 
@@ -493,7 +493,7 @@ TSharedRef<SWidget> UWidget::BuildDesignTimeWidget(TSharedRef<SWidget> WrapWidge
 		.VAlign(VAlign_Fill)
 		[
 			SNew(SBorder)
-			.Visibility( EVisibility::HitTestInvisible )
+			.Visibility(HasAnyDesignerFlags(EWidgetDesignFlags::ShowOutline) ? EVisibility::HitTestInvisible : EVisibility::Collapsed)
 			.BorderImage(FUMGStyle::Get().GetBrush("MarchingAnts"))
 		];
 	}
@@ -506,6 +506,11 @@ TSharedRef<SWidget> UWidget::BuildDesignTimeWidget(TSharedRef<SWidget> WrapWidge
 #if WITH_EDITOR
 #define LOCTEXT_NAMESPACE "UMGEditor"
 
+void UWidget::SetDesignerFlags(EWidgetDesignFlags::Type NewFlags)
+{
+	DesignerFlags = ( EWidgetDesignFlags::Type )( DesignerFlags | NewFlags );
+}
+
 bool UWidget::IsGeneratedName() const
 {
 	FString Name = GetName();
@@ -715,18 +720,6 @@ void UWidget::BuildNavigation()
 	}
 }
 
-#if WITH_EDITOR
-bool UWidget::IsDesignTime() const
-{
-	return bDesignTime;
-}
-#endif
-
-void UWidget::SetIsDesignTime(bool bInDesignTime)
-{
-	bDesignTime = bInDesignTime;
-}
-
 UWorld* UWidget::GetWorld() const
 {
 	// UWidget's are given world scope by their owning user widget.  We can get that through the widget tree that should
diff --git a/Engine/Source/Runtime/UMG/Private/Components/WidgetComponent.cpp b/Engine/Source/Runtime/UMG/Private/Components/WidgetComponent.cpp
--- a/Engine/Source/Runtime/UMG/Private/Components/WidgetComponent.cpp
+++ b/Engine/Source/Runtime/UMG/Private/Components/WidgetComponent.cpp
@@ -737,7 +737,7 @@ void UWidgetComponent::InitWidget()
 		if ( Widget && !GetWorld()->IsGameWorld() )
 		{
 			// Prevent native ticking of editor component previews
-			Widget->SetIsDesignTime(true);
+			Widget->SetDesignerFlags(EWidgetDesignFlags::Designing);
 		}
 	}
 }
diff --git a/Engine/Source/Runtime/UMG/Private/UserWidget.cpp b/Engine/Source/Runtime/UMG/Private/UserWidget.cpp
--- a/Engine/Source/Runtime/UMG/Private/UserWidget.cpp
+++ b/Engine/Source/Runtime/UMG/Private/UserWidget.cpp
@@ -190,15 +190,6 @@ UWorld* UUserWidget::GetWorld() const
 	return nullptr;
 }
 
-void UUserWidget::SetIsDesignTime(bool bInDesignTime)
-{
-	Super::SetIsDesignTime(bInDesignTime);
-
-	WidgetTree->ForEachWidget([&] (UWidget* Widget) {
-		Widget->SetIsDesignTime(bInDesignTime);
-	});
-}
-
 void UUserWidget::PlayAnimation( const UWidgetAnimation* InAnimation, float StartAtTime, int32 NumberOfLoops, EUMGSequencePlayMode::Type PlayMode)
 {
 	if( InAnimation )
@@ -679,6 +670,15 @@ const FText UUserWidget::GetPaletteCategory()
 	return PaletteCategory;
 }
 
+void UUserWidget::SetDesignerFlags(EWidgetDesignFlags::Type NewFlags)
+{
+	Super::SetDesignerFlags(NewFlags);
+
+	WidgetTree->ForEachWidget([&] (UWidget* Widget) {
+		Widget->SetDesignerFlags(NewFlags);
+	});
+}
+
 #endif
 
 void UUserWidget::OnAnimationStarted_Implementation(const UWidgetAnimation* Animation)
@@ -714,7 +714,7 @@ void UUserWidget::NativeTick(const FGeometry& MyGeometry, float InDeltaTime)
 
 void UUserWidget::TickActionsAndAnimation(const FGeometry& MyGeometry, float InDeltaTime)
 {
-	if ( bDesignTime )
+	if ( IsDesignTime() )
 	{
 		return;
 	}
diff --git a/Engine/Source/Runtime/UMG/Public/Blueprint/UserWidget.h b/Engine/Source/Runtime/UMG/Public/Blueprint/UserWidget.h
--- a/Engine/Source/Runtime/UMG/Public/Blueprint/UserWidget.h
+++ b/Engine/Source/Runtime/UMG/Public/Blueprint/UserWidget.h
@@ -170,9 +170,6 @@ public:
 	virtual void SetContentForSlot(FName SlotName, UWidget* Content) override;
 	// UNamedSlotInterface End
 
-	/** Sets that this widget is being designed sets it on all children as well. */
-	virtual void SetIsDesignTime(bool bInDesignTime) override;
-
 	/**
 	 * Adds it to the game's viewport and fills the entire screen, unless SetDesiredSizeInViewport is called
 	 * to explicitly set the size.
@@ -768,6 +765,8 @@ public:
 	virtual const FSlateBrush* GetEditorIcon() override;
 	virtual const FText GetPaletteCategory() override;
 	// End UWidget interface
+
+	void SetDesignerFlags(EWidgetDesignFlags::Type NewFlags);
 #endif
 
 public:
diff --git a/Engine/Source/Runtime/UMG/Public/Components/PanelWidget.h b/Engine/Source/Runtime/UMG/Public/Components/PanelWidget.h
--- a/Engine/Source/Runtime/UMG/Public/Components/PanelWidget.h
+++ b/Engine/Source/Runtime/UMG/Public/Components/PanelWidget.h
@@ -80,6 +80,8 @@ public:
 	 */
 	void ShiftChild(int32 Index, UWidget* Child);
 
+	void SetDesignerFlags(EWidgetDesignFlags::Type NewFlags);
+
 #endif
 
 	/**
@@ -112,9 +114,6 @@ public:
 		return CanHaveMultipleChildren() || GetChildrenCount() == 0;
 	}
 
-	/** Sets that this widget is being designed sets it on all children as well. */
-	virtual void SetIsDesignTime(bool bInDesignTime) override;
-
 	virtual void ReleaseSlateResources(bool bReleaseChildren) override;
 
 #if WITH_EDITOR
diff --git a/Engine/Source/Runtime/UMG/Public/Components/Widget.h b/Engine/Source/Runtime/UMG/Public/Components/Widget.h
--- a/Engine/Source/Runtime/UMG/Public/Components/Widget.h
+++ b/Engine/Source/Runtime/UMG/Public/Components/Widget.h
@@ -68,6 +68,21 @@ class UPanelSlot;
 
 
 
+/**
+ * Flags used by the widget designer.
+ */
+UENUM()
+namespace EWidgetDesignFlags
+{
+	enum Type
+	{
+		None		= 0,
+		Designing	= 1,
+		ShowOutline	= 2,
+	};
+}
+
+
 /**
  * This is the base class for all wrapped Slate controls that are exposed to UObjects.
  */
@@ -347,14 +362,23 @@ public:
 
 #if WITH_EDITOR
 	/** Returns if the widget is currently being displayed in the designer, it may want to display different data. */
-	bool IsDesignTime() const;
+	FORCEINLINE bool IsDesignTime() const
+	{
+		return HasAnyDesignerFlags(EWidgetDesignFlags::Designing);
+	}
+
+	virtual void SetDesignerFlags(EWidgetDesignFlags::Type NewFlags);
+
+	/** Tests if any of the flags exist on this widget. */
+	FORCEINLINE bool HasAnyDesignerFlags(EWidgetDesignFlags::Type FlagToCheck) const
+	{
+		return ( DesignerFlags&FlagToCheck ) != 0;
+	}
+
 #else
 	FORCEINLINE bool IsDesignTime() const { return false; }
 #endif
 	
-	/** Sets that this widget is being designed */
-	virtual void SetIsDesignTime(bool bInDesignTime);
-
 	/** Mark this object as modified, also mark the slot as modified. */
 	virtual bool Modify(bool bAlwaysMarkDirty = true) override;
 
@@ -454,6 +478,7 @@ protected:
 	/** Function called after the underlying SWidget is constructed. */
 	virtual void OnWidgetRebuilt();
 	
+	/** Utility method for building a design time wrapper widget. */
 	TSharedRef<SWidget> BuildDesignTimeWidget(TSharedRef<SWidget> WrapWidget);
 
 	void UpdateRenderTransform();
@@ -486,10 +511,10 @@ protected:
 
 	/** The underlying SWidget contained in a SObjectWidget */
 	TWeakPtr<class SObjectWidget> MyGCWidget;
-	
-	/** Is this widget being displayed on a designer surface */
+
+	/** Any flags used by the designer at edit time. */
 	UPROPERTY(Transient)
-	bool bDesignTime;
+	TEnumAsByte<EWidgetDesignFlags::Type> DesignerFlags;
 
 	/** Native property bindings. */
 	UPROPERTY(Transient)
