UMG - Code cleanup / simplification.

diff --git a/Engine/Source/Runtime/UMG/Private/Components/Widget.cpp b/Engine/Source/Runtime/UMG/Private/Components/Widget.cpp
--- a/Engine/Source/Runtime/UMG/Private/Components/Widget.cpp
+++ b/Engine/Source/Runtime/UMG/Private/Components/Widget.cpp
@@ -143,20 +143,20 @@ void UWidget::UpdateRenderTransform()
 	TSharedPtr<SWidget> SafeWidget = GetCachedWidget();
 	if (SafeWidget.IsValid())
 	{
-		if (!RenderTransform.IsIdentity())
+		if (RenderTransform.IsIdentity())
 		{
-			FSlateRenderTransform Transform2D = 
-				::Concatenate(
-					FScale2D(RenderTransform.Scale),
-					FShear2D::FromShearAngles(RenderTransform.Shear),
-					FQuat2D(FMath::DegreesToRadians(RenderTransform.Angle)),
-					FVector2D(RenderTransform.Translation));
-
-			SafeWidget->SetRenderTransform(Transform2D);
+			SafeWidget->SetRenderTransform(TOptional<FSlateRenderTransform>());
 		}
 		else
 		{
-			SafeWidget->SetRenderTransform(TOptional<FSlateRenderTransform>());
+			FSlateRenderTransform Transform2D =
+				::Concatenate(
+				FScale2D(RenderTransform.Scale),
+				FShear2D::FromShearAngles(RenderTransform.Shear),
+				FQuat2D(FMath::DegreesToRadians(RenderTransform.Angle)),
+				FVector2D(RenderTransform.Translation));
+
+			SafeWidget->SetRenderTransform(Transform2D);
 		}
 	}
 }
