UMG - Fixing the build.

diff --git a/Engine/Source/Runtime/UMG/Private/Components/Widget.cpp b/Engine/Source/Runtime/UMG/Private/Components/Widget.cpp
--- a/Engine/Source/Runtime/UMG/Private/Components/Widget.cpp
+++ b/Engine/Source/Runtime/UMG/Private/Components/Widget.cpp
@@ -483,6 +483,7 @@ TSharedPtr<SWidget> UWidget::GetCachedWidget() const
 
 TSharedRef<SWidget> UWidget::BuildDesignTimeWidget(TSharedRef<SWidget> WrapWidget)
 {
+#if WITH_EDITOR
 	if (IsDesignTime())
 	{
 		return SNew(SOverlay)
@@ -507,6 +508,9 @@ TSharedRef<SWidget> UWidget::BuildDesignTimeWidget(TSharedRef<SWidget> WrapWidge
 	{
 		return WrapWidget;
 	}
+#else
+	return WrapWidget;
+#endif
 }
 
 #if WITH_EDITOR
diff --git a/Engine/Source/Runtime/UMG/Private/Components/WidgetComponent.cpp b/Engine/Source/Runtime/UMG/Private/Components/WidgetComponent.cpp
--- a/Engine/Source/Runtime/UMG/Private/Components/WidgetComponent.cpp
+++ b/Engine/Source/Runtime/UMG/Private/Components/WidgetComponent.cpp
@@ -734,11 +734,13 @@ void UWidgetComponent::InitWidget()
 			Widget = nullptr;
 		}
 
+#if WITH_EDITOR
 		if ( Widget && !GetWorld()->IsGameWorld() )
 		{
 			// Prevent native ticking of editor component previews
 			Widget->SetDesignerFlags(EWidgetDesignFlags::Designing);
 		}
+#endif
 	}
 }
 
