UMG - Scoping the designer flags for widgets and requiring an accessor for them.

diff --git a/Engine/Source/Runtime/UMG/Public/Components/Widget.h b/Engine/Source/Runtime/UMG/Public/Components/Widget.h
--- a/Engine/Source/Runtime/UMG/Public/Components/Widget.h
+++ b/Engine/Source/Runtime/UMG/Public/Components/Widget.h
@@ -367,8 +367,15 @@ public:
 		return HasAnyDesignerFlags(EWidgetDesignFlags::Designing);
 	}
 
+	/** Sets the designer flags on the widget. */
 	virtual void SetDesignerFlags(EWidgetDesignFlags::Type NewFlags);
 
+	/** Gets the designer flags currently set on the widget. */
+	FORCEINLINE EWidgetDesignFlags::Type GetDesignerFlags() const
+	{
+		return DesignerFlags;
+	}
+
 	/** Tests if any of the flags exist on this widget. */
 	FORCEINLINE bool HasAnyDesignerFlags(EWidgetDesignFlags::Type FlagToCheck) const
 	{
@@ -512,16 +519,17 @@ protected:
 	/** The underlying SWidget contained in a SObjectWidget */
 	TWeakPtr<class SObjectWidget> MyGCWidget;
 
-	/** Any flags used by the designer at edit time. */
-	UPROPERTY(Transient)
-	TEnumAsByte<EWidgetDesignFlags::Type> DesignerFlags;
-
 	/** Native property bindings. */
 	UPROPERTY(Transient)
 	TArray<class UPropertyBinding*> NativeBindings;
 
 	static TArray<TSubclassOf<UPropertyBinding>> BinderClasses;
 
+private:
+	/** Any flags used by the designer at edit time. */
+	UPROPERTY(Transient)
+	TEnumAsByte<EWidgetDesignFlags::Type> DesignerFlags;
+
 private:
 	GAME_SAFE_BINDING_IMPLEMENTATION(FText, ToolTipText)
 	GAME_SAFE_BINDING_IMPLEMENTATION(bool, bIsEnabled)
