UMG - Adding utility methods to go from a widget's geometry to viewport space.

diff --git a/Engine/Source/Runtime/UMG/Private/SlateBlueprintLibrary.cpp b/Engine/Source/Runtime/UMG/Private/SlateBlueprintLibrary.cpp
--- a/Engine/Source/Runtime/UMG/Private/SlateBlueprintLibrary.cpp
+++ b/Engine/Source/Runtime/UMG/Private/SlateBlueprintLibrary.cpp
@@ -3,6 +3,7 @@
 #include "UMGPrivatePCH.h"
 
 #include "Slate/SlateBrushAsset.h"
+#include "Runtime/Engine/Classes/Engine/UserInterfaceSettings.h"
 #include "SlateBlueprintLibrary.h"
 
 #define LOCTEXT_NAMESPACE "UMG"
@@ -35,4 +36,38 @@ FVector2D USlateBlueprintLibrary::GetLocalSize(const FGeometry& Geometry)
 	return Geometry.GetLocalSize();
 }
 
+void USlateBlueprintLibrary::LocalToViewport(UObject* WorldContextObject, const FGeometry& Geometry, FVector2D LocalCoordinate, FVector2D& ScreenPosition, FVector2D& ViewportPosition)
+{
+	FVector2D AbsoluteCoordinate = Geometry.LocalToAbsolute(LocalCoordinate);
+	return AbsoluteToViewport(WorldContextObject, Geometry, AbsoluteCoordinate, ScreenPosition, ViewportPosition);
+}
+
+void USlateBlueprintLibrary::AbsoluteToViewport(UObject* WorldContextObject, const FGeometry& Geometry, FVector2D AbsoluteCoordinate, FVector2D& ScreenPosition, FVector2D& ViewportPosition)
+{
+	UWorld* World = GEngine->GetWorldFromContextObject(WorldContextObject);
+	if ( World && World->IsGameWorld() )
+	{
+		if ( UGameViewportClient* ViewportClient = World->GetGameViewport() )
+		{
+			if ( FViewport* Viewport = ViewportClient->Viewport )
+			{
+				FVector2D ViewportSize;
+				ViewportClient->GetViewportSize(ViewportSize);
+
+				FVector2D AbsolutePositionPostDPI = Viewport->VirtualDesktopPixelToViewport(FIntPoint((int32)AbsoluteCoordinate.X, (int32)AbsoluteCoordinate.Y)) * ViewportSize;
+
+				float CurrentViewportScale = GetDefault<UUserInterfaceSettings>(UUserInterfaceSettings::StaticClass())->GetDPIScaleBasedOnSize(FIntPoint(ViewportSize.X, ViewportSize.Y));
+
+				ScreenPosition = AbsolutePositionPostDPI;
+				ViewportPosition = AbsolutePositionPostDPI / CurrentViewportScale;
+
+				return;
+			}
+		}
+	}
+
+	ScreenPosition = FVector2D(0, 0);
+	ViewportPosition = FVector2D(0, 0);
+}
+
 #undef LOCTEXT_NAMESPACE
\ No newline at end of file
diff --git a/Engine/Source/Runtime/UMG/Public/Blueprint/SlateBlueprintLibrary.h b/Engine/Source/Runtime/UMG/Public/Blueprint/SlateBlueprintLibrary.h
--- a/Engine/Source/Runtime/UMG/Public/Blueprint/SlateBlueprintLibrary.h
+++ b/Engine/Source/Runtime/UMG/Public/Blueprint/SlateBlueprintLibrary.h
@@ -40,4 +40,28 @@ public:
 	/** @return the size of the geometry in local space. */
 	UFUNCTION(BlueprintPure, Category="User Interface|Geometry")
 	static FVector2D GetLocalSize(const FGeometry& Geometry);
+
+	/**
+	 * Translates local coordinate of the geometry provided into local viewport coordinates.
+	 *
+	 * @param ScreenPosition The position in the game's viewport, usable for line traces and 
+	 * other uses where you need a coordinate in the space of viewport resolution units.
+	 * @param ViewportPosition The position in the space of other widgets in the viewport.  Like if you wanted
+	 * to add another widget to the viewport at the same position in viewport space as this location, this is
+	 * what you would use.
+	 */
+	UFUNCTION(BlueprintPure, Category="User Interface|Geometry", meta=( HidePin="WorldContextObject", DefaultToSelf="WorldContextObject" ))
+	static void LocalToViewport(UObject* WorldContextObject, const FGeometry& Geometry, FVector2D LocalCoordinate, FVector2D& ScreenPosition, FVector2D& ViewportPosition);
+
+	/**
+	 * Translates absolute coordinate of the geometry provided into local viewport coordinates.
+	 *
+	 * @param ScreenPosition The position in the game's viewport, usable for line traces and
+	 * other uses where you need a coordinate in the space of viewport resolution units.
+	 * @param ViewportPosition The position in the space of other widgets in the viewport.  Like if you wanted
+	 * to add another widget to the viewport at the same position in viewport space as this location, this is
+	 * what you would use.
+	 */
+	UFUNCTION(BlueprintPure, Category="User Interface|Geometry", meta=( HidePin="WorldContextObject", DefaultToSelf="WorldContextObject" ))
+	static void AbsoluteToViewport(UObject* WorldContextObject, const FGeometry& Geometry, FVector2D AbsoluteCoordinate, FVector2D& ScreenPosition, FVector2D& ViewportPosition);
 };
