UMG - Making sure the abs to viewport utility calculation takes resolution scale into account.

diff --git a/Engine/Source/Runtime/UMG/Private/SlateBlueprintLibrary.cpp b/Engine/Source/Runtime/UMG/Private/SlateBlueprintLibrary.cpp
--- a/Engine/Source/Runtime/UMG/Private/SlateBlueprintLibrary.cpp
+++ b/Engine/Source/Runtime/UMG/Private/SlateBlueprintLibrary.cpp
@@ -54,12 +54,22 @@ void USlateBlueprintLibrary::AbsoluteToViewport(UObject* WorldContextObject, con
 				FVector2D ViewportSize;
 				ViewportClient->GetViewportSize(ViewportSize);
 
-				FVector2D AbsolutePositionPostDPI = Viewport->VirtualDesktopPixelToViewport(FIntPoint((int32)AbsoluteCoordinate.X, (int32)AbsoluteCoordinate.Y)) * ViewportSize;
+				FVector2D PixelLocation = Viewport->VirtualDesktopPixelToViewport(FIntPoint((int32)AbsoluteCoordinate.X, (int32)AbsoluteCoordinate.Y)) * ViewportSize;
 
-				float CurrentViewportScale = GetDefault<UUserInterfaceSettings>(UUserInterfaceSettings::StaticClass())->GetDPIScaleBasedOnSize(FIntPoint(ViewportSize.X, ViewportSize.Y));
+				float CurrentViewportScale = GetDefault<UUserInterfaceSettings>()->GetDPIScaleBasedOnSize(FIntPoint(ViewportSize.X, ViewportSize.Y));
 
-				ScreenPosition = AbsolutePositionPostDPI;
-				ViewportPosition = AbsolutePositionPostDPI / CurrentViewportScale;
+				// If the user has configured a resolution quality we need to multiply
+				// the pixels by the resolution quality to arrive at the true position in
+				// the viewport, as the rendered image will be stretched to fill whatever
+				// size the viewport is at.
+				Scalability::FQualityLevels ScalabilityQuality = Scalability::GetQualityLevels();
+				const float QualityScale = ( ScalabilityQuality.ResolutionQuality / 100.0f );
+
+				// Remove the resolution quality scale.
+				ScreenPosition = PixelLocation * QualityScale;
+
+				// Remove DPI Scaling.
+				ViewportPosition = PixelLocation / CurrentViewportScale;
 
 				return;
 			}
diff --git a/Engine/Source/Runtime/UMG/Private/WidgetLayoutLibrary.cpp b/Engine/Source/Runtime/UMG/Private/WidgetLayoutLibrary.cpp
--- a/Engine/Source/Runtime/UMG/Private/WidgetLayoutLibrary.cpp
+++ b/Engine/Source/Runtime/UMG/Private/WidgetLayoutLibrary.cpp
@@ -57,7 +57,7 @@ float UWidgetLayoutLibrary::GetViewportScale(UObject* WorldContextObject)
 		{
 			FVector2D ViewportSize;
 			ViewportClient->GetViewportSize(ViewportSize);
-			return GetDefault<UUserInterfaceSettings>(UUserInterfaceSettings::StaticClass())->GetDPIScaleBasedOnSize(FIntPoint(ViewportSize.X, ViewportSize.Y));
+			return GetDefault<UUserInterfaceSettings>()->GetDPIScaleBasedOnSize(FIntPoint(ViewportSize.X, ViewportSize.Y));
 		}
 	}
 
