#Orion - Tweaks to some UMG widgets / blueprint functions

--------
Integrated using branch UE4-To-UE4-Orion (reversed) of change#2546377 by [REDACTED_PERSON] on 2015/05/11 21:29:06.

diff --git a/Engine/Source/Runtime/UMG/Private/Components/WidgetSwitcher.cpp b/Engine/Source/Runtime/UMG/Private/Components/WidgetSwitcher.cpp
--- a/Engine/Source/Runtime/UMG/Private/Components/WidgetSwitcher.cpp
+++ b/Engine/Source/Runtime/UMG/Private/Components/WidgetSwitcher.cpp
@@ -66,6 +66,16 @@ void UWidgetSwitcher::SetActiveWidget(UWidget* Widget)
 	}
 }
 
+UWidget* UWidgetSwitcher::GetWidgetAtIndex( int32 Index ) const
+{
+	if ( Slots.IsValidIndex( Index ) )
+	{
+		return Slots[ Index ]->Content;
+	}
+
+	return nullptr;
+}
+
 UClass* UWidgetSwitcher::GetSlotClass() const
 {
 	return UWidgetSwitcherSlot::StaticClass();
diff --git a/Engine/Source/Runtime/UMG/Private/Slate/SObjectWidget.cpp b/Engine/Source/Runtime/UMG/Private/Slate/SObjectWidget.cpp
--- a/Engine/Source/Runtime/UMG/Private/Slate/SObjectWidget.cpp
+++ b/Engine/Source/Runtime/UMG/Private/Slate/SObjectWidget.cpp
@@ -256,7 +256,11 @@ FReply SObjectWidget::OnMouseWheel(const FGeometry& MyGeometry, const FPointerEv
 
 FCursorReply SObjectWidget::OnCursorQuery(const FGeometry& MyGeometry, const FPointerEvent& CursorEvent) const
 {
-	// TODO UMG  Allow conditional overriding of the cursor logic.
+	if ( CanRouteEvent() )
+	{
+		return WidgetObject->NativeOnCursorQuery( MyGeometry, CursorEvent );
+	}
+
 	return FCursorReply::Unhandled();
 }
 
diff --git a/Engine/Source/Runtime/UMG/Private/UserWidget.cpp b/Engine/Source/Runtime/UMG/Private/UserWidget.cpp
--- a/Engine/Source/Runtime/UMG/Private/UserWidget.cpp
+++ b/Engine/Source/Runtime/UMG/Private/UserWidget.cpp
@@ -887,6 +887,11 @@ FReply UUserWidget::NativeOnMotionDetected( const FGeometry& InGeometry, const F
 	return OnMotionDetected( InGeometry, InMotionEvent ).NativeReply;
 }
 
+FCursorReply UUserWidget::NativeOnCursorQuery( const FGeometry& InGeometry, const FPointerEvent& InCursorEvent )
+{
+	return FCursorReply::Unhandled();
+}
+
 void UUserWidget::PostLoad()
 {
 	Super::PostLoad();
diff --git a/Engine/Source/Runtime/UMG/Private/WidgetBlueprintLibrary.cpp b/Engine/Source/Runtime/UMG/Private/WidgetBlueprintLibrary.cpp
--- a/Engine/Source/Runtime/UMG/Private/WidgetBlueprintLibrary.cpp
+++ b/Engine/Source/Runtime/UMG/Private/WidgetBlueprintLibrary.cpp
@@ -183,7 +183,28 @@ FEventReply UWidgetBlueprintLibrary::ReleaseMouseCapture(UPARAM(ref) FEventReply
 	return Reply;
 }
 
-FEventReply UWidgetBlueprintLibrary::SetUserFocus(UPARAM(ref) FEventReply& Reply, UWidget* FocusWidget, bool bInAllUsers/* = false*/)
+FEventReply UWidgetBlueprintLibrary::LockMouse( UPARAM( ref ) FEventReply& Reply, UWidget* CapturingWidget )
+{
+	if ( CapturingWidget )
+	{
+		TSharedPtr< SWidget > SlateWidget = CapturingWidget->GetCachedWidget();
+		if ( SlateWidget.IsValid() )
+		{
+			Reply.NativeReply = Reply.NativeReply.LockMouseToWidget( SlateWidget.ToSharedRef() );
+		}
+	}
+
+	return Reply;
+}
+
+FEventReply UWidgetBlueprintLibrary::UnlockMouse( UPARAM( ref ) FEventReply& Reply )
+{
+	Reply.NativeReply = Reply.NativeReply.ReleaseMouseLock();
+
+	return Reply;
+}
+
+FEventReply UWidgetBlueprintLibrary::SetUserFocus( UPARAM( ref ) FEventReply& Reply, UWidget* FocusWidget, bool bInAllUsers/* = false*/ )
 {
 	if (FocusWidget)
 	{
diff --git a/Engine/Source/Runtime/UMG/Public/Animation/WidgetAnimation.h b/Engine/Source/Runtime/UMG/Public/Animation/WidgetAnimation.h
--- a/Engine/Source/Runtime/UMG/Public/Animation/WidgetAnimation.h
+++ b/Engine/Source/Runtime/UMG/Public/Animation/WidgetAnimation.h
@@ -45,7 +45,7 @@ public:
 };
 
 
-UCLASS(MinimalAPI)
+UCLASS( BlueprintType, MinimalAPI )
 class UWidgetAnimation : public UObject
 {
 	GENERATED_UCLASS_BODY()
diff --git a/Engine/Source/Runtime/UMG/Public/Blueprint/UserWidget.h b/Engine/Source/Runtime/UMG/Public/Blueprint/UserWidget.h
--- a/Engine/Source/Runtime/UMG/Public/Blueprint/UserWidget.h
+++ b/Engine/Source/Runtime/UMG/Public/Blueprint/UserWidget.h
@@ -892,6 +892,7 @@ protected:
 	virtual FReply NativeOnTouchMoved( const FGeometry& InGeometry, const FPointerEvent& InGestureEvent );
 	virtual FReply NativeOnTouchEnded( const FGeometry& InGeometry, const FPointerEvent& InGestureEvent );
 	virtual FReply NativeOnMotionDetected( const FGeometry& InGeometry, const FMotionEvent& InMotionEvent );
+	virtual FCursorReply NativeOnCursorQuery( const FGeometry& InGeometry, const FPointerEvent& InCursorEvent );
 
 protected:
 	/**
diff --git a/Engine/Source/Runtime/UMG/Public/Blueprint/WidgetBlueprintLibrary.h b/Engine/Source/Runtime/UMG/Public/Blueprint/WidgetBlueprintLibrary.h
--- a/Engine/Source/Runtime/UMG/Public/Blueprint/WidgetBlueprintLibrary.h
+++ b/Engine/Source/Runtime/UMG/Public/Blueprint/WidgetBlueprintLibrary.h
@@ -87,6 +87,12 @@ public:
 	UFUNCTION(BlueprintPure, Category="Widget|Event Reply")
 	static FEventReply ReleaseMouseCapture(UPARAM(ref) FEventReply& Reply);
 
+	UFUNCTION( BlueprintPure, Category = "Widget|Event Reply", meta = ( HidePin = "CapturingWidget", DefaultToSelf = "CapturingWidget" ) )
+	static FEventReply LockMouse( UPARAM( ref ) FEventReply& Reply, UWidget* CapturingWidget );
+
+	UFUNCTION( BlueprintPure, Category = "Widget|Event Reply" )
+	static FEventReply UnlockMouse( UPARAM( ref ) FEventReply& Reply );
+
 	/**  */
 	UFUNCTION(BlueprintPure, meta= (HidePin="CapturingWidget", DefaultToSelf="CapturingWidget"), Category="Widget|Event Reply")
 	static FEventReply SetUserFocus(UPARAM(ref) FEventReply& Reply, UWidget* FocusWidget, bool bInAllUsers = false);
diff --git a/Engine/Source/Runtime/UMG/Public/Components/WidgetSwitcher.h b/Engine/Source/Runtime/UMG/Public/Components/WidgetSwitcher.h
--- a/Engine/Source/Runtime/UMG/Public/Components/WidgetSwitcher.h
+++ b/Engine/Source/Runtime/UMG/Public/Components/WidgetSwitcher.h
@@ -34,6 +34,10 @@ public:
 	/** Activates the widget and makes it the active index. */
 	UFUNCTION(BlueprintCallable, Category="Switcher")
 	void SetActiveWidget(UWidget* Widget);
+
+	/** Get a widget at the provided index */
+	UFUNCTION( BlueprintCallable, Category = "Switcher" )
+	UWidget* GetWidgetAtIndex( int32 Index ) const;
 	
 	// UWidget interface
 	virtual void SynchronizeProperties() override;
