Added a method to change a text block's font information via function.

diff --git a/Engine/Source/Runtime/UMG/Private/Components/TextBlock.cpp b/Engine/Source/Runtime/UMG/Private/Components/TextBlock.cpp
--- a/Engine/Source/Runtime/UMG/Private/Components/TextBlock.cpp
+++ b/Engine/Source/Runtime/UMG/Private/Components/TextBlock.cpp
@@ -81,6 +81,15 @@ void UTextBlock::SetJustification( ETextJustify::Type InJustification )
 	}
 }
 
+void UTextBlock::SetFont( FSlateFontInfo FontInfo )
+{
+	Font = FontInfo;
+	if ( MyTextBlock.IsValid() )
+	{
+		MyTextBlock->SetFont( Font );
+	}
+}
+
 TSharedRef<SWidget> UTextBlock::RebuildWidget()
 {
 	MyTextBlock = SNew(STextBlock);
diff --git a/Engine/Source/Runtime/UMG/Public/Components/TextBlock.h b/Engine/Source/Runtime/UMG/Public/Components/TextBlock.h
--- a/Engine/Source/Runtime/UMG/Public/Components/TextBlock.h
+++ b/Engine/Source/Runtime/UMG/Public/Components/TextBlock.h
@@ -64,6 +64,14 @@ public:
 	UFUNCTION( BlueprintCallable, Category = "Appearance" )
 	void SetJustification( ETextJustify::Type InJustification );
 
+	/**
+	 * Set the font information struct
+	 *
+	 * @param FontInfo new font info
+	 */
+	UFUNCTION( BlueprintCallable, Category = "Appearance" )
+	void SetFont( FSlateFontInfo FontInfo );
+
 public:
 	UPROPERTY()
 	USlateWidgetStyleAsset* Style_DEPRECATED;
