UMG - Widget Components now support 3 blend modes, Opaque, Transparent, Masked (new). Masked allows you to have the trasperency you need on edges and such, but actually writes to the depth buffer, so that it sorts correctly.  There is also now an option for ticking when the in world widget is offscreen to continue to allow animations to play and tick to be called on the widget.

diff --git a/Engine/Source/Runtime/Core/Public/UObject/ObjectVersion.h b/Engine/Source/Runtime/Core/Public/UObject/ObjectVersion.h
--- a/Engine/Source/Runtime/Core/Public/UObject/ObjectVersion.h
+++ b/Engine/Source/Runtime/Core/Public/UObject/ObjectVersion.h
@@ -514,6 +514,8 @@ enum EUnrealEngineObjectUE4Version
 	VER_UE4_SKIP_DUPLICATE_EXPORTS_ON_SAVE_PACKAGE,
 	// Pre-gathering of gatherable, localizable text in packages to optimize text gathering operation times
 	VER_UE4_SERIALIZE_TEXT_IN_PACKAGES,
+	// Added pivot to widget components, need to load old versions as a 0,0 pivot, new default is 0.5,0.5
+	VER_UE4_ADD_BLEND_MODE_TO_WIDGET_COMPONENT,
 
 	// -----<new versions can be added before this line>-------------------------------------------------
 	// - this needs to be the last line (see note below)
diff --git a/Engine/Source/Runtime/UMG/Private/Components/WidgetComponent.cpp b/Engine/Source/Runtime/UMG/Private/Components/WidgetComponent.cpp
--- a/Engine/Source/Runtime/UMG/Private/Components/WidgetComponent.cpp
+++ b/Engine/Source/Runtime/UMG/Private/Components/WidgetComponent.cpp
@@ -45,7 +45,7 @@ public:
 		, RenderTarget( InComponent->GetRenderTarget() )
 		, MaterialInstance( InComponent->GetMaterialInstance() )
 		, BodySetup( InComponent->GetBodySetup() )
-		, bIsOpaque( InComponent->IsOpaque() )
+		, BlendMode( InComponent->GetBlendMode() )
 	{
 		bWillEverBeLit = false;	
 	}
@@ -121,13 +121,14 @@ public:
 
 	virtual FPrimitiveViewRelevance GetViewRelevance(const FSceneView* View) override
 	{
-		bool bVisible = true;// View->Family->EngineShowFlags.BillboardSprites;
+		bool bVisible = true;
 
 		FPrimitiveViewRelevance Result;
 		Result.bDrawRelevance = IsShown(View) && bVisible;
-		Result.bOpaqueRelevance = bIsOpaque;
-		Result.bNormalTranslucencyRelevance = !bIsOpaque;
-		Result.bSeparateTranslucencyRelevance = !bIsOpaque;
+		Result.bOpaqueRelevance = BlendMode == EWidgetBlendMode::Opaque;
+		Result.bMaskedRelevance = BlendMode == EWidgetBlendMode::Masked;
+		Result.bNormalTranslucencyRelevance = BlendMode == EWidgetBlendMode::Transparent;
+		Result.bSeparateTranslucencyRelevance = BlendMode == EWidgetBlendMode::Transparent;
 		Result.bDynamicRelevance = true;
 		Result.bShadowRelevance = IsShadowCast(View);
 		Result.bEditorPrimitiveRelevance = false;
@@ -158,7 +159,7 @@ private:
 	UTextureRenderTarget2D* RenderTarget;
 	UMaterialInstanceDynamic* MaterialInstance;
 	UBodySetup* BodySetup;
-	bool bIsOpaque;
+	EWidgetBlendMode BlendMode;
 };
 
 /**
@@ -317,8 +318,10 @@ UWidgetComponent::UWidgetComponent( const FObjectInitializer& PCIP )
 	, DrawSize( FIntPoint( 500, 500 ) )
 	, MaxInteractionDistance( 1000.f )
 	, BackgroundColor( FLinearColor::Transparent )
-	, bIsOpaque( false )
+	, BlendMode( EWidgetBlendMode::Masked )
+	, bIsOpaque_DEPRECATED( false )
 	, bIsTwoSided( false )
+	, TickWhenOffscreen( false )
 {
 	PrimaryComponentTick.bCanEverTick = true;
 	bTickInEditor = true;
@@ -329,18 +332,22 @@ UWidgetComponent::UWidgetComponent( const FObjectInitializer& PCIP )
 
 	// Translucent material instances
 	static ConstructorHelpers::FObjectFinder<UMaterialInterface> TranslucentMaterial_Finder( TEXT("/Engine/EngineMaterials/Widget3DPassThrough_Translucent") );
+	static ConstructorHelpers::FObjectFinder<UMaterialInterface> TranslucentMaterial_OneSided_Finder(TEXT("/Engine/EngineMaterials/Widget3DPassThrough_Translucent_OneSided"));
 	TranslucentMaterial = TranslucentMaterial_Finder.Object;
-
-	static ConstructorHelpers::FObjectFinder<UMaterialInterface> TranslucentMaterial_OneSided_Finder( TEXT( "/Engine/EngineMaterials/Widget3DPassThrough_Translucent_OneSided" ) );
 	TranslucentMaterial_OneSided = TranslucentMaterial_OneSided_Finder.Object;
 
 	// Opaque material instances
 	static ConstructorHelpers::FObjectFinder<UMaterialInterface> OpaqueMaterial_Finder( TEXT( "/Engine/EngineMaterials/Widget3DPassThrough_Opaque" ) );
+	static ConstructorHelpers::FObjectFinder<UMaterialInterface> OpaqueMaterial_OneSided_Finder(TEXT("/Engine/EngineMaterials/Widget3DPassThrough_Opaque_OneSided"));
 	OpaqueMaterial = OpaqueMaterial_Finder.Object;
-
-	static ConstructorHelpers::FObjectFinder<UMaterialInterface> OpaqueMaterial_OneSided_Finder( TEXT( "/Engine/EngineMaterials/Widget3DPassThrough_Opaque_OneSided" ) );
 	OpaqueMaterial_OneSided = OpaqueMaterial_OneSided_Finder.Object;
 
+	// Masked material instances
+	static ConstructorHelpers::FObjectFinder<UMaterialInterface> MaskedMaterial_Finder(TEXT("/Engine/EngineMaterials/Widget3DPassThrough_Masked"));
+	static ConstructorHelpers::FObjectFinder<UMaterialInterface> MaskedMaterial_OneSided_Finder(TEXT("/Engine/EngineMaterials/Widget3DPassThrough_Masked_OneSided"));
+	MaskedMaterial = MaskedMaterial_Finder.Object;
+	MaskedMaterial_OneSided = MaskedMaterial_OneSided_Finder.Object;
+
 	LastLocalHitLocation = FVector2D::ZeroVector;
 	//bGenerateOverlapEvents = false;
 	bUseEditorCompositing = false;
@@ -439,13 +446,17 @@ void UWidgetComponent::OnRegister()
 			if ( !MaterialInstance )
 			{
 				UMaterialInterface* Parent = nullptr;
-				if ( bIsOpaque )
+				switch ( BlendMode )
 				{
+				case EWidgetBlendMode::Opaque:
 					Parent = bIsTwoSided ? OpaqueMaterial : OpaqueMaterial_OneSided;
-				}
-				else
-				{
+					break;
+				case EWidgetBlendMode::Masked:
+					Parent = bIsTwoSided ? MaskedMaterial : MaskedMaterial_OneSided;
+					break;
+				case EWidgetBlendMode::Transparent:
 					Parent = bIsTwoSided ? TranslucentMaterial : TranslucentMaterial_OneSided;
+					break;
 				}
 
 				MaterialInstance = UMaterialInstanceDynamic::Create(Parent, this);
@@ -533,49 +544,52 @@ void UWidgetComponent::TickComponent(float DeltaTime, enum ELevelTick TickType,
 		}
 
 		const float RenderTimeThreshold = .5f;
-		if ( IsVisible() && GetWorld()->TimeSince(LastRenderTime) <= RenderTimeThreshold ) // Don't bother ticking if it hasn't been rendered recently
+		if ( IsVisible() )
 		{
-			SlateWidget->SlatePrepass();
+			// if we don't tick when offscreen, don't bother ticking if it hasn't been rendered recently
+			if ( TickWhenOffscreen || GetWorld()->TimeSince(LastRenderTime) <= RenderTimeThreshold )
+			{
+				SlateWidget->SlatePrepass();
 
-			FGeometry WindowGeometry = FGeometry::MakeRoot(DrawSize, FSlateLayoutTransform());
+				FGeometry WindowGeometry = FGeometry::MakeRoot(DrawSize, FSlateLayoutTransform());
 
-			SlateWidget->TickWidgetsRecursively(WindowGeometry, FApp::GetCurrentTime(), DeltaTime);
+				SlateWidget->TickWidgetsRecursively(WindowGeometry, FApp::GetCurrentTime(), DeltaTime);
 
-			// Ticking can cause geometry changes.  Recompute
-			SlateWidget->SlatePrepass();
+				// Ticking can cause geometry changes.  Recompute
+				SlateWidget->SlatePrepass();
 
-			// Get the free buffer & add our virtual window
-			FSlateDrawBuffer& DrawBuffer = Renderer->GetDrawBuffer();
-			FSlateWindowElementList& WindowElementList = DrawBuffer.AddWindowElementList(SlateWidget.ToSharedRef());
+				// Get the free buffer & add our virtual window
+				FSlateDrawBuffer& DrawBuffer = Renderer->GetDrawBuffer();
+				FSlateWindowElementList& WindowElementList = DrawBuffer.AddWindowElementList(SlateWidget.ToSharedRef());
 
-			int32 MaxLayerId = 0;
-			{
-				// Prepare the test grid 
-				HitTestGrid->ClearGridForNewFrame(WindowGeometry.GetClippingRect());
-
-				// Paint the window
-				MaxLayerId = SlateWidget->Paint(
-					FPaintArgs(SlateWidget.ToSharedRef(), *HitTestGrid, FVector2D::ZeroVector, FSlateApplication::Get().GetCurrentTime(), FSlateApplication::Get().GetDeltaTime()),
-					WindowGeometry, WindowGeometry.GetClippingRect(),
-					WindowElementList,
-					0,
-					FWidgetStyle(),
-					SlateWidget->IsEnabled());
-			}
+				int32 MaxLayerId = 0;
+				{
+					// Prepare the test grid 
+					HitTestGrid->ClearGridForNewFrame(WindowGeometry.GetClippingRect());
+
+					// Paint the window
+					MaxLayerId = SlateWidget->Paint(
+						FPaintArgs(SlateWidget.ToSharedRef(), *HitTestGrid, FVector2D::ZeroVector, FSlateApplication::Get().GetCurrentTime(), FSlateApplication::Get().GetDeltaTime()),
+						WindowGeometry, WindowGeometry.GetClippingRect(),
+						WindowElementList,
+						0,
+						FWidgetStyle(),
+						SlateWidget->IsEnabled());
+				}
 
-			Renderer->DrawWindow_GameThread(DrawBuffer);
+				Renderer->DrawWindow_GameThread(DrawBuffer);
 
-			UpdateRenderTarget();
-
-			// Enqueue a command to unlock the draw buffer after all windows have been drawn
-			ENQUEUE_UNIQUE_RENDER_COMMAND_THREEPARAMETER(UWidgetComponentRenderToTexture,
-				FSlateDrawBuffer&, InDrawBuffer, DrawBuffer,
-				UTextureRenderTarget2D*, InRenderTarget, RenderTarget,
-				ISlate3DRenderer*, InRenderer, Renderer.Get(),
-				{
-					InRenderer->DrawWindowToTarget_RenderThread(RHICmdList, InRenderTarget, InDrawBuffer);
-				});
+				UpdateRenderTarget();
 
+				// Enqueue a command to unlock the draw buffer after all windows have been drawn
+				ENQUEUE_UNIQUE_RENDER_COMMAND_THREEPARAMETER(UWidgetComponentRenderToTexture,
+					FSlateDrawBuffer&, InDrawBuffer, DrawBuffer,
+					UTextureRenderTarget2D*, InRenderTarget, RenderTarget,
+					ISlate3DRenderer*, InRenderer, Renderer.Get(),
+					{
+						InRenderer->DrawWindowToTarget_RenderThread(RHICmdList, InRenderTarget, InDrawBuffer);
+					});
+			}
 		}
 	}
 	else
@@ -803,16 +817,19 @@ void UWidgetComponent::UpdateRenderTarget()
 {
 	bool bClearColorChanged = false;
 
+	FLinearColor ActualBackgroundColor = BackgroundColor;
+	switch ( BlendMode )
+	{
+	case EWidgetBlendMode::Opaque:
+		ActualBackgroundColor.A = 1.0f;
+	case EWidgetBlendMode::Masked:
+		ActualBackgroundColor.A = 0.0f;
+	}
+
 	if(!RenderTarget && DrawSize != FIntPoint::ZeroValue)
 	{
 		RenderTarget = NewObject<UTextureRenderTarget2D>(this);
-
-		RenderTarget->ClearColor = BackgroundColor;
-
-		if (bIsOpaque)
-		{
-			RenderTarget->ClearColor.A = 1.0f;
-		}
+		RenderTarget->ClearColor = ActualBackgroundColor;
 
 		bClearColorChanged = true;
 
@@ -831,21 +848,11 @@ void UWidgetComponent::UpdateRenderTarget()
 		}
 
 		// Update the clear color
-		if (!bIsOpaque && RenderTarget->ClearColor != BackgroundColor)
+		if ( RenderTarget->ClearColor != ActualBackgroundColor )
 		{
-			RenderTarget->ClearColor = BackgroundColor;
+			RenderTarget->ClearColor = ActualBackgroundColor;
 			bClearColorChanged = true;
 		}
-		else if ( bIsOpaque )
-		{
-			// If opaque, make sure the alpha channel is set to 1.0
-			FLinearColor ClearColor( BackgroundColor.R, BackgroundColor.G, BackgroundColor.B, 1.0f );
-			if ( RenderTarget->ClearColor != ClearColor )
-			{
-				RenderTarget->ClearColor = ClearColor;
-				bClearColorChanged = true;
-			}
-		}
 	}
 
 	// If the clear color of the render target changed, update the BackColor of the material to match
@@ -953,4 +960,9 @@ void UWidgetComponent::PostLoad()
 	{
 		Pivot = FVector2D(0, 0);
 	}
+
+	if ( GetLinkerUE4Version() < VER_UE4_ADD_BLEND_MODE_TO_WIDGET_COMPONENT )
+	{
+		BlendMode = bIsOpaque_DEPRECATED ? EWidgetBlendMode::Opaque : EWidgetBlendMode::Transparent;
+	}
 }
diff --git a/Engine/Source/Runtime/UMG/Public/Components/WidgetComponent.h b/Engine/Source/Runtime/UMG/Public/Components/WidgetComponent.h
--- a/Engine/Source/Runtime/UMG/Public/Components/WidgetComponent.h
+++ b/Engine/Source/Runtime/UMG/Public/Components/WidgetComponent.h
@@ -15,6 +15,14 @@ enum class EWidgetSpace : uint8
 	Screen
 };
 
+UENUM()
+enum class EWidgetBlendMode : uint8
+{
+	Opaque,
+	Masked,
+	Transparent
+};
+
 /**
  * Beware! This feature is experimental and may be substantially changed or removed in future releases.
  * A 3D instance of a Widget Blueprint that can be interacted with in the world.
@@ -117,8 +125,8 @@ public:
 	UFUNCTION(BlueprintCallable, Category=UI)
 	void SetMaxInteractionDistance(float Distance);
 
-	/** @return True if the component is opaque */
-	bool IsOpaque() const { return bIsOpaque; }
+	/** Gets the blend mode for the widget. */
+	EWidgetBlendMode GetBlendMode() const { return BlendMode; }
 
 	/** @return The pivot point where the UI is rendered about the origin. */
 	FVector2D GetPivot() const { return Pivot; }
@@ -159,17 +167,21 @@ private:
 	UPROPERTY(EditAnywhere, Category=Rendering)
 	FLinearColor BackgroundColor;
 
-	/** 
-	 * Should the component be rendered opaque? 
-	 * This improves aliasing of the UI in the world.
-	 */
+	/** The blend mode for the widget. */
 	UPROPERTY(EditAnywhere, Category=Rendering)
-	bool bIsOpaque;
+	EWidgetBlendMode BlendMode;
+
+	UPROPERTY()
+	bool bIsOpaque_DEPRECATED;
 
 	/** Is the component visible from behind? */
 	UPROPERTY(EditAnywhere, Category=Rendering)
 	bool bIsTwoSided;
 
+	/** Should the component tick the widget when it's off screen? */
+	UPROPERTY(EditAnywhere, Category=Animation)
+	bool TickWhenOffscreen;
+
 	/** The User Widget object displayed and managed by this component */
 	UPROPERTY(transient, duplicatetransient)
 	UUserWidget* Widget;
@@ -194,6 +206,14 @@ private:
 	UPROPERTY()
 	UMaterialInterface* OpaqueMaterial_OneSided;
 
+	/** The material instance for masked widget components. */
+	UPROPERTY()
+	UMaterialInterface* MaskedMaterial;
+
+	/** The material instance for masked, one-sided widget components. */
+	UPROPERTY()
+	UMaterialInterface* MaskedMaterial_OneSided;
+
 	/** The target to which the user widget is rendered */
 	UPROPERTY(transient, duplicatetransient)
 	UTextureRenderTarget2D* RenderTarget;
