fixed UE-15921 Skeletal Meshes with Multiple Translucent Batches will render Both

diff --git a/Engine/Source/Developer/FunctionalTesting/Private/FuncTestRenderingComponent.cpp b/Engine/Source/Developer/FunctionalTesting/Private/FuncTestRenderingComponent.cpp
--- a/Engine/Source/Developer/FunctionalTesting/Private/FuncTestRenderingComponent.cpp
+++ b/Engine/Source/Developer/FunctionalTesting/Private/FuncTestRenderingComponent.cpp
@@ -36,7 +36,8 @@ public:
 		FPrimitiveViewRelevance Result;
 		Result.bDrawRelevance = IsShown(View) && IsSelected();
 		Result.bDynamicRelevance = true;
-		Result.bNormalTranslucencyRelevance = IsShown(View);
+		// ideally the TranslucencyRelevance should be filled out by the material, here we do it conservative
+		Result.bSeparateTranslucencyRelevance = Result.bNormalTranslucencyRelevance = IsShown(View);
 		return Result;
 	}
 
diff --git a/Engine/Source/Developer/GameplayDebugger/Private/GameplayDebuggingComponent.cpp b/Engine/Source/Developer/GameplayDebugger/Private/GameplayDebuggingComponent.cpp
--- a/Engine/Source/Developer/GameplayDebugger/Private/GameplayDebuggingComponent.cpp
+++ b/Engine/Source/Developer/GameplayDebugger/Private/GameplayDebuggingComponent.cpp
@@ -1383,7 +1383,8 @@ public:
 		FPrimitiveViewRelevance Result;
 		Result.bDrawRelevance = View->Family->EngineShowFlags.GetSingleFlag(ViewFlagIndex);// IsShown(View);
 		Result.bDynamicRelevance = true;
-		Result.bNormalTranslucencyRelevance = IsShown(View);
+		// ideally the TranslucencyRelevance should be filled out by the material, here we do it conservative
+		Result.bSeparateTranslucencyRelevance = Result.bNormalTranslucencyRelevance = IsShown(View);
 		return Result;
 	}
 
diff --git a/Engine/Source/Developer/LogVisualizer/Private/VisualLoggerRenderingActor.cpp b/Engine/Source/Developer/LogVisualizer/Private/VisualLoggerRenderingActor.cpp
--- a/Engine/Source/Developer/LogVisualizer/Private/VisualLoggerRenderingActor.cpp
+++ b/Engine/Source/Developer/LogVisualizer/Private/VisualLoggerRenderingActor.cpp
@@ -31,7 +31,8 @@ public:
 		FPrimitiveViewRelevance Result;
 		Result.bDrawRelevance = IsShown(View);
 		Result.bDynamicRelevance = true;
-		Result.bNormalTranslucencyRelevance = IsShown(View) && GIsEditor;
+		// ideally the TranslucencyRelevance should be filled out by the material, here we do it conservative
+		Result.bSeparateTranslucencyRelevance = Result.bNormalTranslucencyRelevance = IsShown(View) && GIsEditor;
 		return Result;
 	}
 
diff --git a/Engine/Source/Runtime/AIModule/Private/EnvironmentQuery/EQSRenderingComponent.cpp b/Engine/Source/Runtime/AIModule/Private/EnvironmentQuery/EQSRenderingComponent.cpp
--- a/Engine/Source/Runtime/AIModule/Private/EnvironmentQuery/EQSRenderingComponent.cpp
+++ b/Engine/Source/Runtime/AIModule/Private/EnvironmentQuery/EQSRenderingComponent.cpp
@@ -291,7 +291,8 @@ FPrimitiveViewRelevance FEQSSceneProxy::GetViewRelevance(const FSceneView* View)
 	Result.bDrawRelevance = View->Family->EngineShowFlags.GetSingleFlag(ViewFlagIndex) && IsShown(View) 
 		&& (!bDrawOnlyWhenSelected || SafeIsActorSelected());
 	Result.bDynamicRelevance = true;
-	Result.bNormalTranslucencyRelevance = IsShown(View);
+	// ideally the TranslucencyRelevance should be filled out by the material, here we do it conservative
+	Result.bSeparateTranslucencyRelevance = Result.bNormalTranslucencyRelevance = IsShown(View);
 	return Result;
 }
 
diff --git a/Engine/Source/Runtime/Engine/Private/AI/Navigation/NavLinkRenderingComponent.cpp b/Engine/Source/Runtime/Engine/Private/AI/Navigation/NavLinkRenderingComponent.cpp
--- a/Engine/Source/Runtime/Engine/Private/AI/Navigation/NavLinkRenderingComponent.cpp
+++ b/Engine/Source/Runtime/Engine/Private/AI/Navigation/NavLinkRenderingComponent.cpp
@@ -422,7 +422,8 @@ FPrimitiveViewRelevance FNavLinkRenderingProxy::GetViewRelevance(const FSceneVie
 	FPrimitiveViewRelevance Result;
 	Result.bDrawRelevance = IsShown(View) && IsSelected() && (View && View->Family && View->Family->EngineShowFlags.Navigation);
 	Result.bDynamicRelevance = true;
-	Result.bNormalTranslucencyRelevance = IsShown(View);
+	// ideally the TranslucencyRelevance should be filled out by the material, here we do it conservative
+	Result.bSeparateTranslucencyRelevance = Result.bNormalTranslucencyRelevance = IsShown(View);
 	Result.bShadowRelevance = IsShadowCast(View);
 	Result.bEditorPrimitiveRelevance = UseEditorCompositing(View);
 	return Result;
diff --git a/Engine/Source/Runtime/Engine/Private/AI/Navigation/NavTestRenderingComponent.cpp b/Engine/Source/Runtime/Engine/Private/AI/Navigation/NavTestRenderingComponent.cpp
--- a/Engine/Source/Runtime/Engine/Private/AI/Navigation/NavTestRenderingComponent.cpp
+++ b/Engine/Source/Runtime/Engine/Private/AI/Navigation/NavTestRenderingComponent.cpp
@@ -381,7 +381,8 @@ public:
 		FPrimitiveViewRelevance Result;
 		Result.bDrawRelevance = IsShown(View);
 		Result.bDynamicRelevance = true;
-		Result.bNormalTranslucencyRelevance = IsShown(View) && GIsEditor;
+		// ideally the TranslucencyRelevance should be filled out by the material, here we do it conservative
+		Result.bSeparateTranslucencyRelevance = Result.bNormalTranslucencyRelevance = IsShown(View) && GIsEditor;
 		return Result;
 	}
 
diff --git a/Engine/Source/Runtime/Engine/Private/Components/LineBatchComponent.cpp b/Engine/Source/Runtime/Engine/Private/Components/LineBatchComponent.cpp
--- a/Engine/Source/Runtime/Engine/Private/Components/LineBatchComponent.cpp
+++ b/Engine/Source/Runtime/Engine/Private/Components/LineBatchComponent.cpp
@@ -21,7 +21,8 @@ public:
 		bWillEverBeLit = false;
 		ViewRelevance.bDrawRelevance = true;
 		ViewRelevance.bDynamicRelevance = true;
-		ViewRelevance.bNormalTranslucencyRelevance = true;
+		// ideally the TranslucencyRelevance should be filled out by the material, here we do it conservative
+		ViewRelevance.bSeparateTranslucencyRelevance = ViewRelevance.bNormalTranslucencyRelevance = true;
 	}
 
 	virtual void GetDynamicMeshElements(const TArray<const FSceneView*>& Views, const FSceneViewFamily& ViewFamily, uint32 VisibilityMap, FMeshElementCollector& Collector) const override
diff --git a/Engine/Source/Runtime/Engine/Private/NavMeshRenderingHelpers.h b/Engine/Source/Runtime/Engine/Private/NavMeshRenderingHelpers.h
--- a/Engine/Source/Runtime/Engine/Private/NavMeshRenderingHelpers.h
+++ b/Engine/Source/Runtime/Engine/Private/NavMeshRenderingHelpers.h
@@ -567,7 +567,8 @@ public:
 		FPrimitiveViewRelevance Result;
 		Result.bDrawRelevance = bVisible && IsShown(View);
 		Result.bDynamicRelevance = true;
-		Result.bNormalTranslucencyRelevance = bVisible && IsShown(View);
+		// ideally the TranslucencyRelevance should be filled out by the material, here we do it conservative
+		Result.bSeparateTranslucencyRelevance = Result.bNormalTranslucencyRelevance = bVisible && IsShown(View);
 		return Result;
 	}
 
diff --git a/Engine/Source/Runtime/Landscape/Private/LandscapeGizmoActor.cpp b/Engine/Source/Runtime/Landscape/Private/LandscapeGizmoActor.cpp
--- a/Engine/Source/Runtime/Landscape/Private/LandscapeGizmoActor.cpp
+++ b/Engine/Source/Runtime/Landscape/Private/LandscapeGizmoActor.cpp
@@ -330,7 +330,8 @@ public:
 		const bool bVisible = View->Family->EngineShowFlags.Landscape;
 		Result.bDrawRelevance = IsShown(View) && bVisible && !View->bIsGameView && GLandscapeEditRenderMode & ELandscapeEditRenderMode::Gizmo;
 		Result.bDynamicRelevance = true;
-		Result.bNormalTranslucencyRelevance = true;
+		// ideally the TranslucencyRelevance should be filled out by the material, here we do it conservative
+		Result.bSeparateTranslucencyRelevance = Result.bNormalTranslucencyRelevance = true;
 #endif
 		return Result;
 	}
diff --git a/Engine/Source/Runtime/UMG/Private/Components/WidgetComponent.cpp b/Engine/Source/Runtime/UMG/Private/Components/WidgetComponent.cpp
--- a/Engine/Source/Runtime/UMG/Private/Components/WidgetComponent.cpp
+++ b/Engine/Source/Runtime/UMG/Private/Components/WidgetComponent.cpp
@@ -127,8 +127,8 @@ public:
 		Result.bDrawRelevance = IsShown(View) && bVisible;
 		Result.bOpaqueRelevance = BlendMode == EWidgetBlendMode::Opaque;
 		Result.bMaskedRelevance = BlendMode == EWidgetBlendMode::Masked;
-		Result.bNormalTranslucencyRelevance = BlendMode == EWidgetBlendMode::Transparent;
-		Result.bSeparateTranslucencyRelevance = BlendMode == EWidgetBlendMode::Transparent;
+		// ideally the TranslucencyRelevance should be filled out by the material, here we do it conservative
+		Result.bSeparateTranslucencyRelevance = Result.bSeparateTranslucencyRelevance = (BlendMode == EWidgetBlendMode::Transparent);
 		Result.bDynamicRelevance = true;
 		Result.bShadowRelevance = IsShadowCast(View);
 		Result.bEditorPrimitiveRelevance = false;
