UMG - The slatebrush's ResourceObject property is no longer directly accessible in a blueprint.  You'll need to use the functions, GetBrushResource, SetBrushResourceToTexture, and SetBrushResourceToMaterial to change the value.

#APIBreak

diff --git a/Engine/Source/Runtime/SlateCore/Public/Styling/SlateBrush.h b/Engine/Source/Runtime/SlateCore/Public/Styling/SlateBrush.h
--- a/Engine/Source/Runtime/SlateCore/Public/Styling/SlateBrush.h
+++ b/Engine/Source/Runtime/SlateCore/Public/Styling/SlateBrush.h
@@ -165,9 +165,9 @@ public:
 	 *
 	 * @return Resource name, or NAME_None if the resource object is not set.
 	 */
-	const FName GetResourceName( ) const
+	const FName GetResourceName() const
 	{
-		return ((ResourceName != NAME_None) || (ResourceObject == nullptr))
+		return ( ( ResourceName != NAME_None ) || ( ResourceObject == nullptr ) )
 			? ResourceName
 			: ResourceObject->GetFName();
 	}
@@ -185,8 +185,8 @@ public:
 	}
 
 	/**
-	* Sets the UObject that represents the brush resource.
-	*/
+	 * Sets the UObject that represents the brush resource.
+	 */
 	void SetResourceObject(class UObject* InResourceObject)
 	{
 		ResourceObject = InResourceObject;
@@ -255,6 +255,8 @@ public:
 		return !(*this == Other);
 	}
 
+	void PostSerialize(const FArchive& Ar);
+
 public:
 
 	/**
@@ -269,7 +271,7 @@ protected:
 	/**
 	 * The image to render for this brush, can be a UTexture2D or Material.
 	 */
-	UPROPERTY(EditAnywhere, BlueprintReadWrite, Category=Brush, meta=( DisplayThumbnail="true", ThumbnailSize="X=40 Y=40", DisplayName="Image", AllowedClasses="Texture2D,MaterialInterface" ))
+	UPROPERTY(EditAnywhere, Category=Brush, meta=( DisplayThumbnail="true", ThumbnailSize="X=40 Y=40", DisplayName="Image", AllowedClasses="Texture2D,MaterialInterface" ))
 	UObject* ResourceObject;
 
 	/** The name of the rendering resource to use */
diff --git a/Engine/Source/Runtime/UMG/Private/WidgetBlueprintLibrary.cpp b/Engine/Source/Runtime/UMG/Private/WidgetBlueprintLibrary.cpp
--- a/Engine/Source/Runtime/UMG/Private/WidgetBlueprintLibrary.cpp
+++ b/Engine/Source/Runtime/UMG/Private/WidgetBlueprintLibrary.cpp
@@ -340,6 +340,21 @@ FSlateBrush UWidgetBlueprintLibrary::MakeBrushFromMaterial(UMaterialInterface* M
 	return FSlateNoResource();
 }
 
+UObject* UWidgetBlueprintLibrary::GetBrushResource(FSlateBrush& Brush)
+{
+	return Brush.GetResourceObject();
+}
+
+void UWidgetBlueprintLibrary::SetBrushResourceToTexture(FSlateBrush& Brush, UTexture2D* Texture)
+{
+	Brush.SetResourceObject(Texture);
+}
+
+void UWidgetBlueprintLibrary::SetBrushResourceToMaterial(FSlateBrush& Brush, UMaterialInterface* Material)
+{
+	Brush.SetResourceObject(Material);
+}
+
 FSlateBrush UWidgetBlueprintLibrary::NoResourceBrush()
 {
 	return FSlateNoResource();
diff --git a/Engine/Source/Runtime/UMG/Public/Blueprint/WidgetBlueprintLibrary.h b/Engine/Source/Runtime/UMG/Public/Blueprint/WidgetBlueprintLibrary.h
--- a/Engine/Source/Runtime/UMG/Public/Blueprint/WidgetBlueprintLibrary.h
+++ b/Engine/Source/Runtime/UMG/Public/Blueprint/WidgetBlueprintLibrary.h
@@ -71,7 +71,10 @@ public:
 	UFUNCTION(BlueprintCallable, Category="Painting")
 	static void DrawText(UPARAM(ref) FPaintContext& Context, const FString& InString, FVector2D Position, FLinearColor Tint = FLinearColor::White);
 
-	/** The default event reply when simply handling an event. */
+	/**
+	 * The event reply to use when you choose to handle an event.  This will prevent the event 
+	 * from continuing to bubble up / down the widget hierarchy.
+	 */
 	UFUNCTION(BlueprintPure, Category="Widget|Event Reply")
 	static FEventReply Handled();
 
@@ -112,8 +115,8 @@ public:
 	static FEventReply SetMousePosition(UPARAM(ref) FEventReply& Reply, FVector2D NewMousePosition);
 
 	/**
-	 * Ask Slate to detect if a user started dragging in this widget.
-	 * If a drag is detected, Slate will send an OnDragDetected event.
+	 * Ask Slate to detect if a user starts dragging in this widget later.  Slate internally tracks the movement
+	 * and if it surpasses the drag threshold, Slate will send an OnDragDetected event to the widget.
 	 *
 	 * @param WidgetDetectingDrag  Detect dragging in this widget
 	 * @param DragKey		       This button should be pressed to detect the drag
@@ -121,6 +124,13 @@ public:
 	UFUNCTION(BlueprintPure, meta=( HidePin="WidgetDetectingDrag", DefaultToSelf="WidgetDetectingDrag" ), Category="Widget|Drag and Drop|Event Reply")
 	static FEventReply DetectDrag(UPARAM(ref) FEventReply& Reply, UWidget* WidgetDetectingDrag, FKey DragKey);
 
+	/**
+	 * Given the pointer event, emit the DetectDrag reply if the provided key was pressed.
+	 * If the DragKey is a touch key, that will also automatically work.
+	 * @param PointerEvent	The pointer device event coming in.
+	 * @param WidgetDetectingDrag  Detect dragging in this widget.
+	 * @param DragKey		       This button should be pressed to detect the drag, won't emit the DetectDrag FEventReply unless this is pressed.
+	 */
 	UFUNCTION(BlueprintCallable, meta=( HidePin="WidgetDetectingDrag", DefaultToSelf="WidgetDetectingDrag" ), Category="Widget|Drag and Drop|Event Reply")
 	static FEventReply DetectDragIfPressed(const FPointerEvent& PointerEvent, UWidget* WidgetDetectingDrag, FKey DragKey);
 
@@ -137,7 +147,7 @@ public:
 	static bool IsDragDropping();
 
 	/**
-	 * Returns the drag and drop operation that is currently occuring if any, otherwise nothing.
+	 * Returns the drag and drop operation that is currently occurring if any, otherwise nothing.
 	 */
 	UFUNCTION(BlueprintPure, BlueprintCosmetic, Category="Widget|Drag and Drop")
 	static UDragDropOperation* GetDragDroppingContent();
@@ -170,6 +180,24 @@ public:
 	UFUNCTION(BlueprintPure, Category="Widget|Brush")
 	static FSlateBrush MakeBrushFromMaterial(UMaterialInterface* Material, int32 Width = 32, int32 Height = 32);
 
+	/**
+	 * Gets the resource object on a brush.  This could be a UTexture2D or a UMaterialInterface.
+	 */
+	UFUNCTION(BlueprintPure, Category="Widget|Brush")
+	static UObject* GetBrushResource(FSlateBrush& Brush);
+
+	/**
+	 * Sets the resource on a brush to be a UTexture2D.
+	 */
+	UFUNCTION(BlueprintCallable, Category="Widget|Brush")
+	static void SetBrushResourceToTexture(FSlateBrush& Brush, UTexture2D* Texture);
+
+	/**
+	 * Sets the resource on a brush to be a Material.
+	 */
+	UFUNCTION(BlueprintCallable, Category="Widget|Brush")
+	static void SetBrushResourceToMaterial(FSlateBrush& Brush, UMaterialInterface* Material);
+
 	/**
 	 * Creates a Slate Brush that wont draw anything, the "Null Brush".
 	 *
