UMG - Adding a blueprint function for canceling any drag drop operation.

diff --git a/Engine/Source/Runtime/UMG/Private/WidgetBlueprintLibrary.cpp b/Engine/Source/Runtime/UMG/Private/WidgetBlueprintLibrary.cpp
--- a/Engine/Source/Runtime/UMG/Private/WidgetBlueprintLibrary.cpp
+++ b/Engine/Source/Runtime/UMG/Private/WidgetBlueprintLibrary.cpp
@@ -301,6 +301,11 @@ UDragDropOperation* UWidgetBlueprintLibrary::GetDragDroppingContent()
 	return nullptr;
 }
 
+void UWidgetBlueprintLibrary::CancelDragDrop()
+{
+	FSlateApplication::Get().CancelDragDrop();
+}
+
 FSlateBrush UWidgetBlueprintLibrary::MakeBrushFromAsset(USlateBrushAsset* BrushAsset)
 {
 	if ( BrushAsset )
diff --git a/Engine/Source/Runtime/UMG/Public/Blueprint/WidgetBlueprintLibrary.h b/Engine/Source/Runtime/UMG/Public/Blueprint/WidgetBlueprintLibrary.h
--- a/Engine/Source/Runtime/UMG/Public/Blueprint/WidgetBlueprintLibrary.h
+++ b/Engine/Source/Runtime/UMG/Public/Blueprint/WidgetBlueprintLibrary.h
@@ -152,6 +152,12 @@ public:
 	UFUNCTION(BlueprintPure, BlueprintCosmetic, Category="Widget|Drag and Drop")
 	static UDragDropOperation* GetDragDroppingContent();
 
+	/**
+	 * Cancels any current drag drop operation.
+	 */
+	UFUNCTION(BlueprintCallable, Category="Widget|Drag and Drop")
+	static void CancelDragDrop();
+
 	/**
 	 * Creates a Slate Brush from a Slate Brush Asset
 	 *
