UMG - Adding a parabolic distortion option for widget components.

diff --git a/Engine/Source/Runtime/UMG/Private/Components/WidgetComponent.cpp b/Engine/Source/Runtime/UMG/Private/Components/WidgetComponent.cpp
--- a/Engine/Source/Runtime/UMG/Private/Components/WidgetComponent.cpp
+++ b/Engine/Source/Runtime/UMG/Private/Components/WidgetComponent.cpp
@@ -47,7 +47,7 @@ public:
 		, BodySetup( InComponent->GetBodySetup() )
 		, BlendMode( InComponent->GetBlendMode() )
 	{
-		bWillEverBeLit = false;	
+		bWillEverBeLit = false;
 	}
 
 	~FWidget3DSceneProxy()
@@ -198,7 +198,7 @@ public:
 								// Make sure the player is close enough to the widget to interact with it
 								if ( FVector::DistSquared(CachedHitResult.TraceStart, CachedHitResult.ImpactPoint) <= FMath::Square(WidgetComponent->GetMaxInteractionDistance()) )
 								{
-									return WidgetComponent->GetHitWidgetPath(CachedHitResult, bIgnoreEnabledStatus);
+									return WidgetComponent->GetHitWidgetPath(CachedHitResult.Location, bIgnoreEnabledStatus);
 								}
 							}
 						}
@@ -248,7 +248,7 @@ public:
 								TSharedPtr<FVirtualPointerPosition> VirtualCursorPos = MakeShareable(new FVirtualPointerPosition);
 
 								FVector2D LocalHitLocation;
-								WidgetComponent->GetLocalHitLocation(CachedHitResult, LocalHitLocation);
+								WidgetComponent->GetLocalHitLocation(CachedHitResult.Location, LocalHitLocation);
 
 								VirtualCursorPos->CurrentCursorPosition = LocalHitLocation;
 								VirtualCursorPos->LastCursorPosition = LocalHitLocation;
@@ -321,6 +321,7 @@ UWidgetComponent::UWidgetComponent( const FObjectInitializer& PCIP )
 	, BlendMode( EWidgetBlendMode::Masked )
 	, bIsOpaque_DEPRECATED( false )
 	, bIsTwoSided( false )
+	, ParabolaDistortion( 0 )
 	, TickWhenOffscreen( false )
 {
 	PrimaryComponentTick.bCanEverTick = true;
@@ -710,6 +711,7 @@ void UWidgetComponent::PostEditChangeProperty(FPropertyChangedEvent& PropertyCha
 		static FName IsOpaqueName("bIsOpaque");
 		static FName IsTwoSidedName("bIsTwoSided");
 		static FName BackgroundColorName("BackgroundColor");
+		static FName ParabolaDistortionName(TEXT("ParabolaDistortion"));
 
 		auto PropertyName = Property->GetFName();
 
@@ -724,7 +726,7 @@ void UWidgetComponent::PostEditChangeProperty(FPropertyChangedEvent& PropertyCha
 
 			MarkRenderStateDirty();
 		}
-		else if ( PropertyName == IsOpaqueName  || PropertyName == BackgroundColorName || PropertyName == IsTwoSidedName)
+		else if ( PropertyName == IsOpaqueName || PropertyName == BackgroundColorName || PropertyName == IsTwoSidedName || PropertyName == ParabolaDistortionName)
 		{
 			MarkRenderStateDirty();
 		}
@@ -815,6 +817,7 @@ void UWidgetComponent::UpdateWidget()
 
 void UWidgetComponent::UpdateRenderTarget()
 {
+	bool bRenderStateDirty = false;
 	bool bClearColorChanged = false;
 
 	FLinearColor ActualBackgroundColor = BackgroundColor;
@@ -831,12 +834,11 @@ void UWidgetComponent::UpdateRenderTarget()
 		RenderTarget = NewObject<UTextureRenderTarget2D>(this);
 		RenderTarget->ClearColor = ActualBackgroundColor;
 
-		bClearColorChanged = true;
+		bClearColorChanged = bRenderStateDirty = true;
 
 		RenderTarget->InitCustomFormat(DrawSize.X, DrawSize.Y, PF_B8G8R8A8, false);
 
 		MaterialInstance->SetTextureParameterValue( "SlateUI", RenderTarget );
-		MarkRenderStateDirty();
 	}
 	else if ( DrawSize != FIntPoint::ZeroValue )
 	{
@@ -844,14 +846,14 @@ void UWidgetComponent::UpdateRenderTarget()
 		if ( RenderTarget->SizeX != DrawSize.X || RenderTarget->SizeY != DrawSize.Y )
 		{
 			RenderTarget->InitCustomFormat( DrawSize.X, DrawSize.Y, PF_B8G8R8A8, false );
-			MarkRenderStateDirty();
+			bRenderStateDirty = true;
 		}
 
 		// Update the clear color
 		if ( RenderTarget->ClearColor != ActualBackgroundColor )
 		{
 			RenderTarget->ClearColor = ActualBackgroundColor;
-			bClearColorChanged = true;
+			bClearColorChanged = bRenderStateDirty = true;
 		}
 	}
 
@@ -859,6 +861,19 @@ void UWidgetComponent::UpdateRenderTarget()
 	if ( bClearColorChanged )
 	{
 		MaterialInstance->SetVectorParameterValue( "BackColor", RenderTarget->ClearColor );
+	}
+
+	static FName ParabolaDistortionName(TEXT("ParabolaDistortion"));
+
+	float CurrentParabolaValue;
+	if ( MaterialInstance->GetScalarParameterValue(ParabolaDistortionName, CurrentParabolaValue) && CurrentParabolaValue != ParabolaDistortion )
+	{
+		MaterialInstance->SetScalarParameterValue(ParabolaDistortionName, ParabolaDistortion);
+		bRenderStateDirty = true;
+	}
+
+	if ( bRenderStateDirty )
+	{
 		MarkRenderStateDirty();
 	}
 }
@@ -885,14 +900,20 @@ void UWidgetComponent::UpdateBodySetup( bool bDrawSizeChanged )
 	}	
 }
 
-void UWidgetComponent::GetLocalHitLocation(const FHitResult& HitResult, FVector2D& OutLocalHitLocation) const
+void UWidgetComponent::GetLocalHitLocation(FVector WorldHitLocation, FVector2D& OutLocalWidgetHitLocation) const
 {
 	// Find the hit location on the component
-	OutLocalHitLocation = FVector2D(ComponentToWorld.InverseTransformPosition(HitResult.Location));
+	OutLocalWidgetHitLocation = FVector2D(ComponentToWorld.InverseTransformPosition(WorldHitLocation));
 
 	// Offset the position by the pivot to get the position in widget space.
-	OutLocalHitLocation.X += DrawSize.X * Pivot.X;
-	OutLocalHitLocation.Y += DrawSize.Y * Pivot.Y;
+	OutLocalWidgetHitLocation.X += DrawSize.X * Pivot.X;
+	OutLocalWidgetHitLocation.Y += DrawSize.Y * Pivot.Y;
+
+	// Apply the parabola distortion
+	FVector2D NormalizedLocation = OutLocalWidgetHitLocation / DrawSize;
+	NormalizedLocation.Y += ParabolaDistortion * ( -2.0f * NormalizedLocation.Y + 1.0f ) * NormalizedLocation.X * ( NormalizedLocation.X - 1.0f );
+
+	OutLocalWidgetHitLocation.Y = DrawSize.Y * NormalizedLocation.Y;
 }
 
 UUserWidget* UWidgetComponent::GetUserWidgetObject() const
@@ -900,10 +921,10 @@ UUserWidget* UWidgetComponent::GetUserWidgetObject() const
 	return Widget;
 }
 
-TArray<FWidgetAndPointer> UWidgetComponent::GetHitWidgetPath(const FHitResult& HitResult, bool bIgnoreEnabledStatus, float CursorRadius )
+TArray<FWidgetAndPointer> UWidgetComponent::GetHitWidgetPath(FVector WorldHitLocation, bool bIgnoreEnabledStatus, float CursorRadius)
 {
 	FVector2D LocalHitLocation;
-	GetLocalHitLocation(HitResult, LocalHitLocation);
+	GetLocalHitLocation(WorldHitLocation, LocalHitLocation);
 
 	TSharedRef<FVirtualPointerPosition> VirtualMouseCoordinate = MakeShareable( new FVirtualPointerPosition );
 
@@ -952,6 +973,11 @@ void UWidgetComponent::SetMaxInteractionDistance(float Distance)
 	MaxInteractionDistance = Distance;
 }
 
+TSharedPtr< SWindow > UWidgetComponent::GetVirtualWindow() const
+{
+	return StaticCastSharedPtr<SWindow>(SlateWidget);
+}
+
 void UWidgetComponent::PostLoad()
 {
 	Super::PostLoad();
diff --git a/Engine/Source/Runtime/UMG/Public/Components/WidgetComponent.h b/Engine/Source/Runtime/UMG/Public/Components/WidgetComponent.h
--- a/Engine/Source/Runtime/UMG/Public/Components/WidgetComponent.h
+++ b/Engine/Source/Runtime/UMG/Public/Components/WidgetComponent.h
@@ -76,7 +76,7 @@ public:
 	 * @param HitResult The hit on this widget in the world
 	 * @param (Out) The transformed 2D hit location on the widget
 	 */
-	void GetLocalHitLocation(const FHitResult& HitResult, FVector2D& OutLocalHitLocation) const;
+	void GetLocalHitLocation(FVector WorldHitLocation, FVector2D& OutLocalHitLocation) const;
 
 	/** @return The class of the user widget displayed by this component */
 	TSubclassOf<UUserWidget> GetWidgetClass() const { return WidgetClass; }
@@ -86,7 +86,7 @@ public:
 	UUserWidget* GetUserWidgetObject() const;
 
 	/** @return List of widgets with their geometry and the cursor position transformed into this Widget component's space. */
-	TArray<FWidgetAndPointer> GetHitWidgetPath( const FHitResult& HitResult, bool bIgnoreEnabledStatus, float CursorRadius = 0.0f );
+	TArray<FWidgetAndPointer> GetHitWidgetPath(FVector WorldHitLocation, bool bIgnoreEnabledStatus, float CursorRadius = 0.0f);
 
 	/** @return The render target to which the user widget is rendered */
 	UTextureRenderTarget2D* GetRenderTarget() const { return RenderTarget; }
@@ -130,8 +130,11 @@ public:
 
 	/** @return The pivot point where the UI is rendered about the origin. */
 	FVector2D GetPivot() const { return Pivot; }
+
+	/** Get the fake window we create for widgets displayed in the world. */
+	TSharedPtr< SWindow > GetVirtualWindow() const;
 	
-private:
+protected:
 	/** The coordinate space in which to render the widget */
 	UPROPERTY(EditAnywhere, Category=UI)
 	EWidgetSpace Space;
@@ -177,17 +180,25 @@ private:
 	/** Is the component visible from behind? */
 	UPROPERTY(EditAnywhere, Category=Rendering)
 	bool bIsTwoSided;
+	
+	/**
+	 * When enabled, distorts the UI along a parabola shape giving the UI the appearance 
+	 * that it's on a curved surface in front of the users face.  This only works for UI 
+	 * rendered to a render target.
+	 */
+	UPROPERTY(EditAnywhere, Category=Rendering)
+	float ParabolaDistortion;
 
 	/** Should the component tick the widget when it's off screen? */
 	UPROPERTY(EditAnywhere, Category=Animation)
 	bool TickWhenOffscreen;
 
 	/** The User Widget object displayed and managed by this component */
-	UPROPERTY(transient, duplicatetransient)
+	UPROPERTY(Transient, DuplicateTransient)
 	UUserWidget* Widget;
 
 	/** The body setup of the displayed quad */
-	UPROPERTY(transient, duplicatetransient)
+	UPROPERTY(Transient, DuplicateTransient)
 	class UBodySetup* BodySetup;
 
 	/** The material instance for translucent widget components */
@@ -215,13 +226,15 @@ private:
 	UMaterialInterface* MaskedMaterial_OneSided;
 
 	/** The target to which the user widget is rendered */
-	UPROPERTY(transient, duplicatetransient)
+	UPROPERTY(Transient, DuplicateTransient)
 	UTextureRenderTarget2D* RenderTarget;
 
 	/** The dynamic instance of the material that the render target is attached to */
-	UPROPERTY(transient, duplicatetransient)
+	UPROPERTY(Transient, DuplicateTransient)
 	UMaterialInstanceDynamic* MaterialInstance;
 
+protected:
+
 	/** The grid used to find actual hit actual widgets once input has been translated to the components local space */
 	TSharedPtr<class FHittestGrid> HitTestGrid;
 
