UMG - UE-16247 - We now correctly call SyncProperties on all core widgets.  Additionally there's now an ensure to prevent future mistakes due to forgetting to call Super::SyncProperties().

diff --git a/Engine/Source/Runtime/UMG/Private/Components/ScrollBar.cpp b/Engine/Source/Runtime/UMG/Private/Components/ScrollBar.cpp
--- a/Engine/Source/Runtime/UMG/Private/Components/ScrollBar.cpp
+++ b/Engine/Source/Runtime/UMG/Private/Components/ScrollBar.cpp
@@ -42,6 +42,8 @@ TSharedRef<SWidget> UScrollBar::RebuildWidget()
 
 void UScrollBar::SynchronizeProperties()
 {
+	Super::SynchronizeProperties();
+
 	//MyScrollBar->SetScrollOffset(DesiredScrollOffset);
 }
 
diff --git a/Engine/Source/Runtime/UMG/Private/Components/ScrollBox.cpp b/Engine/Source/Runtime/UMG/Private/Components/ScrollBox.cpp
--- a/Engine/Source/Runtime/UMG/Private/Components/ScrollBox.cpp
+++ b/Engine/Source/Runtime/UMG/Private/Components/ScrollBox.cpp
@@ -83,6 +83,8 @@ TSharedRef<SWidget> UScrollBox::RebuildWidget()
 
 void UScrollBox::SynchronizeProperties()
 {
+	Super::SynchronizeProperties();
+
 	MyScrollBox->SetScrollOffset(DesiredScrollOffset);
 	MyScrollBox->SetOrientation(Orientation);
 	MyScrollBox->SetScrollBarVisibility(UWidget::ConvertSerializedVisibilityToRuntime(ScrollBarVisibility));
diff --git a/Engine/Source/Runtime/UMG/Private/Components/Widget.cpp b/Engine/Source/Runtime/UMG/Private/Components/Widget.cpp
--- a/Engine/Source/Runtime/UMG/Private/Components/Widget.cpp
+++ b/Engine/Source/Runtime/UMG/Private/Components/Widget.cpp
@@ -100,7 +100,9 @@ UWidget::UWidget(const FObjectInitializer& ObjectInitializer)
 {
 	bIsEnabled = true;
 	bIsVariable = true;
+#if WITH_EDITOR
 	DesignerFlags = EWidgetDesignFlags::None;
+#endif
 	Visiblity_DEPRECATED = Visibility = ESlateVisibility::Visible;	
 	RenderTransformPivot = FVector2D(0.5f, 0.5f);
 
@@ -471,9 +473,14 @@ TSharedRef<SWidget> UWidget::TakeWidget()
 
 			MyGCWidget = SafeGCWidget;
 
+#if !(UE_BUILD_SHIPPING || UE_BUILD_TEST)
+			bRoutedSynchronizeProperties = false;
+#endif
+
 			// Always synchronize properties of a user widget and call construct AFTER we've
 			// properly setup the GCWidget and synced all the properties.
 			SynchronizeProperties();
+			VerifySynchronizeProperties();
 			OnWidgetRebuilt();
 
 			return SafeGCWidget.ToSharedRef();
@@ -483,7 +490,12 @@ TSharedRef<SWidget> UWidget::TakeWidget()
 	{
 		if ( bNewlyCreated )
 		{
+#if !(UE_BUILD_SHIPPING || UE_BUILD_TEST)
+			bRoutedSynchronizeProperties = false;
+#endif
+
 			SynchronizeProperties();
+			VerifySynchronizeProperties();
 			OnWidgetRebuilt();
 		}
 
@@ -491,6 +503,13 @@ TSharedRef<SWidget> UWidget::TakeWidget()
 	}
 }
 
+#if !(UE_BUILD_SHIPPING || UE_BUILD_TEST)
+void UWidget::VerifySynchronizeProperties()
+{
+	ensureMsgf(bRoutedSynchronizeProperties, TEXT("%s failed to route SynchronizeProperties.  Please call Super::SynchronizeProperties() in your <className>::SynchronizeProperties() function."), *GetFullName());
+}
+#endif
+
 void UWidget::OnWidgetRebuilt()
 {
 }
@@ -685,6 +704,10 @@ TSharedRef<SWidget> UWidget::RebuildWidget()
 
 void UWidget::SynchronizeProperties()
 {
+#if !(UE_BUILD_SHIPPING || UE_BUILD_TEST)
+	bRoutedSynchronizeProperties = true;
+#endif
+
 	// We want to apply the bindings to the cached widget, which could be the SWidget, or the SObjectWidget, 
 	// in the case where it's a user widget.  We always want to prefer the SObjectWidget so that bindings to 
 	// visibility and enabled status are not stomping values setup in the root widget in the User Widget.
diff --git a/Engine/Source/Runtime/UMG/Public/Components/Widget.h b/Engine/Source/Runtime/UMG/Public/Components/Widget.h
--- a/Engine/Source/Runtime/UMG/Public/Components/Widget.h
+++ b/Engine/Source/Runtime/UMG/Public/Components/Widget.h
@@ -528,9 +528,24 @@ protected:
 	static TArray<TSubclassOf<UPropertyBinding>> BinderClasses;
 
 private:
+
+#if WITH_EDITORONLY_DATA
 	/** Any flags used by the designer at edit time. */
 	UPROPERTY(Transient)
 	TEnumAsByte<EWidgetDesignFlags::Type> DesignerFlags;
+#endif
+
+#if !(UE_BUILD_SHIPPING || UE_BUILD_TEST)
+
+	void VerifySynchronizeProperties();
+
+	/** Did we route the synchronize properties call? */
+	UPROPERTY(Transient)
+	bool bRoutedSynchronizeProperties;
+
+#else
+	FORCEINLINE void VerifySynchronizeProperties() { }
+#endif
 
 private:
 	GAME_SAFE_BINDING_IMPLEMENTATION(FText, ToolTipText)
