UMG - Making the UUserWidget initialize a WidgetTree if the user is doing some weird behavior where they've extended from the UUserWidget but now made a blueprint of it.

diff --git a/Engine/Source/Runtime/UMG/Private/UserWidget.cpp b/Engine/Source/Runtime/UMG/Private/UserWidget.cpp
--- a/Engine/Source/Runtime/UMG/Private/UserWidget.cpp
+++ b/Engine/Source/Runtime/UMG/Private/UserWidget.cpp
@@ -48,6 +48,13 @@ void UUserWidget::Initialize()
 			BGClass->InitializeWidget(this);
 		}
 
+		OnInitialize();
+
+		if ( WidgetTree == nullptr )
+		{
+			WidgetTree = NewObject<UWidgetTree>(this, TEXT("WidgetTree"));
+		}
+
 		// Map the named slot bindings to the available slots.
 		WidgetTree->ForEachWidget([&] (UWidget* Widget) {
 			if ( UNamedSlot* NamedWidet = Cast<UNamedSlot>(Widget) )
@@ -66,6 +73,11 @@ void UUserWidget::Initialize()
 	}
 }
 
+void UUserWidget::OnInitialize()
+{
+
+}
+
 void UUserWidget::BeginDestroy()
 {
 	Super::BeginDestroy();
diff --git a/Engine/Source/Runtime/UMG/Public/Blueprint/UserWidget.h b/Engine/Source/Runtime/UMG/Public/Blueprint/UserWidget.h
--- a/Engine/Source/Runtime/UMG/Public/Blueprint/UserWidget.h
+++ b/Engine/Source/Runtime/UMG/Public/Blueprint/UserWidget.h
@@ -838,6 +838,9 @@ public:
 #endif
 
 protected:
+	/** Called after the blueprint class has had a chance to initialize the class. */
+	virtual void OnInitialize();
+
 	/** Adds the widget to the screen, either to the viewport or to the player's screen depending on if the LocalPlayer is null. */
 	void AddToScreen(ULocalPlayer* LocalPlayer, int32 ZOrder);
 
