UMG - Exposing the Next/Previous widget navigation options on all widgets.

diff --git a/Engine/Source/Editor/UMGEditor/Private/Customizations/WidgetNavigationCustomization.cpp b/Engine/Source/Editor/UMGEditor/Private/Customizations/WidgetNavigationCustomization.cpp
--- a/Engine/Source/Editor/UMGEditor/Private/Customizations/WidgetNavigationCustomization.cpp
+++ b/Engine/Source/Editor/UMGEditor/Private/Customizations/WidgetNavigationCustomization.cpp
@@ -27,6 +27,8 @@ void FWidgetNavigationCustomization::CustomizeChildren(TSharedRef<IPropertyHandl
 	MakeNavRow(PropertyHandlePtr, ChildBuilder, EUINavigation::Right, LOCTEXT("RightNavigation", "Right"));
 	MakeNavRow(PropertyHandlePtr, ChildBuilder, EUINavigation::Up, LOCTEXT("UpNavigation", "Up"));
 	MakeNavRow(PropertyHandlePtr, ChildBuilder, EUINavigation::Down, LOCTEXT("DownNavigation", "Down"));
+	MakeNavRow(PropertyHandlePtr, ChildBuilder, EUINavigation::Next, LOCTEXT("NextNavigation", "Next"));
+	MakeNavRow(PropertyHandlePtr, ChildBuilder, EUINavigation::Previous, LOCTEXT("PreviousNavigation", "Previous"));
 }
 
 void FWidgetNavigationCustomization::FillOutChildren(TSharedRef<IPropertyHandle> PropertyHandle, IDetailChildrenBuilder& ChildBuilder, IPropertyTypeCustomizationUtils& CustomizationUtils)
@@ -308,6 +310,12 @@ void FWidgetNavigationCustomization::SetNav(UWidget* Widget, EUINavigation Nav,
 	case EUINavigation::Down:
 		DirectionNavigation = &WidgetNavigation->Down;
 		break;
+	case EUINavigation::Next:
+		DirectionNavigation = &WidgetNavigation->Next;
+		break;
+	case EUINavigation::Previous:
+		DirectionNavigation = &WidgetNavigation->Previous;
+		break;
 	default:
 		// Should not be possible.
 		check(false);
diff --git a/Engine/Source/Runtime/UMG/Private/UWidgetNavigation.cpp b/Engine/Source/Runtime/UMG/Private/UWidgetNavigation.cpp
--- a/Engine/Source/Runtime/UMG/Private/UWidgetNavigation.cpp
+++ b/Engine/Source/Runtime/UMG/Private/UWidgetNavigation.cpp
@@ -24,6 +24,10 @@ FWidgetNavigationData& UWidgetNavigation::GetNavigationData(EUINavigation Nav)
 		return Left;
 	case EUINavigation::Right:
 		return Right;
+	case EUINavigation::Next:
+		return Next;
+	case EUINavigation::Previous:
+		return Previous;
 	default:
 		break;
 	}
@@ -47,7 +51,9 @@ EUINavigationRule UWidgetNavigation::GetNavigationRule(EUINavigation Nav)
 	case EUINavigation::Right:
 		return Right.Rule;
 	case EUINavigation::Next:
+		return Next.Rule;
 	case EUINavigation::Previous:
+		return Previous.Rule;
 		break;
 	}
 	return EUINavigationRule::Escape;
@@ -73,6 +79,14 @@ void UWidgetNavigation::ResolveExplictRules(UWidgetTree* WidgetTree)
 	{
 		Right.Widget = WidgetTree->FindWidget(Right.WidgetToFocus);
 	}
+	if ( Next.Rule == EUINavigationRule::Explicit )
+	{
+		Next.Widget = WidgetTree->FindWidget(Next.WidgetToFocus);
+	}
+	if ( Previous.Rule == EUINavigationRule::Explicit )
+	{
+		Previous.Widget = WidgetTree->FindWidget(Previous.WidgetToFocus);
+	}
 }
 
 void UWidgetNavigation::UpdateMetaData(TSharedRef<FNavigationMetaData> MetaData)
@@ -81,6 +95,8 @@ void UWidgetNavigation::UpdateMetaData(TSharedRef<FNavigationMetaData> MetaData)
 	UpdateMetaDataEntry(MetaData, Down, EUINavigation::Down);
 	UpdateMetaDataEntry(MetaData, Left, EUINavigation::Left);
 	UpdateMetaDataEntry(MetaData, Right, EUINavigation::Right);
+	UpdateMetaDataEntry(MetaData, Next, EUINavigation::Next);
+	UpdateMetaDataEntry(MetaData, Previous, EUINavigation::Previous);
 }
 
 bool UWidgetNavigation::IsDefault() const
@@ -88,7 +104,9 @@ bool UWidgetNavigation::IsDefault() const
 	return Up.Rule == EUINavigationRule::Escape &&
 		Down.Rule == EUINavigationRule::Escape &&
 		Left.Rule == EUINavigationRule::Escape &&
-		Right.Rule == EUINavigationRule::Escape;
+		Right.Rule == EUINavigationRule::Escape &&
+		Next.Rule == EUINavigationRule::Escape &&
+		Previous.Rule == EUINavigationRule::Escape;
 }
 
 void UWidgetNavigation::UpdateMetaDataEntry(TSharedRef<FNavigationMetaData> MetaData, const FWidgetNavigationData & NavData, EUINavigation Nav)
diff --git a/Engine/Source/Runtime/UMG/Public/Blueprint/WidgetNavigation.h b/Engine/Source/Runtime/UMG/Public/Blueprint/WidgetNavigation.h
--- a/Engine/Source/Runtime/UMG/Public/Blueprint/WidgetNavigation.h
+++ b/Engine/Source/Runtime/UMG/Public/Blueprint/WidgetNavigation.h
@@ -34,22 +34,30 @@ class UMG_API UWidgetNavigation : public UObject
 	GENERATED_UCLASS_BODY()
 	
 public:
-	/**  */
+	/** Happens when the user presses up arrow, joystick, d-pad. */
 	UPROPERTY(EditAnywhere, BlueprintReadOnly, Category="Navigation")
 	FWidgetNavigationData Up;
 
-	/**  */
+	/** Happens when the user presses down arrow, joystick, d-pad. */
 	UPROPERTY(EditAnywhere, BlueprintReadOnly, Category="Navigation")
 	FWidgetNavigationData Down;
 
-	/**  */
+	/** Happens when the user presses left arrow, joystick, d-pad. */
 	UPROPERTY(EditAnywhere, BlueprintReadOnly, Category="Navigation")
 	FWidgetNavigationData Left;
 
-	/**  */
+	/** Happens when the user presses right arrow, joystick, d-pad. */
 	UPROPERTY(EditAnywhere, BlueprintReadOnly, Category="Navigation")
 	FWidgetNavigationData Right;
 
+	/** Happens when the user presses Tab. */
+	UPROPERTY(EditAnywhere, BlueprintReadOnly, Category="Navigation")
+	FWidgetNavigationData Next;
+
+	/** Happens when the user presses Shift+Tab. */
+	UPROPERTY(EditAnywhere, BlueprintReadOnly, Category="Navigation")
+	FWidgetNavigationData Previous;
+
 public:
 
 #if WITH_EDITOR
@@ -62,6 +70,7 @@ public:
 
 #endif
 
+	/** Resolve widget names */
 	void ResolveExplictRules(class UWidgetTree* WidgetTree);
 
 	/** Updates a slate metadata object to match this configured navigation ruleset. */
