UMG - Removeing OnInitialize from UUserWidget.  May reintroduce it after I've given some more thought where I want to allow users to muck with init/construct.

diff --git a/Engine/Source/Runtime/UMG/Private/UserWidget.cpp b/Engine/Source/Runtime/UMG/Private/UserWidget.cpp
--- a/Engine/Source/Runtime/UMG/Private/UserWidget.cpp
+++ b/Engine/Source/Runtime/UMG/Private/UserWidget.cpp
@@ -48,8 +48,6 @@ void UUserWidget::Initialize()
 			BGClass->InitializeWidget(this);
 		}
 
-		OnInitialize();
-
 		if ( WidgetTree == nullptr )
 		{
 			WidgetTree = NewObject<UWidgetTree>(this, TEXT("WidgetTree"));
@@ -73,11 +71,6 @@ void UUserWidget::Initialize()
 	}
 }
 
-void UUserWidget::OnInitialize()
-{
-
-}
-
 void UUserWidget::BeginDestroy()
 {
 	Super::BeginDestroy();
diff --git a/Engine/Source/Runtime/UMG/Public/Blueprint/UserWidget.h b/Engine/Source/Runtime/UMG/Public/Blueprint/UserWidget.h
--- a/Engine/Source/Runtime/UMG/Public/Blueprint/UserWidget.h
+++ b/Engine/Source/Runtime/UMG/Public/Blueprint/UserWidget.h
@@ -838,8 +838,6 @@ public:
 #endif
 
 protected:
-	/** Called after the blueprint class has had a chance to initialize the class. */
-	virtual void OnInitialize();
 
 	/** Adds the widget to the screen, either to the viewport or to the player's screen depending on if the LocalPlayer is null. */
 	void AddToScreen(ULocalPlayer* LocalPlayer, int32 ZOrder);
