UMG  - Removing the DisallowCreateNew from the widget class field from the widget component, that should now create the correct widget blueprint.

diff --git a/Engine/Source/Runtime/UMG/Public/Components/WidgetComponent.h b/Engine/Source/Runtime/UMG/Public/Components/WidgetComponent.h
--- a/Engine/Source/Runtime/UMG/Public/Components/WidgetComponent.h
+++ b/Engine/Source/Runtime/UMG/Public/Components/WidgetComponent.h
@@ -140,11 +140,11 @@ protected:
 	EWidgetSpace Space;
 
 	/** The class of User Widget to create and display an instance of */
-	UPROPERTY(EditAnywhere, Category=UI, meta=(DisallowCreateNew))
+	UPROPERTY(EditAnywhere, Category=UI)
 	TSubclassOf<UUserWidget> WidgetClass;
 	
 	/** The size of the displayed quad. */
-	UPROPERTY(EditAnywhere, Category="UI")
+	UPROPERTY(EditAnywhere, Category=UI)
 	FIntPoint DrawSize;
 
 	/** The Alignment/Pivot point that the widget is placed at relative to the position. */
