[UE-16400] We were allowing imports to reference a stale function that was being removed by CleanAndSanitizeClass

diff --git a/Engine/Source/Runtime/CoreUObject/Private/UObject/LinkerLoad.cpp b/Engine/Source/Runtime/CoreUObject/Private/UObject/LinkerLoad.cpp
--- a/Engine/Source/Runtime/CoreUObject/Private/UObject/LinkerLoad.cpp
+++ b/Engine/Source/Runtime/CoreUObject/Private/UObject/LinkerLoad.cpp
@@ -3757,7 +3757,7 @@ UObject* FLinkerLoad::CreateExport( int32 Index )
 			}
 		}
 	}
-	return Export.Object;
+	return Export.bExportLoadFailed ? nullptr : Export.Object;
 }
 
 bool FLinkerLoad::IsImportNative(const int32 Index) const
diff --git a/Engine/Source/Runtime/UMG/Private/Components/PanelWidget.cpp b/Engine/Source/Runtime/UMG/Private/Components/PanelWidget.cpp
--- a/Engine/Source/Runtime/UMG/Private/Components/PanelWidget.cpp
+++ b/Engine/Source/Runtime/UMG/Private/Components/PanelWidget.cpp
@@ -225,7 +225,7 @@ void UPanelWidget::PostLoad()
 	for ( int32 SlotIndex = 0; SlotIndex < Slots.Num(); SlotIndex++ )
 	{
 		// Remove any slots where their content is null, we don't support content-less slots.
-		if ( Slots[SlotIndex]->Content == nullptr )
+		if (!Slots[SlotIndex] || Slots[SlotIndex]->Content == nullptr)
 		{
 			Slots.RemoveAt(SlotIndex);
 			SlotIndex--;
