Custom (no-actor, no-component) objects created by [REDACTED_PERSON] function have RF_StrongRefOnFrame flag set.

diff --git a/Engine/Source/Editor/BlueprintGraph/Classes/K2Node_BaseAsyncTask.h b/Engine/Source/Editor/BlueprintGraph/Classes/K2Node_BaseAsyncTask.h
--- a/Engine/Source/Editor/BlueprintGraph/Classes/K2Node_BaseAsyncTask.h
+++ b/Engine/Source/Editor/BlueprintGraph/Classes/K2Node_BaseAsyncTask.h
@@ -10,6 +10,8 @@ class UK2Node_CustomEvent;
 class UEdGraphPin;
 class UK2Node_TemporaryVariable;
 
+/** !!! The proxy object should have RF_StrongRefOnFrame flag. !!! */
+
 UCLASS(Abstract)
 class BLUEPRINTGRAPH_API UK2Node_BaseAsyncTask : public UK2Node
 {
diff --git a/Engine/Source/Editor/Kismet/Public/Nodes/K2Node_AsyncAction.h b/Engine/Source/Editor/Kismet/Public/Nodes/K2Node_AsyncAction.h
--- a/Engine/Source/Editor/Kismet/Public/Nodes/K2Node_AsyncAction.h
+++ b/Engine/Source/Editor/Kismet/Public/Nodes/K2Node_AsyncAction.h
@@ -4,6 +4,8 @@
 
 #include "K2Node_AsyncAction.generated.h"
 
+/** !!! The proxy object should have RF_StrongRefOnFrame flag. !!! */
+
 UCLASS()
 class KISMET_API UK2Node_AsyncAction : public UK2Node_BaseAsyncTask
 {
diff --git a/Engine/Source/Runtime/AIModule/Private/Blueprint/AIBlueprintHelperLibrary.cpp b/Engine/Source/Runtime/AIModule/Private/Blueprint/AIBlueprintHelperLibrary.cpp
--- a/Engine/Source/Runtime/AIModule/Private/Blueprint/AIBlueprintHelperLibrary.cpp
+++ b/Engine/Source/Runtime/AIModule/Private/Blueprint/AIBlueprintHelperLibrary.cpp
@@ -18,6 +18,7 @@ UAIAsyncTaskBlueprintProxy::UAIAsyncTaskBlueprintProxy(const FObjectInitializer&
 	MyWorld = Cast<UWorld>(GetOuter());
 	if (HasAnyFlags(RF_ClassDefaultObject) == false)
 	{
+		SetFlags(RF_StrongRefOnFrame);
 		UAISystem* const AISystem = UAISystem::GetCurrentSafe(MyWorld.Get());
 		if (AISystem)
 		{
diff --git a/Engine/Source/Runtime/Engine/Private/Kismet/BlueprintAsyncActionBase.cpp b/Engine/Source/Runtime/Engine/Private/Kismet/BlueprintAsyncActionBase.cpp
--- a/Engine/Source/Runtime/Engine/Private/Kismet/BlueprintAsyncActionBase.cpp
+++ b/Engine/Source/Runtime/Engine/Private/Kismet/BlueprintAsyncActionBase.cpp
@@ -9,6 +9,7 @@
 UBlueprintAsyncActionBase::UBlueprintAsyncActionBase(const FObjectInitializer& ObjectInitializer)
 	: Super(ObjectInitializer)
 {
+	SetFlags(RF_StrongRefOnFrame);
 }
 
 void UBlueprintAsyncActionBase::Activate()
diff --git a/Engine/Source/Runtime/GameplayTasks/Private/GameplayTask.cpp b/Engine/Source/Runtime/GameplayTasks/Private/GameplayTask.cpp
--- a/Engine/Source/Runtime/GameplayTasks/Private/GameplayTask.cpp
+++ b/Engine/Source/Runtime/GameplayTasks/Private/GameplayTask.cpp
@@ -12,6 +12,7 @@ UGameplayTask::UGameplayTask(const FObjectInitializer& ObjectInitializer)
 	bIsSimulating = false;
 	bOwnedByTasksComponent = false;
 	TaskState = EGameplayTaskState::Uninitialized;
+	SetFlags(RF_StrongRefOnFrame);
 }
 
 IGameplayTaskOwnerInterface* UGameplayTask::ConvertToTaskOwner(UObject& OwnerObject)
diff --git a/Engine/Source/Runtime/Online/OnlineSubsystemUtils/Classes/LeaderboardFlushCallbackProxy.cpp b/Engine/Source/Runtime/Online/OnlineSubsystemUtils/Classes/LeaderboardFlushCallbackProxy.cpp
--- a/Engine/Source/Runtime/Online/OnlineSubsystemUtils/Classes/LeaderboardFlushCallbackProxy.cpp
+++ b/Engine/Source/Runtime/Online/OnlineSubsystemUtils/Classes/LeaderboardFlushCallbackProxy.cpp
@@ -97,6 +97,7 @@ void ULeaderboardFlushCallbackProxy::BeginDestroy()
 ULeaderboardFlushCallbackProxy* ULeaderboardFlushCallbackProxy::CreateProxyObjectForFlush(class APlayerController* PlayerController, FName SessionName)
 {
 	ULeaderboardFlushCallbackProxy* Proxy = NewObject<ULeaderboardFlushCallbackProxy>();
+	Proxy->SetFlags(RF_StrongRefOnFrame);
 	Proxy->TriggerFlush(PlayerController, SessionName);
 	return Proxy;
 }
diff --git a/Engine/Source/Runtime/Online/OnlineSubsystemUtils/Classes/LeaderboardQueryCallbackProxy.cpp b/Engine/Source/Runtime/Online/OnlineSubsystemUtils/Classes/LeaderboardQueryCallbackProxy.cpp
--- a/Engine/Source/Runtime/Online/OnlineSubsystemUtils/Classes/LeaderboardQueryCallbackProxy.cpp
+++ b/Engine/Source/Runtime/Online/OnlineSubsystemUtils/Classes/LeaderboardQueryCallbackProxy.cpp
@@ -151,6 +151,7 @@ void ULeaderboardQueryCallbackProxy::BeginDestroy()
 ULeaderboardQueryCallbackProxy* ULeaderboardQueryCallbackProxy::CreateProxyObjectForIntQuery(class APlayerController* PlayerController, FName StatName)
 {
 	ULeaderboardQueryCallbackProxy* Proxy = NewObject<ULeaderboardQueryCallbackProxy>();
+	Proxy->SetFlags(RF_StrongRefOnFrame);
 	Proxy->TriggerQuery(PlayerController, StatName, EOnlineKeyValuePairDataType::Int32);
 	return Proxy;
 }
diff --git a/Engine/Source/Runtime/Online/OnlineSubsystemUtils/Private/InAppPurchaseCallbackProxy.cpp b/Engine/Source/Runtime/Online/OnlineSubsystemUtils/Private/InAppPurchaseCallbackProxy.cpp
--- a/Engine/Source/Runtime/Online/OnlineSubsystemUtils/Private/InAppPurchaseCallbackProxy.cpp
+++ b/Engine/Source/Runtime/Online/OnlineSubsystemUtils/Private/InAppPurchaseCallbackProxy.cpp
@@ -127,6 +127,7 @@ void UInAppPurchaseCallbackProxy::BeginDestroy()
 UInAppPurchaseCallbackProxy* UInAppPurchaseCallbackProxy::CreateProxyObjectForInAppPurchase(class APlayerController* PlayerController, const FInAppPurchaseProductRequest& ProductRequest)
 {
 	UInAppPurchaseCallbackProxy* Proxy = NewObject<UInAppPurchaseCallbackProxy>();
+	Proxy->SetFlags(RF_StrongRefOnFrame);
 	Proxy->Trigger(PlayerController, ProductRequest);
 	return Proxy;
 }
diff --git a/Engine/Source/Runtime/Online/OnlineSubsystemUtils/Private/InAppPurchaseQueryCallbackProxy.cpp b/Engine/Source/Runtime/Online/OnlineSubsystemUtils/Private/InAppPurchaseQueryCallbackProxy.cpp
--- a/Engine/Source/Runtime/Online/OnlineSubsystemUtils/Private/InAppPurchaseQueryCallbackProxy.cpp
+++ b/Engine/Source/Runtime/Online/OnlineSubsystemUtils/Private/InAppPurchaseQueryCallbackProxy.cpp
@@ -116,6 +116,7 @@ void UInAppPurchaseQueryCallbackProxy::BeginDestroy()
 UInAppPurchaseQueryCallbackProxy* UInAppPurchaseQueryCallbackProxy::CreateProxyObjectForInAppPurchaseQuery(class APlayerController* PlayerController, const TArray<FString>& ProductIdentifiers)
 {
 	UInAppPurchaseQueryCallbackProxy* Proxy = NewObject<UInAppPurchaseQueryCallbackProxy>();
+	Proxy->SetFlags(RF_StrongRefOnFrame);
 	Proxy->TriggerQuery(PlayerController, ProductIdentifiers);
 	return Proxy;
 }
diff --git a/Engine/Source/Runtime/Online/OnlineSubsystemUtils/Private/InAppPurchaseRestoreCallbackProxy.cpp b/Engine/Source/Runtime/Online/OnlineSubsystemUtils/Private/InAppPurchaseRestoreCallbackProxy.cpp
--- a/Engine/Source/Runtime/Online/OnlineSubsystemUtils/Private/InAppPurchaseRestoreCallbackProxy.cpp
+++ b/Engine/Source/Runtime/Online/OnlineSubsystemUtils/Private/InAppPurchaseRestoreCallbackProxy.cpp
@@ -118,6 +118,7 @@ void UInAppPurchaseRestoreCallbackProxy::BeginDestroy()
 UInAppPurchaseRestoreCallbackProxy* UInAppPurchaseRestoreCallbackProxy::CreateProxyObjectForInAppPurchaseRestore(class APlayerController* PlayerController)
 {
 	UInAppPurchaseRestoreCallbackProxy* Proxy = NewObject<UInAppPurchaseRestoreCallbackProxy>();
+	Proxy->SetFlags(RF_StrongRefOnFrame);
 	Proxy->Trigger(PlayerController);
 	return Proxy;
 }
diff --git a/Engine/Source/Runtime/Online/OnlineSubsystemUtils/Private/OnlineBlueprintCallProxyBase.cpp b/Engine/Source/Runtime/Online/OnlineSubsystemUtils/Private/OnlineBlueprintCallProxyBase.cpp
--- a/Engine/Source/Runtime/Online/OnlineSubsystemUtils/Private/OnlineBlueprintCallProxyBase.cpp
+++ b/Engine/Source/Runtime/Online/OnlineSubsystemUtils/Private/OnlineBlueprintCallProxyBase.cpp
@@ -9,6 +9,7 @@
 UOnlineBlueprintCallProxyBase::UOnlineBlueprintCallProxyBase(const FObjectInitializer& ObjectInitializer)
 	: Super(ObjectInitializer)
 {
+	SetFlags(RF_StrongRefOnFrame);
 }
 
 void UOnlineBlueprintCallProxyBase::Activate()
diff --git a/Engine/Source/Runtime/UMG/Private/WidgetBlueprintLibrary.cpp b/Engine/Source/Runtime/UMG/Private/WidgetBlueprintLibrary.cpp
--- a/Engine/Source/Runtime/UMG/Private/WidgetBlueprintLibrary.cpp
+++ b/Engine/Source/Runtime/UMG/Private/WidgetBlueprintLibrary.cpp
@@ -24,27 +24,41 @@ UUserWidget* UWidgetBlueprintLibrary::Create(UObject* WorldContextObject, TSubcl
 		return nullptr;
 	}
 
+	UUserWidget* UserWidget = nullptr;
 	if ( OwningPlayer == nullptr )
 	{
 		UWorld* World = GEngine->GetWorldFromContextObject(WorldContextObject);
-		return CreateWidget<UUserWidget>(World, WidgetType);
+		UserWidget = CreateWidget<UUserWidget>(World, WidgetType);
 	}
 	else
 	{
-		return CreateWidget<UUserWidget>(OwningPlayer, WidgetType);
+		UserWidget = CreateWidget<UUserWidget>(OwningPlayer, WidgetType);
 	}
+
+	if (UserWidget)
+	{
+		UserWidget->SetFlags(RF_StrongRefOnFrame);
+	}
+	return UserWidget;
 }
 
 UDragDropOperation* UWidgetBlueprintLibrary::CreateDragDropOperation(TSubclassOf<UDragDropOperation> Operation)
 {
+	UDragDropOperation* DragDropOperation = nullptr;
 	if ( Operation )
 	{
-		return NewObject<UDragDropOperation>(GetTransientPackage(), Operation);
+		DragDropOperation = NewObject<UDragDropOperation>(GetTransientPackage(), Operation);
 	}
 	else
 	{
-		return NewObject<UDragDropOperation>();
+		DragDropOperation = NewObject<UDragDropOperation>();
+	}
+
+	if (DragDropOperation)
+	{
+		DragDropOperation->SetFlags(RF_StrongRefOnFrame);
 	}
+	return DragDropOperation;
 }
 
 void UWidgetBlueprintLibrary::SetInputMode_UIOnly(APlayerController* Target, UWidget* InWidgetToFocus, bool bLockMouseToViewport)
