UMG - SConstraintCanvas no longer attempts to draw elements outside its view.  Making a lot fewer calls to the enabled delegate when painting, prefiltering arranging children that are not visible and avoiding the math for them, also now consistantly drawing children in the correct order.  Previously we only sorted based on ZOrder, but it appears the arrays don't nessesarily end in logical order mutated by ZOrder, you in fact need to check if zorder is equal, then sort based on child index, otherwise the order doesn't nessesarily remain the same it began in as one might expect.

diff --git a/Engine/Source/Runtime/UMG/Private/Slate/SConstraintCanvas.cpp b/Engine/Source/Runtime/UMG/Private/Slate/SConstraintCanvas.cpp
--- a/Engine/Source/Runtime/UMG/Private/Slate/SConstraintCanvas.cpp
+++ b/Engine/Source/Runtime/UMG/Private/Slate/SConstraintCanvas.cpp
@@ -51,7 +51,7 @@ struct FSortSlotsByZOrder
 {
 	FORCEINLINE bool operator()(const FChildZOrder& A, const FChildZOrder& B) const
 	{
-		return A.ZOrder < B.ZOrder;
+		return A.ZOrder == B.ZOrder ? A.ChildIndex < B.ChildIndex : A.ZOrder < B.ZOrder;
 	}
 };
 
@@ -83,66 +83,71 @@ void SConstraintCanvas::OnArrangeChildren( const FGeometry& AllottedGeometry, FA
 		{
 			const FChildZOrder& CurSlot = SlotOrder[ChildIndex];
 			const SConstraintCanvas::FSlot& CurChild = Children[CurSlot.ChildIndex];
+			const TSharedRef<SWidget>& CurWidget = CurChild.GetWidget();
 
-			const FMargin Offset = CurChild.OffsetAttr.Get();
-			const FVector2D Alignment = CurChild.AlignmentAttr.Get();
-			const FAnchors Anchors = CurChild.AnchorsAttr.Get();
-
-			const bool AutoSize = CurChild.AutoSizeAttr.Get();
-
-			const FMargin AnchorPixels =
-				FMargin(Anchors.Minimum.X * AllottedGeometry.Size.X,
-						Anchors.Minimum.Y * AllottedGeometry.Size.Y,
-						Anchors.Maximum.X * AllottedGeometry.Size.X,
-						Anchors.Maximum.Y * AllottedGeometry.Size.Y);
-
-			const bool bIsHorizontalStretch = Anchors.Minimum.X != Anchors.Maximum.X;
-			const bool bIsVerticalStretch = Anchors.Minimum.Y != Anchors.Maximum.Y;
-			
-			const FVector2D SlotSize = FVector2D(Offset.Right, Offset.Bottom);
-			const FVector2D WidgetDesiredSize = CurChild.GetWidget()->GetDesiredSize();
-
-			const FVector2D Size = AutoSize ? WidgetDesiredSize : SlotSize;
-			
-			// Calculate the offset based on the pivot position.
-			FVector2D AlignmentOffset = Size * Alignment;
-
-			// Calculate the local position based on the anchor and position offset.
-			FVector2D LocalPosition, LocalSize;
-
-			// Calculate the position and size based on the horizontal stretch or non-stretch
-			if ( bIsHorizontalStretch )
+			const EVisibility ChildVisibility = CurWidget->GetVisibility();
+			if ( ArrangedChildren.Accepts(ChildVisibility) )
 			{
-				LocalPosition.X = AnchorPixels.Left + Offset.Left;
-				LocalSize.X = AnchorPixels.Right - LocalPosition.X - Offset.Right;
+				const FMargin Offset = CurChild.OffsetAttr.Get();
+				const FVector2D Alignment = CurChild.AlignmentAttr.Get();
+				const FAnchors Anchors = CurChild.AnchorsAttr.Get();
+
+				const bool AutoSize = CurChild.AutoSizeAttr.Get();
+
+				const FMargin AnchorPixels =
+					FMargin(Anchors.Minimum.X * AllottedGeometry.Size.X,
+					Anchors.Minimum.Y * AllottedGeometry.Size.Y,
+					Anchors.Maximum.X * AllottedGeometry.Size.X,
+					Anchors.Maximum.Y * AllottedGeometry.Size.Y);
+
+				const bool bIsHorizontalStretch = Anchors.Minimum.X != Anchors.Maximum.X;
+				const bool bIsVerticalStretch = Anchors.Minimum.Y != Anchors.Maximum.Y;
+
+				const FVector2D SlotSize = FVector2D(Offset.Right, Offset.Bottom);
+				const FVector2D WidgetDesiredSize = CurWidget->GetDesiredSize();
+
+				const FVector2D Size = AutoSize ? WidgetDesiredSize : SlotSize;
+
+				// Calculate the offset based on the pivot position.
+				FVector2D AlignmentOffset = Size * Alignment;
+
+				// Calculate the local position based on the anchor and position offset.
+				FVector2D LocalPosition, LocalSize;
+
+				// Calculate the position and size based on the horizontal stretch or non-stretch
+				if ( bIsHorizontalStretch )
+				{
+					LocalPosition.X = AnchorPixels.Left + Offset.Left;
+					LocalSize.X = AnchorPixels.Right - LocalPosition.X - Offset.Right;
+				}
+				else
+				{
+					LocalPosition.X = AnchorPixels.Left + Offset.Left - AlignmentOffset.X;
+					LocalSize.X = Size.X;
+				}
+
+				// Calculate the position and size based on the vertical stretch or non-stretch
+				if ( bIsVerticalStretch )
+				{
+					LocalPosition.Y = AnchorPixels.Top + Offset.Top;
+					LocalSize.Y = AnchorPixels.Bottom - LocalPosition.Y - Offset.Bottom;
+				}
+				else
+				{
+					LocalPosition.Y = AnchorPixels.Top + Offset.Top - AlignmentOffset.Y;
+					LocalSize.Y = Size.Y;
+				}
+
+				// Add the information about this child to the output list (ArrangedChildren)
+				ArrangedChildren.AddWidget(ChildVisibility, AllottedGeometry.MakeChild(
+					// The child widget being arranged
+					CurWidget,
+					// Child's local position (i.e. position within parent)
+					LocalPosition,
+					// Child's size
+					LocalSize
+					));
 			}
-			else
-			{
-				LocalPosition.X = AnchorPixels.Left + Offset.Left - AlignmentOffset.X;
-				LocalSize.X = Size.X;
-			}
-
-			// Calculate the position and size based on the vertical stretch or non-stretch
-			if ( bIsVerticalStretch )
-			{
-				LocalPosition.Y = AnchorPixels.Top + Offset.Top;
-				LocalSize.Y = AnchorPixels.Bottom - LocalPosition.Y - Offset.Bottom;
-			}
-			else
-			{
-				LocalPosition.Y = AnchorPixels.Top + Offset.Top - AlignmentOffset.Y;
-				LocalSize.Y = Size.Y;
-			}
-			
-			// Add the information about this child to the output list (ArrangedChildren)
-			ArrangedChildren.AddWidget( AllottedGeometry.MakeChild(
-				// The child widget being arranged
-				CurChild.GetWidget(),
-				// Child's local position (i.e. position within parent)
-				LocalPosition,
-				// Child's size
-				LocalSize
-			));
 		}
 	}
 }
@@ -150,7 +155,9 @@ void SConstraintCanvas::OnArrangeChildren( const FGeometry& AllottedGeometry, FA
 int32 SConstraintCanvas::OnPaint( const FPaintArgs& Args, const FGeometry& AllottedGeometry, const FSlateRect& MyClippingRect, FSlateWindowElementList& OutDrawElements, int32 LayerId, const FWidgetStyle& InWidgetStyle, bool bParentEnabled ) const
 {
 	FArrangedChildren ArrangedChildren(EVisibility::Visible);
-	this->ArrangeChildren(AllottedGeometry, ArrangedChildren);
+	ArrangeChildren(AllottedGeometry, ArrangedChildren);
+
+	bool bForwardedEnabled = ShouldBeEnabled(bParentEnabled);
 
 	// Because we paint multiple children, we must track the maximum layer id that they produced in case one of our parents
 	// wants to an overlay for all of its contents.
@@ -160,9 +167,13 @@ int32 SConstraintCanvas::OnPaint( const FPaintArgs& Args, const FGeometry& Allot
 	{
 		FArrangedWidget& CurWidget = ArrangedChildren[ChildIndex];
 		FSlateRect ChildClipRect = MyClippingRect.IntersectionWith(CurWidget.Geometry.GetClippingRect());
-		const int32 CurWidgetsMaxLayerId = CurWidget.Widget->Paint( Args.WithNewParent(this), CurWidget.Geometry, ChildClipRect, OutDrawElements, MaxLayerId + 1, InWidgetStyle, ShouldBeEnabled(bParentEnabled));
 
-		MaxLayerId = FMath::Max(MaxLayerId, CurWidgetsMaxLayerId);
+		if ( ChildClipRect.GetSize().IsZero() == false )
+		{
+			const int32 CurWidgetsMaxLayerId = CurWidget.Widget->Paint(Args.WithNewParent(this), CurWidget.Geometry, ChildClipRect, OutDrawElements, MaxLayerId + 1, InWidgetStyle, bForwardedEnabled);
+
+			MaxLayerId = FMath::Max(MaxLayerId, CurWidgetsMaxLayerId);
+		}
 	}
 
 	return MaxLayerId;
