UMG - Moving the code for generating the slate render transform into a utility function on the widget transform class used by UMG.

diff --git a/Engine/Source/Runtime/UMG/Private/Components/Widget.cpp b/Engine/Source/Runtime/UMG/Private/Components/Widget.cpp
--- a/Engine/Source/Runtime/UMG/Private/Components/Widget.cpp
+++ b/Engine/Source/Runtime/UMG/Private/Components/Widget.cpp
@@ -151,14 +151,7 @@ void UWidget::UpdateRenderTransform()
 		}
 		else
 		{
-			FSlateRenderTransform Transform2D =
-				::Concatenate(
-				FScale2D(RenderTransform.Scale),
-				FShear2D::FromShearAngles(RenderTransform.Shear),
-				FQuat2D(FMath::DegreesToRadians(RenderTransform.Angle)),
-				FVector2D(RenderTransform.Translation));
-
-			SafeWidget->SetRenderTransform(Transform2D);
+			SafeWidget->SetRenderTransform(RenderTransform.ToSlateRenderTransform());
 		}
 	}
 }
diff --git a/Engine/Source/Runtime/UMG/Public/Slate/WidgetTransform.h b/Engine/Source/Runtime/UMG/Public/Slate/WidgetTransform.h
--- a/Engine/Source/Runtime/UMG/Public/Slate/WidgetTransform.h
+++ b/Engine/Source/Runtime/UMG/Public/Slate/WidgetTransform.h
@@ -12,8 +12,10 @@ struct FWidgetTransform
 {
 	GENERATED_USTRUCT_BODY()
 
+public:
+
 	/** The amount to translate the widget in slate units */
-	UPROPERTY(EditAnywhere, BlueprintReadWrite, Category=Transform, meta=(Delta = "1"))
+	UPROPERTY(EditAnywhere, BlueprintReadWrite, Category=Transform, meta=( Delta = "1" ))
 	FVector2D Translation;
 
 	/** The scale to apply to the widget */
@@ -23,11 +25,12 @@ struct FWidgetTransform
 	/** The amount to shear the widget in slate units */
 	UPROPERTY(EditAnywhere, BlueprintReadWrite, Category=Transform, meta=( UIMin = "-89", ClampMin = "-89", UIMax = "89", ClampMax = "89", Delta = "1" ))
 	FVector2D Shear;
-	
+
 	/** The angle in degrees to rotate */
 	UPROPERTY(EditAnywhere, BlueprintReadWrite, Category=Transform, meta=( UIMin = "-180", UIMax = "180", Delta = "1" ))
 	float Angle;
-	
+
+public:
 
 	FWidgetTransform()
 		: Translation(FVector2D::ZeroVector)
@@ -37,11 +40,11 @@ struct FWidgetTransform
 	{
 	}
 
-	FWidgetTransform( const FVector2D& InTranslation, const FVector2D& InScale, const FVector2D& InShear, float InAngle )
-		: Translation( InTranslation )
-		, Scale( InScale )
-		, Shear( InShear )
-		, Angle( InAngle )
+	FWidgetTransform(const FVector2D& InTranslation, const FVector2D& InScale, const FVector2D& InShear, float InAngle)
+		: Translation(InTranslation)
+		, Scale(InScale)
+		, Shear(InShear)
+		, Angle(InAngle)
 	{
 	}
 
@@ -52,13 +55,23 @@ struct FWidgetTransform
 		return Identity == *this;
 	}
 
-	bool operator==( const FWidgetTransform &Other ) const
+	FORCEINLINE bool operator==( const FWidgetTransform &Other ) const
 	{
 		return Scale == Other.Scale && Shear == Other.Shear && Angle == Other.Angle && Translation == Other.Translation;
 	}
 
-	bool operator!=( const FWidgetTransform& Other ) const
+	FORCEINLINE bool operator!=( const FWidgetTransform& Other ) const
 	{
-		return !(*this == Other);
+		return !( *this == Other );
+	}
+
+	FORCEINLINE FSlateRenderTransform ToSlateRenderTransform() const
+	{
+		return ::Concatenate(
+			FScale2D(Scale),
+			FShear2D::FromShearAngles(Shear),
+			FQuat2D(FMath::DegreesToRadians(Angle)),
+			FVector2D(Translation));
 	}
 };
+
