UMG - DrawLines, DrawFormattedText added to the set of low level painting methods for UMG.

diff --git a/Engine/Source/Runtime/UMG/Private/WidgetBlueprintLibrary.cpp b/Engine/Source/Runtime/UMG/Private/WidgetBlueprintLibrary.cpp
--- a/Engine/Source/Runtime/UMG/Private/WidgetBlueprintLibrary.cpp
+++ b/Engine/Source/Runtime/UMG/Private/WidgetBlueprintLibrary.cpp
@@ -142,6 +142,21 @@ void UWidgetBlueprintLibrary::DrawLine(UPARAM(ref) FPaintContext& Context, FVect
 		bAntiAlias);
 }
 
+void UWidgetBlueprintLibrary::DrawLines(UPARAM(ref) FPaintContext& Context, const TArray<FVector2D>& Points, float Thickness, FLinearColor Tint, bool bAntiAlias)
+{
+	Context.MaxLayer++;
+
+	FSlateDrawElement::MakeLines(
+		Context.OutDrawElements,
+		Context.MaxLayer,
+		Context.AllottedGeometry.ToPaintGeometry(),
+		Points,
+		Context.MyClippingRect,
+		ESlateDrawEffect::None,
+		Tint,
+		bAntiAlias);
+}
+
 void UWidgetBlueprintLibrary::DrawText(UPARAM(ref) FPaintContext& Context, const FString& InString, FVector2D Position, FLinearColor Tint)
 {
 	Context.MaxLayer++;
@@ -160,6 +175,27 @@ void UWidgetBlueprintLibrary::DrawText(UPARAM(ref) FPaintContext& Context, const
 		Tint);
 }
 
+void UWidgetBlueprintLibrary::DrawTextFormatted(UPARAM(ref) FPaintContext& Context, const FText& Text, FVector2D Position, UFont* Font, int32 FontSize, FName FontTypeFace, FLinearColor Tint)
+{
+	if ( Font )
+	{
+		Context.MaxLayer++;
+
+		//TODO UMG Create a font asset usable as a UFont or as a slate font asset.
+		FSlateFontInfo FontInfo(Font, FontSize, FontTypeFace);
+
+		FSlateDrawElement::MakeText(
+			Context.OutDrawElements,
+			Context.MaxLayer,
+			Context.AllottedGeometry.ToPaintGeometry(),
+			Text,
+			FontInfo,
+			Context.MyClippingRect,
+			ESlateDrawEffect::None,
+			Tint);
+	}
+}
+
 FEventReply UWidgetBlueprintLibrary::Handled()
 {
 	FEventReply Reply;
diff --git a/Engine/Source/Runtime/UMG/Public/Blueprint/WidgetBlueprintLibrary.h b/Engine/Source/Runtime/UMG/Public/Blueprint/WidgetBlueprintLibrary.h
--- a/Engine/Source/Runtime/UMG/Public/Blueprint/WidgetBlueprintLibrary.h
+++ b/Engine/Source/Runtime/UMG/Public/Blueprint/WidgetBlueprintLibrary.h
@@ -53,24 +53,42 @@ public:
 	 *
 	 * @param PositionA		Starting position of the line in local space.
 	 * @param PositionB		Ending position of the line in local space.
-	 * @param Thickness				How many pixels thick this line should be.
+	 * @param Thickness		How many pixels thick this line should be.
 	 * @param Tint			Color to render the line.
 	 */
 	UFUNCTION(BlueprintCallable, meta=( AdvancedDisplay = "5" ), Category="Painting" )
 	static void DrawLine(UPARAM(ref) FPaintContext& Context, FVector2D PositionA, FVector2D PositionB, float Thickness = 1.0f, FLinearColor Tint = FLinearColor::White, bool bAntiAlias = true);
 
-	// TODO UMG DrawLines
+	/**
+	 * Draws several line segments.
+	 *
+	 * @param Points		Line pairs, each line needs to be 2 separate points in the array.
+	 * @param Thickness		How many pixels thick this line should be.
+	 * @param Tint			Color to render the line.
+	 */
+	UFUNCTION(BlueprintCallable, meta=( AdvancedDisplay = "4" ), Category="Painting" )
+	static void DrawLines(UPARAM(ref) FPaintContext& Context, const TArray<FVector2D>& Points, float Thickness = 1.0f, FLinearColor Tint = FLinearColor::White, bool bAntiAlias = true);
 
-	/** 
+	/**
 	 * Draws text.
 	 *
 	 * @param InString		The string to draw.
 	 * @param Position		The starting position where the text is drawn in local space.
 	 * @param Tint			Color to render the line.
 	 */
-	UFUNCTION(BlueprintCallable, Category="Painting")
+	UFUNCTION(BlueprintCallable, Category="Painting", meta=( DeprecatedFunction, DeprecationMessage = "Use Draw Text instead", DisplayName="Draw String"))
 	static void DrawText(UPARAM(ref) FPaintContext& Context, const FString& InString, FVector2D Position, FLinearColor Tint = FLinearColor::White);
 
+	/**
+	 * Draws text.
+	 *
+	 * @param Text			The string to draw.
+	 * @param Position		The starting position where the text is drawn in local space.
+	 * @param Tint			Color to render the line.
+	 */
+	UFUNCTION(BlueprintCallable, Category="Painting", meta=(DisplayName="Draw Text"))
+	void DrawTextFormatted(UPARAM(ref) FPaintContext& Context, const FText& Text, FVector2D Position, UFont* Font, int32 FontSize = 16, FName FontTypeFace = FName(TEXT("Regular")), FLinearColor Tint = FLinearColor::White);
+
 	/**
 	 * The event reply to use when you choose to handle an event.  This will prevent the event 
 	 * from continuing to bubble up / down the widget hierarchy.
