UMG - Rounding the projected location when calculating screen position from world position before applying inverse DPI scale, this produces a much more stable location with no half-pixel offsets due to later applied scale, since we do it before the inverse we more or less end up at the whole number we expect.  This makes things like text above characters heads much more stable.

diff --git a/Engine/Source/Runtime/UMG/Private/WidgetLayoutLibrary.cpp b/Engine/Source/Runtime/UMG/Private/WidgetLayoutLibrary.cpp
--- a/Engine/Source/Runtime/UMG/Private/WidgetLayoutLibrary.cpp
+++ b/Engine/Source/Runtime/UMG/Private/WidgetLayoutLibrary.cpp
@@ -35,6 +35,13 @@ bool UWidgetLayoutLibrary::ProjectWorldLocationToWidgetPosition(APlayerControlle
 			// Remove the resolution quality scale.
 			ScreenPosition = PixelLocation / QualityScale;
 
+			// Round the pixel projected value to reduce jittering due to layout rounding,
+			// I do this before I remove scaling, because scaling is going to be applied later
+			// in the opposite direction, so as long as we round, before inverse scale, scale should
+			// result in more or less the same value, especially after slate does layout rounding.
+			ScreenPosition.X = FMath::RoundToInt(ScreenPosition.X);
+			ScreenPosition.Y = FMath::RoundToInt(ScreenPosition.Y);
+
 			// Get the application / DPI scale
 			const float Scale = UWidgetLayoutLibrary::GetViewportScale(PlayerController);
 
