UMG - Fixing animation and tick rates for widgets living in the world, they should now correctly be controlled by time dilation as they no longer are getting ticked twice, and the second tick using the wrong delta time.  This broke when we introduced tick folding, which ticks again inside of paint, but continued to use the global slate delta time, as well as tick before hand.

diff --git a/Engine/Source/Runtime/UMG/Private/Components/WidgetComponent.cpp b/Engine/Source/Runtime/UMG/Private/Components/WidgetComponent.cpp
--- a/Engine/Source/Runtime/UMG/Private/Components/WidgetComponent.cpp
+++ b/Engine/Source/Runtime/UMG/Private/Components/WidgetComponent.cpp
@@ -11,6 +11,8 @@
 #include "Scalability.h"
 #include "WidgetLayoutLibrary.h"
 
+extern SLATECORE_API int32 bFoldTick;
+
 DECLARE_CYCLE_STAT(TEXT("3DHitTesting"), STAT_Slate3DHitTesting, STATGROUP_Slate);
 
 class SVirtualWindow : public SWindow
@@ -546,7 +548,10 @@ void UWidgetComponent::TickComponent(float DeltaTime, enum ELevelTick TickType,
 
 				FGeometry WindowGeometry = FGeometry::MakeRoot(DrawSize, FSlateLayoutTransform());
 
-				SlateWidget->TickWidgetsRecursively(WindowGeometry, FApp::GetCurrentTime(), DeltaTime);
+				if ( !bFoldTick )
+				{
+					SlateWidget->TickWidgetsRecursively(WindowGeometry, FApp::GetCurrentTime(), DeltaTime);
+				}
 
 				// Ticking can cause geometry changes.  Recompute
 				SlateWidget->SlatePrepass();
@@ -562,7 +567,7 @@ void UWidgetComponent::TickComponent(float DeltaTime, enum ELevelTick TickType,
 
 					// Paint the window
 					MaxLayerId = SlateWidget->Paint(
-						FPaintArgs(SlateWidget.ToSharedRef(), *HitTestGrid, FVector2D::ZeroVector, FSlateApplication::Get().GetCurrentTime(), FSlateApplication::Get().GetDeltaTime()),
+						FPaintArgs(SlateWidget.ToSharedRef(), *HitTestGrid, FVector2D::ZeroVector, FApp::GetCurrentTime(), DeltaTime),
 						WindowGeometry, WindowGeometry.GetClippingRect(),
 						WindowElementList,
 						0,
