Added a SafeZone widget to UMG

--------
Integrated using branch UE4-To-UE4-Orion (reversed) of change#2581687 by [REDACTED_PERSON] on 2015/06/09 15:44:49.

diff --git a/Engine/Source/Runtime/Slate/Private/Widgets/Layout/SSafeZone.cpp b/Engine/Source/Runtime/Slate/Private/Widgets/Layout/SSafeZone.cpp
--- a/Engine/Source/Runtime/Slate/Private/Widgets/Layout/SSafeZone.cpp
+++ b/Engine/Source/Runtime/Slate/Private/Widgets/Layout/SSafeZone.cpp
@@ -15,18 +15,7 @@ void SSafeZone::Construct(const FArguments& InArgs)
 		]
 	);
 
-	// @todo: should we have a function that just returns the safe area size?
-	FDisplayMetrics Metrics;
-	FSlateApplication::Get().GetDisplayMetrics(Metrics);
-
-	if (InArgs._IsTitleSafe)
-	{
-		SafeMargin = FMargin(Metrics.TitleSafePaddingSize.X, Metrics.TitleSafePaddingSize.Y);
-	}
-	else
-	{
-		SafeMargin = FMargin(Metrics.ActionSafePaddingSize.X, Metrics.ActionSafePaddingSize.Y);
-	}
+	SetTitleSafe( InArgs._IsTitleSafe );
 
 	Padding = InArgs._Padding;
 }
@@ -36,3 +25,18 @@ FMargin SSafeZone::GetSafeZonePadding() const
 	// return either the TitleSafe or the ActionSafe size, added to the user padding
 	return Padding.Get() + SafeMargin;
 }
+
+void SSafeZone::SetTitleSafe( bool bIsTitleSafe )
+{
+	FDisplayMetrics Metrics;
+	FSlateApplication::Get().GetDisplayMetrics( Metrics );
+
+	if ( bIsTitleSafe )
+	{
+		SafeMargin = FMargin( Metrics.TitleSafePaddingSize.X, Metrics.TitleSafePaddingSize.Y );
+	}
+	else
+	{
+		SafeMargin = FMargin( Metrics.ActionSafePaddingSize.X, Metrics.ActionSafePaddingSize.Y );
+	}
+}
diff --git a/Engine/Source/Runtime/Slate/Public/Widgets/Layout/SSafeZone.h b/Engine/Source/Runtime/Slate/Public/Widgets/Layout/SSafeZone.h
--- a/Engine/Source/Runtime/Slate/Public/Widgets/Layout/SSafeZone.h
+++ b/Engine/Source/Runtime/Slate/Public/Widgets/Layout/SSafeZone.h
@@ -51,6 +51,8 @@ class SLATE_API SSafeZone : public SBox
 
 	FMargin GetSafeZonePadding() const;
 
+	void SetTitleSafe( bool bIsTitleSafe );
+
 private:
 	/** Cached values from the args */
 	TAttribute<FMargin> Padding;
diff --git a/Engine/Source/Runtime/UMG/Private/Components/SafeZone.cpp b/Engine/Source/Runtime/UMG/Private/Components/SafeZone.cpp
new file mode 100644
--- /dev/null
+++ b/Engine/Source/Runtime/UMG/Private/Components/SafeZone.cpp
@@ -0,0 +1,79 @@
+// Copyright 1998-2015 Epic Games, Inc. All Rights Reserved.
+
+#include "UMGPrivatePCH.h"
+
+#include "Components/SafeZone.h"
+#include "Components/SafeZoneSlot.h"
+#include "Layout/SSafeZone.h"
+
+#define LOCTEXT_NAMESPACE "UMG"
+
+USafeZone::USafeZone()
+{
+	bCanHaveMultipleChildren = false;
+}
+
+#if WITH_EDITOR
+const FSlateBrush* USafeZone::GetEditorIcon()
+{
+	return FUMGStyle::Get().GetBrush( "Widget.SafeZone" );
+}
+
+const FText USafeZone::GetPaletteCategory()
+{
+	return LOCTEXT( "Panel", "Panel" );
+}
+#endif
+
+void USafeZone::OnSlotAdded( UPanelSlot* Slot )
+{
+	Super::OnSlotAdded( Slot );
+
+	UpdateWidgetProperties();
+}
+
+void USafeZone::OnSlotRemoved( UPanelSlot* Slot )
+{
+	Super::OnSlotRemoved( Slot );
+
+	if ( MySafeZone.IsValid() )
+	{
+		MySafeZone->SetContent( SNullWidget::NullWidget );
+	}
+}
+
+UClass* USafeZone::GetSlotClass() const
+{
+	return USafeZoneSlot::StaticClass();
+}
+
+void USafeZone::UpdateWidgetProperties()
+{
+	if ( MySafeZone.IsValid() && GetChildrenCount() > 0 )
+	{
+		USafeZoneSlot* SafeSlot = CastChecked< USafeZoneSlot >( Slots[ 0 ] );
+
+		MySafeZone->SetTitleSafe( SafeSlot->bIsTitleSafe );
+		MySafeZone->SetHAlign( SafeSlot->HAlign );
+		MySafeZone->SetVAlign( SafeSlot->VAlign );
+		MySafeZone->SetPadding( SafeSlot->Padding );
+	}
+}
+
+TSharedRef<SWidget> USafeZone::RebuildWidget()
+{
+	USafeZoneSlot* SafeSlot = Slots.Num() > 0 ? Cast< USafeZoneSlot >( Slots[ 0 ] ) : nullptr;
+
+	MySafeZone = SNew( SSafeZone )
+		.IsTitleSafe( SafeSlot ? SafeSlot->bIsTitleSafe : false )
+		.HAlign( SafeSlot ? SafeSlot->HAlign : HAlign_Fill )
+		.VAlign( SafeSlot ? SafeSlot->VAlign : VAlign_Fill )
+		.Padding( SafeSlot ? SafeSlot->Padding : FMargin() )
+		[
+			GetChildAt( 0 ) ? GetChildAt( 0 )->TakeWidget() : SNullWidget::NullWidget
+		];
+
+	return BuildDesignTimeWidget( MySafeZone.ToSharedRef() );
+}
+
+#undef LOCTEXT_NAMESPACE
\ No newline at end of file
diff --git a/Engine/Source/Runtime/UMG/Private/Components/SafeZoneSlot.cpp b/Engine/Source/Runtime/UMG/Private/Components/SafeZoneSlot.cpp
new file mode 100644
--- /dev/null
+++ b/Engine/Source/Runtime/UMG/Private/Components/SafeZoneSlot.cpp
@@ -0,0 +1,21 @@
+// Copyright 1998-2015 Epic Games, Inc. All Rights Reserved.
+
+#include "UMGPrivatePCH.h"
+
+#include "Components/SafeZoneSlot.h"
+#include "Components/SafeZone.h"
+
+USafeZoneSlot::USafeZoneSlot()
+{
+
+}
+
+void USafeZoneSlot::SynchronizeProperties()
+{
+	Super::SynchronizeProperties();
+
+	if ( IsValid( Parent ) )
+	{
+		CastChecked< USafeZone >( Parent )->UpdateWidgetProperties();
+	}
+}
diff --git a/Engine/Source/Runtime/UMG/Public/Components/SafeZone.h b/Engine/Source/Runtime/UMG/Public/Components/SafeZone.h
new file mode 100644
--- /dev/null
+++ b/Engine/Source/Runtime/UMG/Public/Components/SafeZone.h
@@ -0,0 +1,31 @@
+// Copyright 1998-2015 Epic Games, Inc. All Rights Reserved.
+
+#pragma once
+
+#include "SafeZone.generated.h"
+
+UCLASS()
+class UMG_API USafeZone : public UContentWidget
+{
+	GENERATED_BODY()
+public:
+	USafeZone();
+
+#if WITH_EDITOR
+	virtual const FSlateBrush* GetEditorIcon() override;
+	virtual const FText GetPaletteCategory() override;
+#endif
+
+	virtual void OnSlotAdded( UPanelSlot* Slot ) override;
+	virtual void OnSlotRemoved( UPanelSlot* Slot ) override;
+	virtual UClass* GetSlotClass() const override;
+
+	void UpdateWidgetProperties();
+
+protected:
+	// UWidget interface
+	virtual TSharedRef<SWidget> RebuildWidget() override;
+	// End of UWidget interface
+
+	TSharedPtr< class SSafeZone > MySafeZone;
+};
\ No newline at end of file
diff --git a/Engine/Source/Runtime/UMG/Public/Components/SafeZoneSlot.h b/Engine/Source/Runtime/UMG/Public/Components/SafeZoneSlot.h
new file mode 100644
--- /dev/null
+++ b/Engine/Source/Runtime/UMG/Public/Components/SafeZoneSlot.h
@@ -0,0 +1,30 @@
+// Copyright 1998-2015 Epic Games, Inc. All Rights Reserved.
+
+#pragma once
+
+#include "Layout/SSafeZone.h"
+
+#include "SafeZoneSlot.generated.h"
+
+UCLASS()
+class USafeZoneSlot : public UPanelSlot
+{
+	GENERATED_BODY()
+public:
+
+	UPROPERTY( EditAnywhere, BlueprintReadOnly, Category = SafeZone )
+	bool bIsTitleSafe;
+
+	UPROPERTY( EditAnywhere, BlueprintReadOnly, Category = SafeZone )
+	TEnumAsByte< EHorizontalAlignment > HAlign;
+
+	UPROPERTY( EditAnywhere, BlueprintReadOnly, Category = SafeZone )
+	TEnumAsByte< EVerticalAlignment > VAlign;
+
+	UPROPERTY( EditAnywhere, BlueprintReadOnly, Category = SafeZone )
+	FMargin Padding;
+
+	USafeZoneSlot();
+
+	virtual void SynchronizeProperties() override;
+};
\ No newline at end of file
