Use GetValue on TEnumAsByte to avoid linux compile errors.

--------
Integrated using branch UE4-To-UE4-Orion (reversed) of change#2581782 by [REDACTED_PERSON] on 2015/06/09 16:31:59.

diff --git a/Engine/Source/Runtime/UMG/Private/Components/SafeZone.cpp b/Engine/Source/Runtime/UMG/Private/Components/SafeZone.cpp
--- a/Engine/Source/Runtime/UMG/Private/Components/SafeZone.cpp
+++ b/Engine/Source/Runtime/UMG/Private/Components/SafeZone.cpp
@@ -54,8 +54,8 @@ void USafeZone::UpdateWidgetProperties()
 		USafeZoneSlot* SafeSlot = CastChecked< USafeZoneSlot >( Slots[ 0 ] );
 
 		MySafeZone->SetTitleSafe( SafeSlot->bIsTitleSafe );
-		MySafeZone->SetHAlign( SafeSlot->HAlign );
-		MySafeZone->SetVAlign( SafeSlot->VAlign );
+		MySafeZone->SetHAlign( SafeSlot->HAlign.GetValue() );
+		MySafeZone->SetVAlign( SafeSlot->VAlign.GetValue() );
 		MySafeZone->SetPadding( SafeSlot->Padding );
 	}
 }
@@ -66,8 +66,8 @@ TSharedRef<SWidget> USafeZone::RebuildWidget()
 
 	MySafeZone = SNew( SSafeZone )
 		.IsTitleSafe( SafeSlot ? SafeSlot->bIsTitleSafe : false )
-		.HAlign( SafeSlot ? SafeSlot->HAlign : HAlign_Fill )
-		.VAlign( SafeSlot ? SafeSlot->VAlign : VAlign_Fill )
+		.HAlign( SafeSlot ? SafeSlot->HAlign.GetValue() : HAlign_Fill )
+		.VAlign( SafeSlot ? SafeSlot->VAlign.GetValue() : VAlign_Fill )
 		.Padding( SafeSlot ? SafeSlot->Padding : FMargin() )
 		[
 			GetChildAt( 0 ) ? GetChildAt( 0 )->TakeWidget() : SNullWidget::NullWidget
