UMG - Fixing a bug in SObjectWidget that was preventing natively handled key presses from being returned as handled to the rest of slate.

diff --git a/Engine/Source/Runtime/UMG/Private/Slate/SObjectWidget.cpp b/Engine/Source/Runtime/UMG/Private/Slate/SObjectWidget.cpp
--- a/Engine/Source/Runtime/UMG/Private/Slate/SObjectWidget.cpp
+++ b/Engine/Source/Runtime/UMG/Private/Slate/SObjectWidget.cpp
@@ -146,6 +146,8 @@ FReply SObjectWidget::OnKeyDown(const FGeometry& MyGeometry, const FKeyEvent& In
 		{
 			return SCompoundWidget::OnKeyDown(MyGeometry, InKeyEvent);
 		}
+
+		return Result;
 	}
 
 	return FReply::Unhandled();
@@ -160,6 +162,8 @@ FReply SObjectWidget::OnKeyUp(const FGeometry& MyGeometry, const FKeyEvent& InKe
 		{
 			return SCompoundWidget::OnKeyUp(MyGeometry, InKeyEvent);
 		}
+
+		return Result;
 	}
 
 	return FReply::Unhandled();
@@ -174,6 +178,8 @@ FReply SObjectWidget::OnAnalogValueChanged(const FGeometry& MyGeometry, const FA
 		{
 			return SCompoundWidget::OnAnalogValueChanged(MyGeometry, InAnalogInputEvent);
 		}
+
+		return Result;
 	}
 
 	return FReply::Unhandled();
