Fix bad default rotation for widget components.  New widget components now face down +X by default
Fixed changing rendering properties on widget components not updating the component

diff --git a/Engine/Source/Runtime/Core/Public/UObject/ObjectVersion.h b/Engine/Source/Runtime/Core/Public/UObject/ObjectVersion.h
--- a/Engine/Source/Runtime/Core/Public/UObject/ObjectVersion.h
+++ b/Engine/Source/Runtime/Core/Public/UObject/ObjectVersion.h
@@ -534,7 +534,8 @@ enum EUnrealEngineObjectUE4Version
 	VER_UE4_STATIC_SHADOWMAP_PENUMBRA_SIZE,
 	// Fix BC on Niagara effects from the data object and dev UI changes.
 	VER_UE4_NIAGARA_DATA_OBJECT_DEV_UI_FIX,
-
+	// Fixed the default orientation of widget component so it faces down +x
+	VER_UE4_FIXED_DEFAULT_ORIENTATION_OF_WIDGET_COMPONENT,
 	// -----<new versions can be added before this line>-------------------------------------------------
 	// - this needs to be the last line (see note below)
 	VER_UE4_AUTOMATIC_VERSION_PLUS_ONE,
diff --git a/Engine/Source/Runtime/UMG/Private/Components/WidgetComponent.cpp b/Engine/Source/Runtime/UMG/Private/Components/WidgetComponent.cpp
--- a/Engine/Source/Runtime/UMG/Private/Components/WidgetComponent.cpp
+++ b/Engine/Source/Runtime/UMG/Private/Components/WidgetComponent.cpp
@@ -49,6 +49,7 @@ public:
 		, MaterialInstance( InComponent->GetMaterialInstance() )
 		, BodySetup( InComponent->GetBodySetup() )
 		, BlendMode( InComponent->GetBlendMode() )
+		, bUseLegacyRotation( InComponent->IsUsingLegacyRotation() )
 	{
 		bWillEverBeLit = false;
 	}
@@ -99,10 +100,20 @@ public:
 
 					if ( VisibilityMap & ( 1 << ViewIndex ) )
 					{
-						VertexIndices[0] = MeshBuilder.AddVertex(FVector(U, V, 0), FVector2D(0, 0), FVector(1, 0, 0), FVector(0, 1, 0), FVector(0, 0, 1), FColor::White);
-						VertexIndices[1] = MeshBuilder.AddVertex(FVector(U, VL, 0), FVector2D(0, 1), FVector(1, 0, 0), FVector(0, 1, 0), FVector(0, 0, 1), FColor::White);
-						VertexIndices[2] = MeshBuilder.AddVertex(FVector(UL, VL, 0), FVector2D(1, 1), FVector(1, 0, 0), FVector(0, 1, 0), FVector(0, 0, 1), FColor::White);
-						VertexIndices[3] = MeshBuilder.AddVertex(FVector(UL, V, 0), FVector2D(1, 0), FVector(1, 0, 0), FVector(0, 1, 0), FVector(0, 0, 1), FColor::White);
+						if( bUseLegacyRotation )
+						{
+							VertexIndices[0] = MeshBuilder.AddVertex(FVector(U, V, 0), FVector2D(0, 0), FVector(1, 0, 0), FVector(0, 1, 0), FVector(0, 0, 1), FColor::White);
+							VertexIndices[1] = MeshBuilder.AddVertex(FVector(U, VL, 0), FVector2D(0, 1), FVector(1, 0, 0), FVector(0, 1, 0), FVector(0, 0, 1), FColor::White);
+							VertexIndices[2] = MeshBuilder.AddVertex(FVector(UL, VL, 0), FVector2D(1, 1), FVector(1, 0, 0), FVector(0, 1, 0), FVector(0, 0, 1), FColor::White);
+							VertexIndices[3] = MeshBuilder.AddVertex(FVector(UL, V, 0), FVector2D(1, 0), FVector(1, 0, 0), FVector(0, 1, 0), FVector(0, 0, 1), FColor::White);
+						}
+						else
+						{
+							VertexIndices[0] = MeshBuilder.AddVertex(-FVector(0, U, V), FVector2D(0, 0), FVector(1, 0, 0), FVector(0, 1, 0), FVector(0, 0, 1), FColor::White);
+							VertexIndices[1] = MeshBuilder.AddVertex(-FVector(0, U, VL), FVector2D(0, 1), FVector(1, 0, 0), FVector(0, 1, 0), FVector(0, 0, 1), FColor::White);
+							VertexIndices[2] = MeshBuilder.AddVertex(-FVector(0, UL, VL), FVector2D(1, 1), FVector(1, 0, 0), FVector(0, 1, 0), FVector(0, 0, 1), FColor::White);
+							VertexIndices[3] = MeshBuilder.AddVertex(-FVector(0, UL, V), FVector2D(1, 0), FVector(1, 0, 0), FVector(0, 1, 0), FVector(0, 0, 1), FColor::White);
+						}
 
 						MeshBuilder.AddTriangle(VertexIndices[0], VertexIndices[1], VertexIndices[2]);
 						MeshBuilder.AddTriangle(VertexIndices[0], VertexIndices[2], VertexIndices[3]);
@@ -159,6 +170,7 @@ private:
 	UMaterialInstanceDynamic* MaterialInstance;
 	UBodySetup* BodySetup;
 	EWidgetBlendMode BlendMode;
+	bool bUseLegacyRotation;
 };
 
 /**
@@ -326,9 +338,9 @@ UWidgetComponent::UWidgetComponent( const FObjectInitializer& PCIP )
 	PrimaryComponentTick.bCanEverTick = true;
 	PrimaryComponentTick.TickGroup = TG_PostUpdateWork;
 	bTickInEditor = true;
-	
-	RelativeRotation = FRotator(0.f, 0.f, 90.f);
-	
+
+	RelativeRotation = FRotator::ZeroRotator;
+
 	BodyInstance.SetCollisionProfileName(FName(TEXT("UI")));
 
 	// Translucent material instances
@@ -353,6 +365,8 @@ UWidgetComponent::UWidgetComponent( const FObjectInitializer& PCIP )
 	//bGenerateOverlapEvents = false;
 	bUseEditorCompositing = false;
 
+	bUseLegacyRotation = false;
+
 	Space = EWidgetSpace::World;
 	Pivot = FVector2D(0.5, 0.5);
 	ZOrder = -100;
@@ -440,23 +454,9 @@ void UWidgetComponent::OnRegister()
 				}
 			}
 
-			if ( !MaterialInstance )
+			if( !MaterialInstance )
 			{
-				UMaterialInterface* Parent = nullptr;
-				switch ( BlendMode )
-				{
-				case EWidgetBlendMode::Opaque:
-					Parent = bIsTwoSided ? OpaqueMaterial : OpaqueMaterial_OneSided;
-					break;
-				case EWidgetBlendMode::Masked:
-					Parent = bIsTwoSided ? MaskedMaterial : MaskedMaterial_OneSided;
-					break;
-				case EWidgetBlendMode::Transparent:
-					Parent = bIsTwoSided ? TranslucentMaterial : TranslucentMaterial_OneSided;
-					break;
-				}
-
-				MaterialInstance = UMaterialInstanceDynamic::Create(Parent, this);
+				UpdateMaterialInstance();
 			}
 		}
 
@@ -718,7 +718,12 @@ void UWidgetComponent::PostEditChangeProperty(FPropertyChangedEvent& PropertyCha
 
 			MarkRenderStateDirty();
 		}
-		else if ( PropertyName == IsOpaqueName || PropertyName == BackgroundColorName || PropertyName == IsTwoSidedName || PropertyName == ParabolaDistortionName)
+		else if ( PropertyName == IsOpaqueName || PropertyName == IsTwoSidedName )
+		{
+			UpdateMaterialInstance();
+			MarkRenderStateDirty();
+		}
+		else if( PropertyName == BackgroundColorName || PropertyName == ParabolaDistortionName )
 		{
 			MarkRenderStateDirty();
 		}
@@ -983,4 +988,51 @@ void UWidgetComponent::PostLoad()
 	{
 		BlendMode = bIsOpaque_DEPRECATED ? EWidgetBlendMode::Opaque : EWidgetBlendMode::Transparent;
 	}
+
+	if( GetLinkerUE4Version() < VER_UE4_FIXED_DEFAULT_ORIENTATION_OF_WIDGET_COMPONENT )
+	{	
+		// This indicates the value does not differ from the default.  In some rare cases this could cause incorrect rotation for anyone who directly set a value of 0,0,0 for rotation
+		// However due to delta serialization we have no way to know if this value is actually different from the default so assume it is not.
+		if( RelativeRotation == FRotator::ZeroRotator )
+		{
+			RelativeRotation = FRotator(0.f, 0.f, 90.f);
+		}
+		bUseLegacyRotation = true;
+	}
+}
+
+void UWidgetComponent::UpdateMaterialInstance()
+{
+	UMaterialInterface* Parent = nullptr;
+	switch ( BlendMode )
+	{
+	case EWidgetBlendMode::Opaque:
+		Parent = bIsTwoSided ? OpaqueMaterial : OpaqueMaterial_OneSided;
+		break;
+	case EWidgetBlendMode::Masked:
+		Parent = bIsTwoSided ? MaskedMaterial : MaskedMaterial_OneSided;
+		break;
+	case EWidgetBlendMode::Transparent:
+		Parent = bIsTwoSided ? TranslucentMaterial : TranslucentMaterial_OneSided;
+		break;
+	}
+
+	if( MaterialInstance )
+	{
+		MaterialInstance->MarkPendingKill();
+		MaterialInstance = nullptr;
+	}
+
+	MaterialInstance = UMaterialInstanceDynamic::Create(Parent, this);
+
+	if( MaterialInstance )
+	{
+		MaterialInstance->SetTextureParameterValue("SlateUI", RenderTarget);
+	}
+
+}
+
+void UWidgetComponent::SetWidgetClass(TSubclassOf<UUserWidget> InWidgetClass)
+{
+	WidgetClass = InWidgetClass;
 }
diff --git a/Engine/Source/Runtime/UMG/Public/Components/WidgetComponent.h b/Engine/Source/Runtime/UMG/Public/Components/WidgetComponent.h
--- a/Engine/Source/Runtime/UMG/Public/Components/WidgetComponent.h
+++ b/Engine/Source/Runtime/UMG/Public/Components/WidgetComponent.h
@@ -136,6 +136,14 @@ public:
 	/** Get the fake window we create for widgets displayed in the world. */
 	TSharedPtr< SWindow > GetVirtualWindow() const;
 	
+	/** Whether or not this component uses legacy default rotation */
+	bool IsUsingLegacyRotation() const { return bUseLegacyRotation; }
+
+	/** Updates the actual material being used */
+	void UpdateMaterialInstance();
+	
+	/** Sets the widget class used to generate the widget for this component */
+	void SetWidgetClass(TSubclassOf<UUserWidget> InWidgetClass);
 protected:
 	/** The coordinate space in which to render the widget */
 	UPROPERTY(EditAnywhere, Category=UI)
@@ -235,6 +243,8 @@ protected:
 	UPROPERTY(Transient, DuplicateTransient)
 	UMaterialInstanceDynamic* MaterialInstance;
 
+	UPROPERTY()
+	bool bUseLegacyRotation;
 protected:
 
 	/** The grid used to find actual hit actual widgets once input has been translated to the components local space */
