UMG - UE-17226 - The assignment operator for the FPaintContext wasn't copying the MaxLayerId, painting into layers and not telling slate they exist apparently causes the graphics corruptions seen by several people.

diff --git a/Engine/Source/Runtime/UMG/Public/Blueprint/UserWidget.h b/Engine/Source/Runtime/UMG/Public/Blueprint/UserWidget.h
--- a/Engine/Source/Runtime/UMG/Public/Blueprint/UserWidget.h
+++ b/Engine/Source/Runtime/UMG/Public/Blueprint/UserWidget.h
@@ -62,6 +62,7 @@ public:
 		FPaintContext* Ptr = this;
 		Ptr->~FPaintContext();
 		new(Ptr) FPaintContext(Other.AllottedGeometry, Other.MyClippingRect, Other.OutDrawElements, Other.LayerId, Other.WidgetStyle, Other.bParentEnabled);
+		Ptr->MaxLayer = Other.MaxLayer;
 	}
 
 public:
