UMG - Removing Thickness from the line painting function.  It wasn't actually hooked up to anything, and there's no support in slate for it currently, only when rendering splines.  Will add back once we've added support for it in Slate.

diff --git a/Engine/Source/Runtime/UMG/Private/WidgetBlueprintLibrary.cpp b/Engine/Source/Runtime/UMG/Private/WidgetBlueprintLibrary.cpp
--- a/Engine/Source/Runtime/UMG/Private/WidgetBlueprintLibrary.cpp
+++ b/Engine/Source/Runtime/UMG/Private/WidgetBlueprintLibrary.cpp
@@ -123,7 +123,7 @@ void UWidgetBlueprintLibrary::DrawBox(UPARAM(ref) FPaintContext& Context, FVecto
 	}
 }
 
-void UWidgetBlueprintLibrary::DrawLine(UPARAM(ref) FPaintContext& Context, FVector2D PositionA, FVector2D PositionB, float Thickness, FLinearColor Tint, bool bAntiAlias)
+void UWidgetBlueprintLibrary::DrawLine(UPARAM(ref) FPaintContext& Context, FVector2D PositionA, FVector2D PositionB, FLinearColor Tint, bool bAntiAlias)
 {
 	Context.MaxLayer++;
 
@@ -142,7 +142,7 @@ void UWidgetBlueprintLibrary::DrawLine(UPARAM(ref) FPaintContext& Context, FVect
 		bAntiAlias);
 }
 
-void UWidgetBlueprintLibrary::DrawLines(UPARAM(ref) FPaintContext& Context, const TArray<FVector2D>& Points, float Thickness, FLinearColor Tint, bool bAntiAlias)
+void UWidgetBlueprintLibrary::DrawLines(UPARAM(ref) FPaintContext& Context, const TArray<FVector2D>& Points, FLinearColor Tint, bool bAntiAlias)
 {
 	Context.MaxLayer++;
 
diff --git a/Engine/Source/Runtime/UMG/Public/Blueprint/WidgetBlueprintLibrary.h b/Engine/Source/Runtime/UMG/Public/Blueprint/WidgetBlueprintLibrary.h
--- a/Engine/Source/Runtime/UMG/Public/Blueprint/WidgetBlueprintLibrary.h
+++ b/Engine/Source/Runtime/UMG/Public/Blueprint/WidgetBlueprintLibrary.h
@@ -56,8 +56,8 @@ public:
 	 * @param Thickness		How many pixels thick this line should be.
 	 * @param Tint			Color to render the line.
 	 */
-	UFUNCTION(BlueprintCallable, meta=( AdvancedDisplay = "5" ), Category="Painting" )
-	static void DrawLine(UPARAM(ref) FPaintContext& Context, FVector2D PositionA, FVector2D PositionB, float Thickness = 1.0f, FLinearColor Tint = FLinearColor::White, bool bAntiAlias = true);
+	UFUNCTION(BlueprintCallable, meta=( AdvancedDisplay = "4" ), Category="Painting" )
+	static void DrawLine(UPARAM(ref) FPaintContext& Context, FVector2D PositionA, FVector2D PositionB, FLinearColor Tint = FLinearColor::White, bool bAntiAlias = true);
 
 	/**
 	 * Draws several line segments.
@@ -66,8 +66,8 @@ public:
 	 * @param Thickness		How many pixels thick this line should be.
 	 * @param Tint			Color to render the line.
 	 */
-	UFUNCTION(BlueprintCallable, meta=( AdvancedDisplay = "4" ), Category="Painting" )
-	static void DrawLines(UPARAM(ref) FPaintContext& Context, const TArray<FVector2D>& Points, float Thickness = 1.0f, FLinearColor Tint = FLinearColor::White, bool bAntiAlias = true);
+	UFUNCTION(BlueprintCallable, meta=( AdvancedDisplay = "3" ), Category="Painting" )
+	static void DrawLines(UPARAM(ref) FPaintContext& Context, const TArray<FVector2D>& Points, FLinearColor Tint = FLinearColor::White, bool bAntiAlias = true);
 
 	/**
 	 * Draws text.
