merge SafeZone CL 2592668

override ReleaseSlateResources to reset safe zone

diff --git a/Engine/Source/Runtime/UMG/Private/Components/SafeZone.cpp b/Engine/Source/Runtime/UMG/Private/Components/SafeZone.cpp
--- a/Engine/Source/Runtime/UMG/Private/Components/SafeZone.cpp
+++ b/Engine/Source/Runtime/UMG/Private/Components/SafeZone.cpp
@@ -76,4 +76,11 @@ TSharedRef<SWidget> USafeZone::RebuildWidget()
 	return BuildDesignTimeWidget( MySafeZone.ToSharedRef() );
 }
 
+void USafeZone::ReleaseSlateResources(bool bReleaseChildren)
+{
+	Super::ReleaseSlateResources(bReleaseChildren);
+
+	MySafeZone.Reset();
+}
+
 #undef LOCTEXT_NAMESPACE
\ No newline at end of file
diff --git a/Engine/Source/Runtime/UMG/Public/Components/SafeZone.h b/Engine/Source/Runtime/UMG/Public/Components/SafeZone.h
--- a/Engine/Source/Runtime/UMG/Public/Components/SafeZone.h
+++ b/Engine/Source/Runtime/UMG/Public/Components/SafeZone.h
@@ -27,5 +27,7 @@ protected:
 	virtual TSharedRef<SWidget> RebuildWidget() override;
 	// End of UWidget interface
 
+	virtual void ReleaseSlateResources(bool bReleaseChildren) override;
+
 	TSharedPtr< class SSafeZone > MySafeZone;
 };
\ No newline at end of file
