Additional fixes for bad default widget component rotations

diff --git a/Engine/Source/Runtime/UMG/Private/Components/WidgetComponent.cpp b/Engine/Source/Runtime/UMG/Private/Components/WidgetComponent.cpp
--- a/Engine/Source/Runtime/UMG/Private/Components/WidgetComponent.cpp
+++ b/Engine/Source/Runtime/UMG/Private/Components/WidgetComponent.cpp
@@ -386,12 +386,26 @@ FBoxSphereBounds UWidgetComponent::CalcBounds(const FTransform & LocalToWorld) c
 {
 	if ( Space != EWidgetSpace::Screen )
 	{
-		const FVector Origin = FVector(
+
+		if( bUseLegacyRotation )
+		{
+			const FVector Origin = FVector(
 			( DrawSize.X * 0.5f ) - ( DrawSize.X * Pivot.X ),
 			( DrawSize.Y * 0.5f ) - ( DrawSize.Y * Pivot.Y ), .5f);
-		const FVector BoxExtent = FVector(DrawSize.X / 2.0f, DrawSize.Y / 2.0f, 1.0f);
+			const FVector BoxExtent = FVector(DrawSize.X / 2.0f, DrawSize.Y / 2.0f, 1.0f);
+
+			return FBoxSphereBounds(Origin, BoxExtent, DrawSize.Size() / 2.0f).TransformBy(LocalToWorld);
+		}
+		else
+		{
+			const FVector Origin = FVector(.5f,
+			( DrawSize.X * 0.5f ) - ( DrawSize.X * Pivot.X ),
+			( DrawSize.Y * 0.5f ) - ( DrawSize.Y * Pivot.Y ));
+
+			const FVector BoxExtent = FVector(1.f, DrawSize.X / 2.0f, DrawSize.Y / 2.0f);
 
-		return FBoxSphereBounds(Origin, BoxExtent, DrawSize.Size() / 2.0f).TransformBy(LocalToWorld);
+			return FBoxSphereBounds(Origin, BoxExtent, DrawSize.Size() / 2.0f).TransformBy(LocalToWorld);
+		}
 	}
 	else
 	{
@@ -409,7 +423,17 @@ FCollisionShape UWidgetComponent::GetCollisionShape(float Inflation) const
 {
 	if ( Space != EWidgetSpace::Screen )
 	{
-		FVector	BoxHalfExtent = ( FVector(DrawSize.X * 0.5f, DrawSize.Y * 0.5f, 1.0f) * ComponentToWorld.GetScale3D() ) + Inflation;
+		FVector BoxHalfExtent;
+
+		if( bUseLegacyRotation )
+		{
+			BoxHalfExtent = ( FVector(DrawSize.X * 0.5f, DrawSize.Y * 0.5f, 1.0f) * ComponentToWorld.GetScale3D() ) + Inflation;
+		}
+		else
+		{
+			BoxHalfExtent = ( FVector(1.0f, DrawSize.X * 0.5f, DrawSize.Y * 0.5f) * ComponentToWorld.GetScale3D() ) + Inflation;
+		}
+
 		if ( Inflation < 0.0f )
 		{
 			// Don't shrink below zero size.
@@ -885,15 +909,31 @@ void UWidgetComponent::UpdateBodySetup( bool bDrawSizeChanged )
 
 		FKBoxElem* BoxElem = BodySetup->AggGeom.BoxElems.GetData();
 
-		const FVector Origin = FVector(
-			(DrawSize.X * 0.5f) - ( DrawSize.X * Pivot.X ),
-			(DrawSize.Y * 0.5f) - ( DrawSize.Y * Pivot.Y ), .5f);
+		FVector Origin;
+		if( bUseLegacyRotation )
+		{
+			Origin = FVector(
+				(DrawSize.X * 0.5f) - ( DrawSize.X * Pivot.X ),
+				(DrawSize.Y * 0.5f) - ( DrawSize.Y * Pivot.Y ), .5f);
+
+					
+			BoxElem->X = DrawSize.X;
+			BoxElem->Y = DrawSize.Y;
+			BoxElem->Z = 1.0f;
+		}
+		else
+		{
+			Origin = FVector(.5f,
+				(DrawSize.X * 0.5f) - ( DrawSize.X * Pivot.X ),
+				(DrawSize.Y * 0.5f) - ( DrawSize.Y * Pivot.Y ) );
+					
+			BoxElem->X = 1.0f;
+			BoxElem->Y = DrawSize.X;
+			BoxElem->Z = DrawSize.Y;
+		}
 
 		BoxElem->SetTransform(FTransform::Identity);
 		BoxElem->Center = Origin;
-		BoxElem->X = DrawSize.X;
-		BoxElem->Y = DrawSize.Y;
-		BoxElem->Z = 1.0f;
 	}	
 }
 
