UMG - The WidgetComponent has been reworked to use the new IGameLayer system to create a single shared canvas layer that all the components add themselves to when they're displaying their information in screen space, instead of world space.  They now properly sort based on depth/distance from the camera.  Additionally, the layer widget now updates the positions for the widgets.  Thereby solving the problem of tick groups causing weird update orders for objects touching the WidgetComponent's Actor as a base.  The WidgetComponent now has a normal tick group.

New Framework Functions:
APlayerController::ProjectWorldLocationToScreenWithDistance
UWidgetLayoutLibrary::ProjectWorldLocationToWidgetPositionWithDistance

diff --git a/Engine/Source/Runtime/Engine/Classes/GameFramework/PlayerController.h b/Engine/Source/Runtime/Engine/Classes/GameFramework/PlayerController.h
--- a/Engine/Source/Runtime/Engine/Classes/GameFramework/PlayerController.h
+++ b/Engine/Source/Runtime/Engine/Classes/GameFramework/PlayerController.h
@@ -479,6 +479,12 @@ public:
 	 */
 	UFUNCTION(BlueprintCallable, Category = "Game|Player", meta = ( DisplayName = "ConvertWorldLocationToScreenLocation" ))
 	bool ProjectWorldLocationToScreen(FVector WorldLocation, FVector2D& ScreenLocation) const;
+
+	/**
+	 * Convert a World Space 3D position into a 3D Screen Space position.
+	 * @return true if the world coordinate was successfully projected to the screen.
+	 */
+	bool ProjectWorldLocationToScreenWithDistance(FVector WorldLocation, FVector& ScreenLocation) const;
 	
 	/**
 	  * Updates the rotation of player, based on ControlRotation after RotationInput has been applied.
diff --git a/Engine/Source/Runtime/Engine/Private/PlayerController.cpp b/Engine/Source/Runtime/Engine/Private/PlayerController.cpp
--- a/Engine/Source/Runtime/Engine/Private/PlayerController.cpp
+++ b/Engine/Source/Runtime/Engine/Private/PlayerController.cpp
@@ -1884,6 +1884,14 @@ bool APlayerController::DeprojectScreenPositionToWorld(float ScreenX, float Scre
 
 
 bool APlayerController::ProjectWorldLocationToScreen(FVector WorldLocation, FVector2D& ScreenLocation) const
+{
+	FVector FullScreenLocation;
+	const bool bSuccess = ProjectWorldLocationToScreenWithDistance(WorldLocation, FullScreenLocation);
+	ScreenLocation = FVector2D(FullScreenLocation.X, FullScreenLocation.Y);
+	return bSuccess;
+}
+
+bool APlayerController::ProjectWorldLocationToScreenWithDistance(FVector WorldLocation, FVector& ScreenLocation) const
 {
 	ULocalPlayer* LocalPlayer = Cast<ULocalPlayer>(Player);
 	if (LocalPlayer != NULL && LocalPlayer->ViewportClient != NULL && LocalPlayer->ViewportClient->Viewport != NULL)
@@ -1902,7 +1910,10 @@ bool APlayerController::ProjectWorldLocationToScreen(FVector WorldLocation, FVec
 
 		if (SceneView) 
 		{
-			return SceneView->WorldToPixel(WorldLocation, ScreenLocation);
+			FVector2D ScreenLocation2D;
+			const bool bSuccess = SceneView->WorldToPixel(WorldLocation, ScreenLocation2D);
+			ScreenLocation = FVector(ScreenLocation2D.X, ScreenLocation2D.Y, FVector::Dist(ViewLocation, WorldLocation));
+			return bSuccess;
 		}
 	}
 
diff --git a/Engine/Source/Runtime/UMG/Private/Components/WidgetComponent.cpp b/Engine/Source/Runtime/UMG/Private/Components/WidgetComponent.cpp
--- a/Engine/Source/Runtime/UMG/Private/Components/WidgetComponent.cpp
+++ b/Engine/Source/Runtime/UMG/Private/Components/WidgetComponent.cpp
@@ -11,11 +11,162 @@
 #include "Scalability.h"
 #include "WidgetLayoutLibrary.h"
 #include "PhysicsEngine/BodySetup.h"
+#include "SGameLayerManager.h"
 
 extern SLATECORE_API int32 bFoldTick;
 
 DECLARE_CYCLE_STAT(TEXT("3DHitTesting"), STAT_Slate3DHitTesting, STATGROUP_Slate);
 
+static const FName SharedLayerName(TEXT("WidgetComponentScreenLayer"));
+
+class SWorldWidgetScreenLayer : public SCompoundWidget
+{
+	SLATE_BEGIN_ARGS(SWorldWidgetScreenLayer)
+	{
+		_Visibility = EVisibility::SelfHitTestInvisible;
+	}
+	SLATE_END_ARGS()
+
+public:
+	void Construct(const FArguments& InArgs, const TArray<TWeakObjectPtr<UWidgetComponent>>& InComponents)
+	{
+		Components = InComponents;
+
+		ChildSlot
+		[
+			SAssignNew(Canvas, SConstraintCanvas)
+		];
+	}
+
+	virtual void Tick(const FGeometry& AllottedGeometry, const double InCurrentTime, const float InDeltaTime) override
+	{
+		TArray<UWidgetComponent*, TInlineAllocator<1>> DeadComponents;
+
+		for ( TWeakObjectPtr<UWidgetComponent> Component : Components )
+		{
+			if ( UWidgetComponent* WidgetComponent = Component.Get() )
+			{
+				if ( ULocalPlayer* LocalPlayer = WidgetComponent->GetOwnerPlayer() )
+				{
+					if ( APlayerController* PlayerController = LocalPlayer->PlayerController )
+					{
+						FVector WorldLocation = WidgetComponent->GetComponentLocation();
+
+						FVector ScreenPosition;
+						const bool bProjected = UWidgetLayoutLibrary::ProjectWorldLocationToWidgetPositionWithDistance(PlayerController, WorldLocation, ScreenPosition);
+
+						if ( bProjected )
+						{
+							WidgetComponent->GetUserWidgetObject()->SetVisibility(ESlateVisibility::SelfHitTestInvisible);
+
+							if ( SConstraintCanvas::FSlot* CanvasSlot = ComponentToSlot.FindRef(WidgetComponent) )
+							{
+								FVector2D DrawSize = WidgetComponent->GetDrawSize();
+								FVector2D Pivot = WidgetComponent->GetPivot();
+
+								CanvasSlot->AutoSize(DrawSize.IsZero());
+								CanvasSlot->Offset(FMargin(ScreenPosition.X, ScreenPosition.Y, DrawSize.X, DrawSize.Y));
+								CanvasSlot->Anchors(FAnchors(0, 0, 0, 0));
+								CanvasSlot->Alignment(Pivot);
+								CanvasSlot->ZOrder(-ScreenPosition.Z);
+							}
+						}
+						else
+						{
+							WidgetComponent->GetUserWidgetObject()->SetVisibility(ESlateVisibility::Hidden);
+						}
+					}
+				}
+			}
+			else
+			{
+				DeadComponents.Add(WidgetComponent);
+			}
+		}
+
+		// Normally components should be removed by someone calling remove component, but just in case it was 
+		// deleted in a way where they didn't happen, this is our backup solution to enure we remove stale widgets.
+		for ( int32 Index = 0; Index < DeadComponents.Num(); Index++ )
+		{
+			RemoveComponent(DeadComponents[Index]);
+		}
+	}
+
+	void AddComponent(UWidgetComponent* Component)
+	{
+		Components.AddUnique(Component);
+
+		SConstraintCanvas::FSlot* CanvasSlot;
+
+		Canvas->AddSlot()
+		.Expose(CanvasSlot)
+		[
+			Component->GetUserWidgetObject()->TakeWidget()
+		];
+
+		ComponentToSlot.Add(Component, CanvasSlot);
+	}
+
+	void RemoveComponent(UWidgetComponent* Component)
+	{
+		Components.RemoveSwap(Component);
+
+		TSharedPtr<SWidget> CachedWidget = Component->GetUserWidgetObject()->GetCachedWidget();
+		if ( CachedWidget.IsValid() )
+		{
+			Canvas->RemoveSlot(CachedWidget.ToSharedRef());
+		}
+
+		ComponentToSlot.Remove(Component);
+	}
+
+private:
+	TArray<TWeakObjectPtr<UWidgetComponent>> Components;
+	TMap<UWidgetComponent*, SConstraintCanvas::FSlot*> ComponentToSlot;
+	TSharedPtr<SConstraintCanvas> Canvas;
+};
+
+class FWorldWidgetScreenLayer : public IGameLayer
+{
+public:
+	void AddComponent(UWidgetComponent* Component)
+	{
+		Components.AddUnique(Component);
+		if ( ScreenLayer.IsValid() )
+		{
+			ScreenLayer.Pin()->AddComponent(Component);
+		}
+	}
+
+	void RemoveComponent(UWidgetComponent* Component)
+	{
+		Components.RemoveSwap(Component);
+
+		if ( ScreenLayer.IsValid() )
+		{
+			ScreenLayer.Pin()->RemoveComponent(Component);
+		}
+	}
+
+	virtual TSharedRef<SWidget> AsWidget() override
+	{
+		if ( ScreenLayer.IsValid() )
+		{
+			return ScreenLayer.Pin().ToSharedRef();
+		}
+
+		TSharedRef<SWorldWidgetScreenLayer> NewScreenLayer = SNew(SWorldWidgetScreenLayer, Components);
+		ScreenLayer = NewScreenLayer;
+
+		return NewScreenLayer;
+	}
+
+private:
+	TWeakPtr<SWorldWidgetScreenLayer> ScreenLayer;
+	TArray<TWeakObjectPtr<UWidgetComponent>> Components;
+};
+
+
 class SVirtualWindow : public SWindow
 {
 	SLATE_BEGIN_ARGS( SVirtualWindow )
@@ -336,7 +487,6 @@ UWidgetComponent::UWidgetComponent( const FObjectInitializer& PCIP )
 	, TickWhenOffscreen( false )
 {
 	PrimaryComponentTick.bCanEverTick = true;
-	PrimaryComponentTick.TickGroup = TG_PostUpdateWork;
 	bTickInEditor = true;
 
 	RelativeRotation = FRotator::ZeroRotator;
@@ -369,7 +519,8 @@ UWidgetComponent::UWidgetComponent( const FObjectInitializer& PCIP )
 
 	Space = EWidgetSpace::World;
 	Pivot = FVector2D(0.5, 0.5);
-	ZOrder = -100;
+
+	bAddedToScreen = false;
 }
 
 FPrimitiveSceneProxy* UWidgetComponent::CreateSceneProxy()
@@ -535,7 +686,7 @@ void UWidgetComponent::ReleaseResources()
 {
 	if ( Widget )
 	{
-		Widget->RemoveFromParent();
+		RemoveWidgetFromScreen();
 		Widget->MarkPendingKill();
 		Widget = nullptr;
 	}
@@ -549,6 +700,8 @@ void UWidgetComponent::TickComponent(float DeltaTime, enum ELevelTick TickType,
 {
 #if !UE_SERVER
 
+	static const int32 LayerZOrder = -100;
+
 	UpdateWidget();
 
 	if ( Widget == nullptr )
@@ -566,7 +719,7 @@ void UWidgetComponent::TickComponent(float DeltaTime, enum ELevelTick TickType,
 		const float RenderTimeThreshold = .5f;
 		if ( IsVisible() )
 		{
-			// if we don't tick when offscreen, don't bother ticking if it hasn't been rendered recently
+			// If we don't tick when off-screen, don't bother ticking if it hasn't been rendered recently
 			if ( TickWhenOffscreen || GetWorld()->TimeSince(LastRenderTime) <= RenderTimeThreshold )
 			{
 				SlateWidget->SlatePrepass();
@@ -619,40 +772,48 @@ void UWidgetComponent::TickComponent(float DeltaTime, enum ELevelTick TickType,
 	{
 		if ( Widget && !Widget->IsDesignTime() )
 		{
-			ULocalPlayer* TargetPlayer = OwnerPlayer ? OwnerPlayer : GEngine->GetLocalPlayerFromControllerId(GetWorld(), 0);
+			UWorld* ThisWorld = GetWorld();
+
+			ULocalPlayer* TargetPlayer = GetOwnerPlayer();
 			APlayerController* PlayerController = TargetPlayer ? TargetPlayer->PlayerController : nullptr;
 
 			if ( TargetPlayer && PlayerController && IsVisible() )
 			{
-				FVector WorldLocation = GetComponentLocation();
-
-				FVector2D ScreenPosition;
-				const bool bProjected = UWidgetLayoutLibrary::ProjectWorldLocationToWidgetPosition(
-					PlayerController, WorldLocation, ScreenPosition);
-
-				if ( bProjected )
+				if ( !bAddedToScreen )
 				{
-					Widget->SetVisibility(ESlateVisibility::SelfHitTestInvisible);
-					
-					Widget->SetDesiredSizeInViewport(DrawSize);
-					Widget->SetPositionInViewport(ScreenPosition, false);
-					Widget->SetAlignmentInViewport(Pivot);
-				}
-				else
-				{
-					Widget->SetVisibility(ESlateVisibility::Hidden);
-				}
+					if ( ThisWorld->IsGameWorld() )
+					{
+						if ( UGameViewportClient* ViewportClient = World->GetGameViewport() )
+						{
+							TSharedPtr<IGameLayerManager> LayerManager = ViewportClient->GetGameLayerManager();
+							if ( LayerManager.IsValid() )
+							{
+								TSharedPtr<FWorldWidgetScreenLayer> ScreenLayer;
 
-				if ( !Widget->IsInViewport() )
-				{
-					Widget->SetPlayerContext(TargetPlayer);
-					Widget->AddToPlayerScreen(ZOrder);
+								TSharedPtr<IGameLayer> Layer = LayerManager->FindLayerForPlayer(TargetPlayer, SharedLayerName);
+								if ( !Layer.IsValid() )
+								{
+									TSharedRef<FWorldWidgetScreenLayer> NewScreenLayer = MakeShareable(new FWorldWidgetScreenLayer());
+									LayerManager->AddLayerForPlayer(TargetPlayer, SharedLayerName, NewScreenLayer, LayerZOrder);
+									ScreenLayer = NewScreenLayer;
+								}
+								else
+								{
+									ScreenLayer = StaticCastSharedPtr<FWorldWidgetScreenLayer>(Layer);
+								}
+								
+								bAddedToScreen = true;
+								
+								Widget->SetPlayerContext(TargetPlayer);
+								ScreenLayer->AddComponent(this);
+							}
+						}
+					}
 				}
 			}
-			else if ( Widget->IsInViewport() )
+			else if ( bAddedToScreen )
 			{
-				// If the component isn't visible, and the widget is on the viewport, remove it.
-				Widget->RemoveFromParent();
+				RemoveWidgetFromScreen();
 			}
 		}
 	}
@@ -660,6 +821,29 @@ void UWidgetComponent::TickComponent(float DeltaTime, enum ELevelTick TickType,
 #endif // !UE_SERVER
 }
 
+void UWidgetComponent::RemoveWidgetFromScreen()
+{
+#if !UE_SERVER
+	bAddedToScreen = false;
+
+	if ( UGameViewportClient* ViewportClient = World->GetGameViewport() )
+	{
+		TSharedPtr<IGameLayerManager> LayerManager = ViewportClient->GetGameLayerManager();
+		if ( LayerManager.IsValid() )
+		{
+			ULocalPlayer* TargetPlayer = GetOwnerPlayer();
+
+			TSharedPtr<IGameLayer> Layer = LayerManager->FindLayerForPlayer(TargetPlayer, SharedLayerName);
+			if ( Layer.IsValid() )
+			{
+				TSharedPtr<FWorldWidgetScreenLayer> ScreenLayer = StaticCastSharedPtr<FWorldWidgetScreenLayer>(Layer);
+				ScreenLayer->RemoveComponent(this);
+			}
+		}
+	}
+#endif // !UE_SERVER
+}
+
 class FWidgetComponentInstanceData : public FSceneComponentInstanceData
 {
 public:
@@ -788,14 +972,14 @@ void UWidgetComponent::SetOwnerPlayer(ULocalPlayer* LocalPlayer)
 
 ULocalPlayer* UWidgetComponent::GetOwnerPlayer() const
 {
-	return OwnerPlayer;
+	return OwnerPlayer ? OwnerPlayer : GEngine->GetLocalPlayerFromControllerId(GetWorld(), 0);
 }
 
 void UWidgetComponent::SetWidget(UUserWidget* InWidget)
 {
 	if ( Widget )
 	{
-		Widget->RemoveFromParent();
+		RemoveWidgetFromScreen();
 		Widget->MarkPendingKill();
 		Widget = nullptr;
 	}
diff --git a/Engine/Source/Runtime/UMG/Private/WidgetLayoutLibrary.cpp b/Engine/Source/Runtime/UMG/Private/WidgetLayoutLibrary.cpp
--- a/Engine/Source/Runtime/UMG/Private/WidgetLayoutLibrary.cpp
+++ b/Engine/Source/Runtime/UMG/Private/WidgetLayoutLibrary.cpp
@@ -17,11 +17,19 @@ UWidgetLayoutLibrary::UWidgetLayoutLibrary(const FObjectInitializer& ObjectIniti
 }
 
 bool UWidgetLayoutLibrary::ProjectWorldLocationToWidgetPosition(APlayerController* PlayerController, FVector WorldLocation, FVector2D& ScreenPosition)
+{
+	FVector ScreenPosition3D;
+	const bool bSuccess = ProjectWorldLocationToWidgetPositionWithDistance(PlayerController, WorldLocation, ScreenPosition3D);
+	ScreenPosition = FVector2D(ScreenPosition3D.X, ScreenPosition3D.Y);
+	return bSuccess;
+}
+
+bool UWidgetLayoutLibrary::ProjectWorldLocationToWidgetPositionWithDistance(APlayerController* PlayerController, FVector WorldLocation, FVector& ScreenPosition)
 {
 	if ( PlayerController )
 	{
-		FVector2D PixelLocation;
-		const bool bProjected = PlayerController->ProjectWorldLocationToScreen(WorldLocation, PixelLocation);
+		FVector PixelLocation;
+		const bool bProjected = PlayerController->ProjectWorldLocationToScreenWithDistance(WorldLocation, PixelLocation);
 
 		if ( bProjected )
 		{
@@ -41,6 +49,7 @@ bool UWidgetLayoutLibrary::ProjectWorldLocationToWidgetPosition(APlayerControlle
 			// result in more or less the same value, especially after slate does layout rounding.
 			ScreenPosition.X = FMath::RoundToInt(ScreenPosition.X);
 			ScreenPosition.Y = FMath::RoundToInt(ScreenPosition.Y);
+			ScreenPosition.Z = PixelLocation.Z;
 
 			// Get the application / DPI scale
 			const float Scale = UWidgetLayoutLibrary::GetViewportScale(PlayerController);
diff --git a/Engine/Source/Runtime/UMG/Public/Blueprint/WidgetLayoutLibrary.h b/Engine/Source/Runtime/UMG/Public/Blueprint/WidgetLayoutLibrary.h
--- a/Engine/Source/Runtime/UMG/Public/Blueprint/WidgetLayoutLibrary.h
+++ b/Engine/Source/Runtime/UMG/Public/Blueprint/WidgetLayoutLibrary.h
@@ -22,6 +22,14 @@ public:
 	UFUNCTION(BlueprintPure, BlueprintCosmetic, Category="Viewport")
 	static bool ProjectWorldLocationToWidgetPosition(APlayerController* PlayerController, FVector WorldLocation, FVector2D& ScreenPosition);
 
+	/**
+	 * Convert a World Space 3D position into a 2D Widget Screen Space position, with distance from the camera the Z component.  This
+	 * takes into account any quality scaling as well.
+	 *
+	 * @return true if the world coordinate was successfully projected to the screen.
+	 */
+	static bool ProjectWorldLocationToWidgetPositionWithDistance(APlayerController* PlayerController, FVector WorldLocation, FVector& ScreenPosition);
+
 	/**
 	 * Gets the current DPI Scale being applied to the viewport and all the Widgets.
 	 */
diff --git a/Engine/Source/Runtime/UMG/Public/Components/WidgetComponent.h b/Engine/Source/Runtime/UMG/Public/Components/WidgetComponent.h
--- a/Engine/Source/Runtime/UMG/Public/Components/WidgetComponent.h
+++ b/Engine/Source/Runtime/UMG/Public/Components/WidgetComponent.h
@@ -144,6 +144,10 @@ public:
 	
 	/** Sets the widget class used to generate the widget for this component */
 	void SetWidgetClass(TSubclassOf<UUserWidget> InWidgetClass);
+
+protected:
+	void RemoveWidgetFromScreen();
+
 protected:
 	/** The coordinate space in which to render the widget */
 	UPROPERTY(EditAnywhere, Category=UI)
@@ -160,10 +164,6 @@ protected:
 	/** The Alignment/Pivot point that the widget is placed at relative to the position. */
 	UPROPERTY(EditAnywhere, Category=UI)
 	FVector2D Pivot;
-
-	/** The Screen Space ZOrder that is used if the widget is displayed in screen space. */
-	UPROPERTY(EditAnywhere, Category=UI)
-	int32 ZOrder;
 	
 	/** The maximum distance from which a player can interact with this widget */
 	UPROPERTY(EditAnywhere, Category=UI, meta=(ClampMin="0.0", UIMax="5000.0", ClampMax="100000.0"))
@@ -245,6 +245,9 @@ protected:
 
 	UPROPERTY()
 	bool bUseLegacyRotation;
+
+	UPROPERTY(Transient, DuplicateTransient)
+	bool bAddedToScreen;
 protected:
 
 	/** The grid used to find actual hit actual widgets once input has been translated to the components local space */
