- Fixed clang CIS error

diff --git a/Engine/Source/Runtime/UMG/Private/Components/WidgetComponent.cpp b/Engine/Source/Runtime/UMG/Private/Components/WidgetComponent.cpp
--- a/Engine/Source/Runtime/UMG/Private/Components/WidgetComponent.cpp
+++ b/Engine/Source/Runtime/UMG/Private/Components/WidgetComponent.cpp
@@ -129,6 +129,11 @@ private:
 class FWorldWidgetScreenLayer : public IGameLayer
 {
 public:
+	virtual ~FWorldWidgetScreenLayer()
+	{
+		// empty virtual destructor to help clang warning
+	}
+	
 	void AddComponent(UWidgetComponent* Component)
 	{
 		Components.AddUnique(Component);
