UMG: Fixed non-unity

diff --git a/Engine/Source/Runtime/UMG/Public/Components/TableViewBase.h b/Engine/Source/Runtime/UMG/Public/Components/TableViewBase.h
--- a/Engine/Source/Runtime/UMG/Public/Components/TableViewBase.h
+++ b/Engine/Source/Runtime/UMG/Public/Components/TableViewBase.h
@@ -2,6 +2,7 @@
 
 #pragma once
 
+#include "Widget.h"
 #include "TableViewBase.generated.h"
 
 /** The base class for all wrapped table views */
