UMG - Removing an unneeded return.

diff --git a/Engine/Source/Runtime/UMG/Private/SlateBlueprintLibrary.cpp b/Engine/Source/Runtime/UMG/Private/SlateBlueprintLibrary.cpp
--- a/Engine/Source/Runtime/UMG/Private/SlateBlueprintLibrary.cpp
+++ b/Engine/Source/Runtime/UMG/Private/SlateBlueprintLibrary.cpp
@@ -39,7 +39,7 @@ FVector2D USlateBlueprintLibrary::GetLocalSize(const FGeometry& Geometry)
 void USlateBlueprintLibrary::LocalToViewport(UObject* WorldContextObject, const FGeometry& Geometry, FVector2D LocalCoordinate, FVector2D& ScreenPosition, FVector2D& ViewportPosition)
 {
 	FVector2D AbsoluteCoordinate = Geometry.LocalToAbsolute(LocalCoordinate);
-	return AbsoluteToViewport(WorldContextObject, Geometry, AbsoluteCoordinate, ScreenPosition, ViewportPosition);
+	AbsoluteToViewport(WorldContextObject, Geometry, AbsoluteCoordinate, ScreenPosition, ViewportPosition);
 }
 
 void USlateBlueprintLibrary::AbsoluteToViewport(UObject* WorldContextObject, const FGeometry& Geometry, FVector2D AbsoluteCoordinate, FVector2D& ScreenPosition, FVector2D& ViewportPosition)
