UMG - UE-17747 - Routing PostLoad in MultiLineEditableText

diff --git a/Engine/Source/Runtime/UMG/Private/Components/MultiLineEditableText.cpp b/Engine/Source/Runtime/UMG/Private/Components/MultiLineEditableText.cpp
--- a/Engine/Source/Runtime/UMG/Private/Components/MultiLineEditableText.cpp
+++ b/Engine/Source/Runtime/UMG/Private/Components/MultiLineEditableText.cpp
@@ -94,6 +94,8 @@ void UMultiLineEditableText::HandleOnTextCommitted(const FText& InText, ETextCom
 
 void UMultiLineEditableText::PostLoad()
 {
+	Super::PostLoad();
+
 	if (GetLinkerUE4Version() < VER_UE4_DEPRECATE_UMG_STYLE_OVERRIDES)
 	{
 		if (Font_DEPRECATED.HasValidFont())
