UMG - Updating the DragDropOperation and UserWidget classes to use the DontUseGenericSpawnObject metadata flag so that they are not spawnable with the new class spawner in blueprints for standard UObjects.

diff --git a/Engine/Source/Runtime/UMG/Public/Blueprint/DragDropOperation.h b/Engine/Source/Runtime/UMG/Public/Blueprint/DragDropOperation.h
--- a/Engine/Source/Runtime/UMG/Public/Blueprint/DragDropOperation.h
+++ b/Engine/Source/Runtime/UMG/Public/Blueprint/DragDropOperation.h
@@ -30,7 +30,7 @@ enum class EDragPivot : uint8
 /**
  * This class is the base drag drop operation for UMG, extend it to add additional data and add new functionality.
  */
-UCLASS(BlueprintType, Blueprintable)
+UCLASS(BlueprintType, Blueprintable, meta=( DontUseGenericSpawnObject="True" ))
 class UMG_API UDragDropOperation : public UObject
 {
 	GENERATED_UCLASS_BODY()
diff --git a/Engine/Source/Runtime/UMG/Public/Blueprint/UserWidget.h b/Engine/Source/Runtime/UMG/Public/Blueprint/UserWidget.h
--- a/Engine/Source/Runtime/UMG/Public/Blueprint/UserWidget.h
+++ b/Engine/Source/Runtime/UMG/Public/Blueprint/UserWidget.h
@@ -137,7 +137,7 @@ DECLARE_DYNAMIC_MULTICAST_DELEGATE(FOnConstructEvent);
 /**
  * The user widget is extensible by users through the WidgetBlueprint.
  */
-UCLASS(Abstract, editinlinenew, BlueprintType, Blueprintable, meta=( Category="User Controls" ))
+UCLASS(Abstract, editinlinenew, BlueprintType, Blueprintable, meta=( Category="User Controls", DontUseGenericSpawnObject="True" ) )
 class UMG_API UUserWidget : public UWidget, public INamedSlotInterface
 {
 	GENERATED_BODY()
