Renamed "UI" category name to "UserInterface" in the widget component so that it is consistent with other UMG category names

diff --git a/Engine/Source/Runtime/UMG/Public/Components/WidgetComponent.h b/Engine/Source/Runtime/UMG/Public/Components/WidgetComponent.h
--- a/Engine/Source/Runtime/UMG/Public/Components/WidgetComponent.h
+++ b/Engine/Source/Runtime/UMG/Public/Components/WidgetComponent.h
@@ -84,7 +84,7 @@ public:
 	TSubclassOf<UUserWidget> GetWidgetClass() const { return WidgetClass; }
 
 	/** @return The user widget object displayed by this component */
-	UFUNCTION(BlueprintCallable, Category=UI)
+	UFUNCTION(BlueprintCallable, Category=UserInterface)
 	UUserWidget* GetUserWidgetObject() const;
 
 	/** @return List of widgets with their geometry and the cursor position transformed into this Widget component's space. */
@@ -100,31 +100,31 @@ public:
 	TSharedPtr<SWidget> GetSlateWidget() const;
 
 	/** Sets the widget to use directly. */
-	UFUNCTION(BlueprintCallable, Category=UI)
+	UFUNCTION(BlueprintCallable, Category=UserInterface)
 	void SetWidget(UUserWidget* Widget);
 
 	/** Sets the local player that owns this widget component. */
-	UFUNCTION(BlueprintCallable, Category=UI)
+	UFUNCTION(BlueprintCallable, Category=UserInterface)
 	void SetOwnerPlayer(ULocalPlayer* LocalPlayer);
 
 	/** Gets the local player that owns this widget component. */
-	UFUNCTION(BlueprintCallable, Category=UI)
+	UFUNCTION(BlueprintCallable, Category=UserInterface)
 	ULocalPlayer* GetOwnerPlayer() const;
 
 	/** @return The draw size of the quad in the world */
-	UFUNCTION(BlueprintCallable, Category=UI)
+	UFUNCTION(BlueprintCallable, Category=UserInterface)
 	FVector2D GetDrawSize() const;
 
 	/** Sets the draw size of the quad in the world */
-	UFUNCTION(BlueprintCallable, Category=UI)
+	UFUNCTION(BlueprintCallable, Category=UserInterface)
 	void SetDrawSize(FVector2D Size);
 
 	/** @return The max distance from which a player can interact with this widget */
-	UFUNCTION(BlueprintCallable, Category=UI)
+	UFUNCTION(BlueprintCallable, Category=UserInterface)
 	float GetMaxInteractionDistance() const;
 
 	/** Sets the max distance from which a player can interact with this widget */
-	UFUNCTION(BlueprintCallable, Category=UI)
+	UFUNCTION(BlueprintCallable, Category=UserInterface)
 	void SetMaxInteractionDistance(float Distance);
 
 	/** Gets the blend mode for the widget. */
@@ -150,23 +150,23 @@ protected:
 
 protected:
 	/** The coordinate space in which to render the widget */
-	UPROPERTY(EditAnywhere, Category=UI)
+	UPROPERTY(EditAnywhere, Category=UserInterface)
 	EWidgetSpace Space;
 
 	/** The class of User Widget to create and display an instance of */
-	UPROPERTY(EditAnywhere, Category=UI)
+	UPROPERTY(EditAnywhere, Category=UserInterface)
 	TSubclassOf<UUserWidget> WidgetClass;
 	
 	/** The size of the displayed quad. */
-	UPROPERTY(EditAnywhere, Category=UI)
+	UPROPERTY(EditAnywhere, Category=UserInterface)
 	FIntPoint DrawSize;
 
 	/** The Alignment/Pivot point that the widget is placed at relative to the position. */
-	UPROPERTY(EditAnywhere, Category=UI)
+	UPROPERTY(EditAnywhere, Category=UserInterface)
 	FVector2D Pivot;
 	
 	/** The maximum distance from which a player can interact with this widget */
-	UPROPERTY(EditAnywhere, Category=UI, meta=(ClampMin="0.0", UIMax="5000.0", ClampMax="100000.0"))
+	UPROPERTY(EditAnywhere, Category=UserInterface, meta=(ClampMin="0.0", UIMax="5000.0", ClampMax="100000.0"))
 	float MaxInteractionDistance;
 
 	/**
