UMG/Sequencer - Replace usage of UMovieSceneBindings object with a simple map.

diff --git a/Engine/Source/Runtime/UMG/Private/Animation/UMGSequencePlayer.cpp b/Engine/Source/Runtime/UMG/Private/Animation/UMGSequencePlayer.cpp
--- a/Engine/Source/Runtime/UMG/Private/Animation/UMGSequencePlayer.cpp
+++ b/Engine/Source/Runtime/UMG/Private/Animation/UMGSequencePlayer.cpp
@@ -25,12 +25,8 @@ void UUMGSequencePlayer::InitSequencePlayer( const UWidgetAnimation& InAnimation
 	// Cache the time range of the sequence to determine when we stop
 	TimeRange = MovieScene->GetTimeRange();
 
-	RuntimeBindings = NewObject<UMovieSceneBindings>(this);
-	RuntimeBindings->SetRootMovieScene( MovieScene );
-
 	UWidgetTree* WidgetTree = UserWidget.WidgetTree;
 
-	TMap<FGuid, TArray<UObject*> > GuidToRuntimeObjectMap;
 	// Bind to Runtime Objects
 	for (const FWidgetAnimationBinding& Binding : InAnimation.AnimationBindings)
 	{
@@ -42,12 +38,6 @@ void UUMGSequencePlayer::InitSequencePlayer( const UWidgetAnimation& InAnimation
 			Objects.Add(FoundObject);
 		}
 	}
-
-	for( auto It = GuidToRuntimeObjectMap.CreateConstIterator(); It; ++It )
-	{
-		RuntimeBindings->AddBinding( It.Key(), It.Value() );
-	}
-
 }
 
 void UUMGSequencePlayer::Tick(float DeltaTime)
@@ -160,7 +150,7 @@ void UUMGSequencePlayer::Stop()
 
 void UUMGSequencePlayer::GetRuntimeObjects( TSharedRef<FMovieSceneInstance> MovieSceneInstance, const FGuid& ObjectHandle, TArray< UObject* >& OutObjects ) const
 {
-	OutObjects = RuntimeBindings->FindBoundObjects( ObjectHandle );
+	OutObjects.Append(GuidToRuntimeObjectMap[ObjectHandle]);
 }
 
 EMovieScenePlayerStatus::Type UUMGSequencePlayer::GetPlaybackStatus() const
diff --git a/Engine/Source/Runtime/UMG/Public/Animation/UMGSequencePlayer.h b/Engine/Source/Runtime/UMG/Public/Animation/UMGSequencePlayer.h
--- a/Engine/Source/Runtime/UMG/Public/Animation/UMGSequencePlayer.h
+++ b/Engine/Source/Runtime/UMG/Public/Animation/UMGSequencePlayer.h
@@ -50,9 +50,7 @@ private:
 	UPROPERTY()
 	const UWidgetAnimation* Animation;
 
-	/** Bindings to actual live objects */
-	UPROPERTY()
-	UMovieSceneBindings* RuntimeBindings;
+	TMap<FGuid, TArray<UObject*> > GuidToRuntimeObjectMap;
 
 	/** The root movie scene instance to update when playing. */
 	TSharedPtr<class FMovieSceneInstance> RootMovieSceneInstance;
