UMG - If the owner player changes on a widget component, make sure to remove the widget from the screen if it's currently already on a user's screen.

diff --git a/Engine/Source/Runtime/UMG/Private/Components/WidgetComponent.cpp b/Engine/Source/Runtime/UMG/Private/Components/WidgetComponent.cpp
--- a/Engine/Source/Runtime/UMG/Private/Components/WidgetComponent.cpp
+++ b/Engine/Source/Runtime/UMG/Private/Components/WidgetComponent.cpp
@@ -972,7 +972,11 @@ void UWidgetComponent::InitWidget()
 
 void UWidgetComponent::SetOwnerPlayer(ULocalPlayer* LocalPlayer)
 {
-	OwnerPlayer = LocalPlayer;
+	if ( OwnerPlayer != LocalPlayer )
+	{
+		RemoveWidgetFromScreen();
+		OwnerPlayer = LocalPlayer;
+	}
 }
 
 ULocalPlayer* UWidgetComponent::GetOwnerPlayer() const
diff --git a/Engine/Source/Runtime/UMG/Public/Components/WidgetComponent.h b/Engine/Source/Runtime/UMG/Public/Components/WidgetComponent.h
--- a/Engine/Source/Runtime/UMG/Public/Components/WidgetComponent.h
+++ b/Engine/Source/Runtime/UMG/Public/Components/WidgetComponent.h
@@ -103,7 +103,11 @@ public:
 	UFUNCTION(BlueprintCallable, Category=UserInterface)
 	void SetWidget(UUserWidget* Widget);
 
-	/** Sets the local player that owns this widget component. */
+	/**
+	 * Sets the local player that owns this widget component.  Setting the owning player controls
+	 * which player's viewport the widget appears on in a split screen scenario.  Additionally it
+	 * forwards the owning player to the actual UserWidget that is spawned.
+	 */
 	UFUNCTION(BlueprintCallable, Category=UserInterface)
 	void SetOwnerPlayer(ULocalPlayer* LocalPlayer);
 
