UMG - Drag/Drop for UMG now properly maps the cursor to a user widget and no longer shows the native cursor in addition to the custom cursor if any.  Also fixed a crash issue on mobile where the platform doesn't have a cursor object, be we still attempted to interact with it on drag/drop operations.

diff --git a/Engine/Source/Editor/LevelEditor/Private/SLevelViewport.cpp b/Engine/Source/Editor/LevelEditor/Private/SLevelViewport.cpp
--- a/Engine/Source/Editor/LevelEditor/Private/SLevelViewport.cpp
+++ b/Engine/Source/Editor/LevelEditor/Private/SLevelViewport.cpp
@@ -3273,7 +3273,8 @@ void SLevelViewport::StartPlayInEditorSession(UGameViewportClient* PlayClient, c
 	
 	ActiveViewport->OnPlayWorldViewportSwapped( *InactiveViewport );
 
-	PlayClient->SetViewportOverlayWidget( TSharedPtr<SWindow>(), PIEViewportOverlayWidget.ToSharedRef() );
+	TSharedPtr<SWindow> ParentWindow = FSlateApplication::Get().FindWidgetWindow(AsShared());
+	PlayClient->SetViewportOverlayWidget(ParentWindow, PIEViewportOverlayWidget.ToSharedRef());
 	PlayClient->SetGameLayerManager(GameLayerManager);
 
 	// Our viewport widget should start rendering the new viewport for the play in editor scene
diff --git a/Engine/Source/Runtime/Slate/Private/Framework/Application/SlateApplication.cpp b/Engine/Source/Runtime/Slate/Private/Framework/Application/SlateApplication.cpp
--- a/Engine/Source/Runtime/Slate/Private/Framework/Application/SlateApplication.cpp
+++ b/Engine/Source/Runtime/Slate/Private/Framework/Application/SlateApplication.cpp
@@ -4732,18 +4732,17 @@ bool FSlateApplication::ProcessMouseMoveEvent( FPointerEvent& MouseEvent, bool b
 			DragDropWindowPtr = nullptr;
 		}
 
-		// check the drag-drop operation for a cursor switch (on Windows, the OS thinks the mouse is
-		// captured so we wont get QueryCursor calls for drag/drops internal to the Slate application)
-		FCursorReply CursorResult = DragDropContent->OnCursorQuery();
-		if (CursorResult.IsEventHandled())
-		{
-			// Query was handled, so we should set the cursor.
-			PlatformApplication->Cursor->SetType(CursorResult.GetCursorType());
-		}
-		else
+		// Don't update the cursor for the platform if we don't have a valid cursor on this platform
+		if ( PlatformApplication->Cursor.IsValid() )
 		{
-			// reset the cursor to default for drag-drops
-			PlatformApplication->Cursor->SetType( EMouseCursor::Default );
+			FCursorReply CursorReply = DragDropContent->OnCursorQuery();
+			if ( !CursorReply.IsEventHandled() )
+			{
+				// Set the default cursor when there isn't an active window under the cursor and the mouse isn't captured
+				CursorReply = FCursorReply::Cursor(EMouseCursor::Default);
+			}
+
+			ProcessCursorReply(CursorReply);
 		}
 	}
 	else
diff --git a/Engine/Source/Runtime/UMG/Private/Slate/UMGDragDropOp.cpp b/Engine/Source/Runtime/UMG/Private/Slate/UMGDragDropOp.cpp
--- a/Engine/Source/Runtime/UMG/Private/Slate/UMGDragDropOp.cpp
+++ b/Engine/Source/Runtime/UMG/Private/Slate/UMGDragDropOp.cpp
@@ -17,6 +17,7 @@ FUMGDragDropOp::FUMGDragDropOp()
 void FUMGDragDropOp::AddReferencedObjects(FReferenceCollector& Collector)
 {
 	Collector.AddReferencedObject(DragOperation);
+	Collector.AddReferencedObject(GameViewport);
 }
 
 void FUMGDragDropOp::Construct()
@@ -117,6 +118,27 @@ TSharedPtr<SWidget> FUMGDragDropOp::GetDefaultDecorator() const
 	return DecoratorWidget;
 }
 
+FCursorReply FUMGDragDropOp::OnCursorQuery()
+{
+	FCursorReply CursorReply = FGameDragDropOperation::OnCursorQuery();
+
+	if ( !CursorReply.IsEventHandled() )
+	{
+		CursorReply = CursorReply.Cursor(EMouseCursor::Default);
+	}
+
+	if ( GameViewport )
+	{
+		TOptional<TSharedRef<SWidget>> CursorWidget = GameViewport->MapCursor(nullptr, CursorReply);
+		if ( CursorWidget.IsSet() )
+		{
+			CursorReply.SetCursorWidget(GameViewport->GetWindow(), CursorWidget.GetValue());
+		}
+	}
+
+	return CursorReply;
+}
+
 TSharedRef<FUMGDragDropOp> FUMGDragDropOp::New(UDragDropOperation* InOperation, const FVector2D &CursorPosition, const FVector2D &ScreenPositionOfDragee, float DPIScale, TSharedPtr<SObjectWidget> SourceUserWidget)
 {
 	check(InOperation);
@@ -125,6 +147,7 @@ TSharedRef<FUMGDragDropOp> FUMGDragDropOp::New(UDragDropOperation* InOperation,
 	Operation->MouseDownOffset = ScreenPositionOfDragee - CursorPosition;
 	Operation->StartingScreenPos = ScreenPositionOfDragee;
 	Operation->SourceUserWidget = SourceUserWidget;
+	Operation->GameViewport = SourceUserWidget->GetWidgetObject()->GetWorld()->GetGameViewport();
 	Operation->DragOperation = InOperation;
 
 	TSharedPtr<SWidget> DragVisual;
@@ -135,6 +158,7 @@ TSharedRef<FUMGDragDropOp> FUMGDragDropOp::New(UDragDropOperation* InOperation,
 	}
 	else
 	{
+		//TODO Make sure users are not trying to add a widget that already exists elsewhere.
 		DragVisual = InOperation->DefaultDragVisual->TakeWidget();
 	}
 
diff --git a/Engine/Source/Runtime/UMG/Private/WidgetLayoutLibrary.cpp b/Engine/Source/Runtime/UMG/Private/WidgetLayoutLibrary.cpp
--- a/Engine/Source/Runtime/UMG/Private/WidgetLayoutLibrary.cpp
+++ b/Engine/Source/Runtime/UMG/Private/WidgetLayoutLibrary.cpp
@@ -25,7 +25,7 @@ bool UWidgetLayoutLibrary::ProjectWorldLocationToWidgetPosition(APlayerControlle
 }
 
 bool UWidgetLayoutLibrary::ProjectWorldLocationToWidgetPositionWithDistance(APlayerController* PlayerController, FVector WorldLocation, FVector& ScreenPosition)
-{
+{	
 	if ( PlayerController )
 	{
 		FVector PixelLocation;
@@ -33,15 +33,7 @@ bool UWidgetLayoutLibrary::ProjectWorldLocationToWidgetPositionWithDistance(APla
 
 		if ( bProjected )
 		{
-			// If the user has configured a resolution quality we need to multiply
-			// the pixels by the resolution quality to arrive at the true position in
-			// the viewport, as the rendered image will be stretched to fill whatever
-			// size the viewport is at.
-			Scalability::FQualityLevels ScalabilityQuality = Scalability::GetQualityLevels();
-			const float QualityScale = ( ScalabilityQuality.ResolutionQuality / 100.0f );
-
-			// Remove the resolution quality scale.
-			ScreenPosition = PixelLocation / QualityScale;
+			ScreenPosition = PixelLocation;
 
 			// Round the pixel projected value to reduce jittering due to layout rounding,
 			// I do this before I remove scaling, because scaling is going to be applied later
@@ -61,6 +53,8 @@ bool UWidgetLayoutLibrary::ProjectWorldLocationToWidgetPositionWithDistance(APla
 		}
 	}
 
+	ScreenPosition = FVector(0, 0, 0);
+
 	return false;
 }
 
diff --git a/Engine/Source/Runtime/UMG/Public/Slate/UMGDragDropOp.h b/Engine/Source/Runtime/UMG/Public/Slate/UMGDragDropOp.h
--- a/Engine/Source/Runtime/UMG/Public/Slate/UMGDragDropOp.h
+++ b/Engine/Source/Runtime/UMG/Public/Slate/UMGDragDropOp.h
@@ -22,9 +22,8 @@ public:
 	// End FGCObject
 
 	virtual void OnDrop( bool bDropWasHandled, const FPointerEvent& MouseEvent ) override;
-
 	virtual void OnDragged( const class FDragDropEvent& DragDropEvent ) override;
-
+	virtual FCursorReply OnCursorQuery() override;
 	virtual TSharedPtr<SWidget> GetDefaultDecorator() const override;
 
 	UDragDropOperation* GetOperation() const { return DragOperation; }
@@ -40,6 +39,9 @@ private:
 	/** Source User Widget */
 	TSharedPtr<SObjectWidget> SourceUserWidget;
 
+	/** The viewport this drag/drop operation is associated with. */
+	UGameViewportClient* GameViewport;
+
 	/** The widget used during the drag/drop action to show something being dragged. */
 	TSharedPtr<SWidget> DecoratorWidget;
 
