UMG - UE-18029 - Making the SlotAs functions more robust against crashing when called at the wrong time from blueprints.

diff --git a/Engine/Source/Runtime/UMG/Private/WidgetLayoutLibrary.cpp b/Engine/Source/Runtime/UMG/Private/WidgetLayoutLibrary.cpp
--- a/Engine/Source/Runtime/UMG/Private/WidgetLayoutLibrary.cpp
+++ b/Engine/Source/Runtime/UMG/Private/WidgetLayoutLibrary.cpp
@@ -106,32 +106,62 @@ FVector2D UWidgetLayoutLibrary::GetViewportSize(UObject* WorldContextObject)
 
 UCanvasPanelSlot* UWidgetLayoutLibrary::SlotAsCanvasSlot(UWidget* ChildWidget)
 {
-	return Cast<UCanvasPanelSlot>(ChildWidget->Slot);
+	if ( ChildWidget )
+	{
+		return Cast<UCanvasPanelSlot>(ChildWidget->Slot);
+	}
+
+	return nullptr;
 }
 
 UGridSlot* UWidgetLayoutLibrary::SlotAsGridSlot(UWidget* ChildWidget)
 {
-	return Cast<UGridSlot>(ChildWidget->Slot);
+	if ( ChildWidget )
+	{
+		return Cast<UGridSlot>(ChildWidget->Slot);
+	}
+
+	return nullptr;
 }
 
 UHorizontalBoxSlot* UWidgetLayoutLibrary::SlotAsHorizontalBoxSlot(UWidget* ChildWidget)
 {
-	return Cast<UHorizontalBoxSlot>(ChildWidget->Slot);
+	if ( ChildWidget )
+	{
+		return Cast<UHorizontalBoxSlot>(ChildWidget->Slot);
+	}
+
+	return nullptr;
 }
 
 UOverlaySlot* UWidgetLayoutLibrary::SlotAsOverlaySlot(UWidget* ChildWidget)
 {
-	return Cast<UOverlaySlot>(ChildWidget->Slot);
+	if ( ChildWidget )
+	{
+		return Cast<UOverlaySlot>(ChildWidget->Slot);
+	}
+
+	return nullptr;
 }
 
 UUniformGridSlot* UWidgetLayoutLibrary::SlotAsUniformGridSlot(UWidget* ChildWidget)
 {
-	return Cast<UUniformGridSlot>(ChildWidget->Slot);
+	if ( ChildWidget )
+	{
+		return Cast<UUniformGridSlot>(ChildWidget->Slot);
+	}
+
+	return nullptr;
 }
 
 UVerticalBoxSlot* UWidgetLayoutLibrary::SlotAsVerticalBoxSlot(UWidget* ChildWidget)
 {
-	return Cast<UVerticalBoxSlot>(ChildWidget->Slot);
+	if ( ChildWidget )
+	{
+		return Cast<UVerticalBoxSlot>(ChildWidget->Slot);
+	}
+
+	return nullptr;
 }
 
 void UWidgetLayoutLibrary::RemoveAllWidgets(UObject* WorldContextObject)
