UMG - Updating a comment about invalidation.

diff --git a/Engine/Source/Runtime/UMG/Public/Components/Widget.h b/Engine/Source/Runtime/UMG/Public/Components/Widget.h
--- a/Engine/Source/Runtime/UMG/Public/Components/Widget.h
+++ b/Engine/Source/Runtime/UMG/Public/Components/Widget.h
@@ -335,7 +335,8 @@ public:
 	void ForceLayoutPrepass();
 
 	/**
-	 *
+	 * Invalidates the widget from the view of a layout caching widget that may own this widget.
+	 * will force the owning widget to redraw and cache children on the next paint pass.
 	 */
 	UFUNCTION(BlueprintCallable, Category="Widget")
 	void InvalidateLayout();
