Fix crash if no runtime objects can be found and log which animation is missing objects to track it down

https://jira.ol.epicgames.net/browse/UE-18183

diff --git a/Engine/Source/Runtime/UMG/Private/Animation/UMGSequencePlayer.cpp b/Engine/Source/Runtime/UMG/Private/Animation/UMGSequencePlayer.cpp
--- a/Engine/Source/Runtime/UMG/Private/Animation/UMGSequencePlayer.cpp
+++ b/Engine/Source/Runtime/UMG/Private/Animation/UMGSequencePlayer.cpp
@@ -8,6 +8,7 @@
 #include "MovieScene.h"
 #include "WidgetAnimation.h"
 
+
 UUMGSequencePlayer::UUMGSequencePlayer(const FObjectInitializer& ObjectInitializer)
 	: Super(ObjectInitializer)
 {
@@ -150,7 +151,16 @@ void UUMGSequencePlayer::Stop()
 
 void UUMGSequencePlayer::GetRuntimeObjects( TSharedRef<FMovieSceneInstance> MovieSceneInstance, const FGuid& ObjectHandle, TArray< UObject* >& OutObjects ) const
 {
-	OutObjects.Append(GuidToRuntimeObjectMap[ObjectHandle]);
+	const TArray<UObject*>* FoundObjects = GuidToRuntimeObjectMap.Find( ObjectHandle );
+
+	if( FoundObjects )
+	{
+		OutObjects.Append(*FoundObjects);
+	}
+	else
+	{
+		UE_LOG( LogUMG, Warning, TEXT("Failed to find runtime objects for %s animation"), Animation ? *Animation->GetName() : TEXT("(none)") );
+	}
 }
 
 EMovieScenePlayerStatus::Type UUMGSequencePlayer::GetPlaybackStatus() const
diff --git a/Engine/Source/Runtime/UMG/Private/UMGModule.cpp b/Engine/Source/Runtime/UMG/Private/UMGModule.cpp
--- a/Engine/Source/Runtime/UMG/Private/UMGModule.cpp
+++ b/Engine/Source/Runtime/UMG/Private/UMGModule.cpp
@@ -5,6 +5,9 @@
 #include "IUMGModule.h"
 #include "ModuleManager.h"
 
+
+DEFINE_LOG_CATEGORY(LogUMG);
+
 class FUMGModule : public IUMGModule
 {
 public:
diff --git a/Engine/Source/Runtime/UMG/Private/UMGPrivatePCH.h b/Engine/Source/Runtime/UMG/Private/UMGPrivatePCH.h
--- a/Engine/Source/Runtime/UMG/Private/UMGPrivatePCH.h
+++ b/Engine/Source/Runtime/UMG/Private/UMGPrivatePCH.h
@@ -5,3 +5,5 @@
 #include "UMG.h"
 
 #include "UMGStyle.h"
+
+DECLARE_LOG_CATEGORY_EXTERN(LogUMG, Log, All);
