UMG - Fixing several issues with brush and material utility functions not working correctly, pins on the wrong side and such.

diff --git a/Engine/Source/Runtime/UMG/Private/WidgetBlueprintLibrary.cpp b/Engine/Source/Runtime/UMG/Private/WidgetBlueprintLibrary.cpp
--- a/Engine/Source/Runtime/UMG/Private/WidgetBlueprintLibrary.cpp
+++ b/Engine/Source/Runtime/UMG/Private/WidgetBlueprintLibrary.cpp
@@ -400,6 +400,16 @@ UObject* UWidgetBlueprintLibrary::GetBrushResource(FSlateBrush& Brush)
 	return Brush.GetResourceObject();
 }
 
+UTexture2D* UWidgetBlueprintLibrary::GetBrushResourceAsTexture2D(UPARAM(ref) FSlateBrush& Brush)
+{
+	return Cast<UTexture2D>(Brush.GetResourceObject());
+}
+
+UMaterialInterface* UWidgetBlueprintLibrary::GetBrushResourceAsMaterial(UPARAM(ref) FSlateBrush& Brush)
+{
+	return Cast<UMaterialInterface>(Brush.GetResourceObject());
+}
+
 void UWidgetBlueprintLibrary::SetBrushResourceToTexture(FSlateBrush& Brush, UTexture2D* Texture)
 {
 	Brush.SetResourceObject(Texture);
@@ -527,8 +537,34 @@ void UWidgetBlueprintLibrary::GetAllWidgetsWithInterface(UObject* WorldContextOb
 	}
 }
 
-FInputEvent UWidgetBlueprintLibrary::GetInputEventFromKeyEvent(const FKeyEvent& KeyEvent)
+FInputEvent UWidgetBlueprintLibrary::GetInputEventFromKeyEvent(const FKeyEvent& Event)
+{
+	return Event;
+}
+
+FInputEvent UWidgetBlueprintLibrary::GetInputEventFromAnalogInputEvent(const FAnalogInputEvent& Event)
+{
+	return Event;
+}
+
+FInputEvent UWidgetBlueprintLibrary::GetInputEventFromCharacterEvent(const FCharacterEvent& Event)
+{
+	return Event;
+}
+
+FInputEvent UWidgetBlueprintLibrary::GetInputEventFromPointerEvent(const FPointerEvent& Event)
 {
-	return KeyEvent;
+	return Event;
 }
+
+FInputEvent UWidgetBlueprintLibrary::GetInputEventFromControllerEvent(const FControllerEvent& Event)
+{
+	return Event;
+}
+
+FInputEvent UWidgetBlueprintLibrary::GetInputEventFromNavigationEvent(const FNavigationEvent& Event)
+{
+	return Event;
+}
+
 #undef LOCTEXT_NAMESPACE
\ No newline at end of file
diff --git a/Engine/Source/Runtime/UMG/Public/Blueprint/WidgetBlueprintLibrary.h b/Engine/Source/Runtime/UMG/Public/Blueprint/WidgetBlueprintLibrary.h
--- a/Engine/Source/Runtime/UMG/Public/Blueprint/WidgetBlueprintLibrary.h
+++ b/Engine/Source/Runtime/UMG/Public/Blueprint/WidgetBlueprintLibrary.h
@@ -208,19 +208,31 @@ public:
 	 * Gets the resource object on a brush.  This could be a UTexture2D or a UMaterialInterface.
 	 */
 	UFUNCTION(BlueprintPure, Category="Widget|Brush")
-	static UObject* GetBrushResource(FSlateBrush& Brush);
+	static UObject* GetBrushResource(UPARAM(ref) FSlateBrush& Brush);
+
+	/**
+	 * Gets the brush resource as a texture 2D.
+	 */
+	UFUNCTION(BlueprintPure, Category="Widget|Brush")
+	static UTexture2D* GetBrushResourceAsTexture2D(UPARAM(ref) FSlateBrush& Brush);
+
+	/**
+	 * Gets the brush resource as a material.
+	 */
+	UFUNCTION(BlueprintPure, Category="Widget|Brush")
+	static UMaterialInterface* GetBrushResourceAsMaterial(UPARAM(ref) FSlateBrush& Brush);
 
 	/**
 	 * Sets the resource on a brush to be a UTexture2D.
 	 */
 	UFUNCTION(BlueprintCallable, Category="Widget|Brush")
-	static void SetBrushResourceToTexture(FSlateBrush& Brush, UTexture2D* Texture);
+	static void SetBrushResourceToTexture(UPARAM(ref) FSlateBrush& Brush, UTexture2D* Texture);
 
 	/**
 	 * Sets the resource on a brush to be a Material.
 	 */
 	UFUNCTION(BlueprintCallable, Category="Widget|Brush")
-	static void SetBrushResourceToMaterial(FSlateBrush& Brush, UMaterialInterface* Material);
+	static void SetBrushResourceToMaterial(UPARAM(ref) FSlateBrush& Brush, UMaterialInterface* Material);
 
 	/**
 	 * Creates a Slate Brush that wont draw anything, the "Null Brush".
@@ -237,7 +249,7 @@ public:
 	 * @return A material that supports dynamic input from the game.
 	 */
 	UFUNCTION(BlueprintPure, Category="Widget|Brush")
-	static UMaterialInstanceDynamic* GetDynamicMaterial(FSlateBrush& Brush);
+	static UMaterialInstanceDynamic* GetDynamicMaterial(UPARAM(ref) FSlateBrush& Brush);
 
 	/** Closes any popup menu */
 	UFUNCTION(BlueprintCallable, BlueprintCosmetic, Category="Widget|Menu")
@@ -259,9 +271,24 @@ public:
 	* @param FoundWidgets Output array of widgets that implement the specified interface.
 	* @param TopLevelOnly Only the widgets that are direct children of the viewport will be returned.
 	*/
-	UFUNCTION(BlueprintCallable, BlueprintCosmetic, Category = "Widget", meta = (WorldContext = "WorldContextObject"))
+	UFUNCTION(BlueprintCallable, BlueprintCosmetic, Category="Widget", meta = (WorldContext = "WorldContextObject"))
 	static void GetAllWidgetsWithInterface(UObject* WorldContextObject, TSubclassOf<UInterface> Interface, TArray<UUserWidget*>& FoundWidgets, bool TopLevelOnly);
 
-	UFUNCTION(BlueprintPure, Category = "Widget", meta = (CompactNodeTitle = "->", BlueprintAutocast))
-	static FInputEvent GetInputEventFromKeyEvent(const FKeyEvent& KeyEvent);
+	UFUNCTION(BlueprintPure, Category="Widget", meta = (CompactNodeTitle = "->", BlueprintAutocast))
+	static FInputEvent GetInputEventFromKeyEvent(const FKeyEvent& Event);
+
+	UFUNCTION(BlueprintPure, Category="Widget", meta = (CompactNodeTitle = "->", BlueprintAutocast))
+	static FInputEvent GetInputEventFromAnalogInputEvent(const FAnalogInputEvent& Event);
+
+	UFUNCTION(BlueprintPure, Category="Widget", meta = ( CompactNodeTitle = "->", BlueprintAutocast ))
+	static FInputEvent GetInputEventFromCharacterEvent(const FCharacterEvent& Event);
+
+	UFUNCTION(BlueprintPure, Category="Widget", meta = ( CompactNodeTitle = "->", BlueprintAutocast ))
+	static FInputEvent GetInputEventFromPointerEvent(const FPointerEvent& Event);
+
+	UFUNCTION(BlueprintPure, Category="Widget", meta = ( CompactNodeTitle = "->", BlueprintAutocast ))
+	static FInputEvent GetInputEventFromControllerEvent(const FControllerEvent& Event);
+
+	UFUNCTION(BlueprintPure, Category="Widget", meta = ( CompactNodeTitle = "->", BlueprintAutocast ))
+	static FInputEvent GetInputEventFromNavigationEvent(const FNavigationEvent& Event);
 };
