UMG - Fixing LocalToViewport and AbsoluteToViewport so that they no longer remove resolution scalability as that is no longer required.

diff --git a/Engine/Source/Runtime/UMG/Private/SlateBlueprintLibrary.cpp b/Engine/Source/Runtime/UMG/Private/SlateBlueprintLibrary.cpp
--- a/Engine/Source/Runtime/UMG/Private/SlateBlueprintLibrary.cpp
+++ b/Engine/Source/Runtime/UMG/Private/SlateBlueprintLibrary.cpp
@@ -36,13 +36,13 @@ FVector2D USlateBlueprintLibrary::GetLocalSize(const FGeometry& Geometry)
 	return Geometry.GetLocalSize();
 }
 
-void USlateBlueprintLibrary::LocalToViewport(UObject* WorldContextObject, const FGeometry& Geometry, FVector2D LocalCoordinate, FVector2D& ScreenPosition, FVector2D& ViewportPosition)
+void USlateBlueprintLibrary::LocalToViewport(UObject* WorldContextObject, const FGeometry& Geometry, FVector2D LocalCoordinate, FVector2D& PixelPosition, FVector2D& ViewportPosition)
 {
 	FVector2D AbsoluteCoordinate = Geometry.LocalToAbsolute(LocalCoordinate);
-	AbsoluteToViewport(WorldContextObject, Geometry, AbsoluteCoordinate, ScreenPosition, ViewportPosition);
+	AbsoluteToViewport(WorldContextObject, AbsoluteCoordinate, PixelPosition, ViewportPosition);
 }
 
-void USlateBlueprintLibrary::AbsoluteToViewport(UObject* WorldContextObject, const FGeometry& Geometry, FVector2D AbsoluteCoordinate, FVector2D& ScreenPosition, FVector2D& ViewportPosition)
+void USlateBlueprintLibrary::AbsoluteToViewport(UObject* WorldContextObject, FVector2D AbsoluteDesktopCoordinate, FVector2D& PixelPosition, FVector2D& ViewportPosition)
 {
 	UWorld* World = GEngine->GetWorldFromContextObject(WorldContextObject);
 	if ( World && World->IsGameWorld() )
@@ -54,29 +54,19 @@ void USlateBlueprintLibrary::AbsoluteToViewport(UObject* WorldContextObject, con
 				FVector2D ViewportSize;
 				ViewportClient->GetViewportSize(ViewportSize);
 
-				FVector2D PixelLocation = Viewport->VirtualDesktopPixelToViewport(FIntPoint((int32)AbsoluteCoordinate.X, (int32)AbsoluteCoordinate.Y)) * ViewportSize;
+				PixelPosition = Viewport->VirtualDesktopPixelToViewport(FIntPoint((int32)AbsoluteDesktopCoordinate.X, (int32)AbsoluteDesktopCoordinate.Y)) * ViewportSize;
 
 				float CurrentViewportScale = GetDefault<UUserInterfaceSettings>()->GetDPIScaleBasedOnSize(FIntPoint(ViewportSize.X, ViewportSize.Y));
 
-				// If the user has configured a resolution quality we need to multiply
-				// the pixels by the resolution quality to arrive at the true position in
-				// the viewport, as the rendered image will be stretched to fill whatever
-				// size the viewport is at.
-				Scalability::FQualityLevels ScalabilityQuality = Scalability::GetQualityLevels();
-				const float QualityScale = ( ScalabilityQuality.ResolutionQuality / 100.0f );
-
-				// Remove the resolution quality scale.
-				ScreenPosition = PixelLocation * QualityScale;
-
 				// Remove DPI Scaling.
-				ViewportPosition = PixelLocation / CurrentViewportScale;
+				ViewportPosition = PixelPosition / CurrentViewportScale;
 
 				return;
 			}
 		}
 	}
 
-	ScreenPosition = FVector2D(0, 0);
+	PixelPosition = FVector2D(0, 0);
 	ViewportPosition = FVector2D(0, 0);
 }
 
diff --git a/Engine/Source/Runtime/UMG/Public/Blueprint/SlateBlueprintLibrary.h b/Engine/Source/Runtime/UMG/Public/Blueprint/SlateBlueprintLibrary.h
--- a/Engine/Source/Runtime/UMG/Public/Blueprint/SlateBlueprintLibrary.h
+++ b/Engine/Source/Runtime/UMG/Public/Blueprint/SlateBlueprintLibrary.h
@@ -44,24 +44,24 @@ public:
 	/**
 	 * Translates local coordinate of the geometry provided into local viewport coordinates.
 	 *
-	 * @param ScreenPosition The position in the game's viewport, usable for line traces and 
+	 * @param PixelPosition The position in the game's viewport, usable for line traces and 
 	 * other uses where you need a coordinate in the space of viewport resolution units.
 	 * @param ViewportPosition The position in the space of other widgets in the viewport.  Like if you wanted
 	 * to add another widget to the viewport at the same position in viewport space as this location, this is
 	 * what you would use.
 	 */
 	UFUNCTION(BlueprintPure, Category="User Interface|Geometry", meta=( HidePin="WorldContextObject", DefaultToSelf="WorldContextObject" ))
-	static void LocalToViewport(UObject* WorldContextObject, const FGeometry& Geometry, FVector2D LocalCoordinate, FVector2D& ScreenPosition, FVector2D& ViewportPosition);
+	static void LocalToViewport(UObject* WorldContextObject, const FGeometry& Geometry, FVector2D LocalCoordinate, FVector2D& PixelPosition, FVector2D& ViewportPosition);
 
 	/**
-	 * Translates absolute coordinate of the geometry provided into local viewport coordinates.
+	 * Translates absolute coordinate in desktop space of the geometry provided into local viewport coordinates.
 	 *
-	 * @param ScreenPosition The position in the game's viewport, usable for line traces and
+	 * @param PixelPosition The position in the game's viewport, usable for line traces and
 	 * other uses where you need a coordinate in the space of viewport resolution units.
 	 * @param ViewportPosition The position in the space of other widgets in the viewport.  Like if you wanted
 	 * to add another widget to the viewport at the same position in viewport space as this location, this is
 	 * what you would use.
 	 */
 	UFUNCTION(BlueprintPure, Category="User Interface|Geometry", meta=( HidePin="WorldContextObject", DefaultToSelf="WorldContextObject" ))
-	static void AbsoluteToViewport(UObject* WorldContextObject, const FGeometry& Geometry, FVector2D AbsoluteCoordinate, FVector2D& ScreenPosition, FVector2D& ViewportPosition);
+	static void AbsoluteToViewport(UObject* WorldContextObject, FVector2D AbsoluteDesktopCoordinate, FVector2D& PixelPosition, FVector2D& ViewportPosition);
 };
