Improved the warning message for missing UMG animations

diff --git a/Engine/Source/Runtime/UMG/Private/Animation/UMGSequencePlayer.cpp b/Engine/Source/Runtime/UMG/Private/Animation/UMGSequencePlayer.cpp
--- a/Engine/Source/Runtime/UMG/Private/Animation/UMGSequencePlayer.cpp
+++ b/Engine/Source/Runtime/UMG/Private/Animation/UMGSequencePlayer.cpp
@@ -159,7 +159,7 @@ void UUMGSequencePlayer::GetRuntimeObjects( TSharedRef<FMovieSceneInstance> Movi
 	}
 	else
 	{
-		UE_LOG( LogUMG, Warning, TEXT("Failed to find runtime objects for %s animation"), Animation ? *Animation->GetName() : TEXT("(none)") );
+		UE_LOG( LogUMG, Warning, TEXT("Failed to find runtime objects for %s animation"), Animation ? *Animation->GetPathName() : TEXT("(none)") );
 	}
 }
 
