Fixing the crash that occurs when you set draw size to 0, the component no longer attempts to draw the render target if it's size is 0.

diff --git a/Engine/Source/Runtime/UMG/Private/Components/WidgetComponent.cpp b/Engine/Source/Runtime/UMG/Private/Components/WidgetComponent.cpp
--- a/Engine/Source/Runtime/UMG/Private/Components/WidgetComponent.cpp
+++ b/Engine/Source/Runtime/UMG/Private/Components/WidgetComponent.cpp
@@ -721,6 +721,11 @@ void UWidgetComponent::TickComponent(float DeltaTime, enum ELevelTick TickType,
 			return;
 		}
 
+		if ( DrawSize.X == 0 || DrawSize.Y == 0 )
+		{
+			return;
+		}
+
 		const float RenderTimeThreshold = .5f;
 		if ( IsVisible() )
 		{
@@ -1026,6 +1031,10 @@ void UWidgetComponent::UpdateWidget()
 				SlateWidget->SetContent(SNullWidget::NullWidget);
 			}
 		}
+		else
+		{
+			SlateWidget.Reset();
+		}
 	}
 }
 
@@ -1043,52 +1052,58 @@ void UWidgetComponent::UpdateRenderTarget()
 		ActualBackgroundColor.A = 0.0f;
 	}
 
-	if(!RenderTarget && DrawSize != FIntPoint::ZeroValue)
+	if ( DrawSize.X != 0 && DrawSize.Y != 0 )
 	{
-		RenderTarget = NewObject<UTextureRenderTarget2D>(this);
-		RenderTarget->ClearColor = ActualBackgroundColor;
+		if ( RenderTarget == nullptr )
+		{
+			RenderTarget = NewObject<UTextureRenderTarget2D>(this);
+			RenderTarget->ClearColor = ActualBackgroundColor;
 
-		bClearColorChanged = bRenderStateDirty = true;
+			bClearColorChanged = bRenderStateDirty = true;
 
-		RenderTarget->InitCustomFormat(DrawSize.X, DrawSize.Y, PF_B8G8R8A8, false);
+			RenderTarget->InitCustomFormat(DrawSize.X, DrawSize.Y, PF_B8G8R8A8, false);
 
-		MaterialInstance->SetTextureParameterValue( "SlateUI", RenderTarget );
-	}
-	else if ( DrawSize != FIntPoint::ZeroValue )
-	{
-		// Update the format
-		if ( RenderTarget->SizeX != DrawSize.X || RenderTarget->SizeY != DrawSize.Y )
-		{
-			RenderTarget->InitCustomFormat( DrawSize.X, DrawSize.Y, PF_B8G8R8A8, false );
-			bRenderStateDirty = true;
+			MaterialInstance->SetTextureParameterValue("SlateUI", RenderTarget);
 		}
-
-		// Update the clear color
-		if ( RenderTarget->ClearColor != ActualBackgroundColor )
+		else
 		{
-			RenderTarget->ClearColor = ActualBackgroundColor;
-			bClearColorChanged = bRenderStateDirty = true;
+			// Update the format
+			if ( RenderTarget->SizeX != DrawSize.X || RenderTarget->SizeY != DrawSize.Y )
+			{
+				RenderTarget->InitCustomFormat(DrawSize.X, DrawSize.Y, PF_B8G8R8A8, false);
+				bRenderStateDirty = true;
+			}
+
+			// Update the clear color
+			if ( RenderTarget->ClearColor != ActualBackgroundColor )
+			{
+				RenderTarget->ClearColor = ActualBackgroundColor;
+				bClearColorChanged = bRenderStateDirty = true;
+			}
 		}
 	}
 
-	// If the clear color of the render target changed, update the BackColor of the material to match
-	if ( bClearColorChanged )
+	if ( RenderTarget )
 	{
-		MaterialInstance->SetVectorParameterValue( "BackColor", RenderTarget->ClearColor );
-	}
+		// If the clear color of the render target changed, update the BackColor of the material to match
+		if ( bClearColorChanged )
+		{
+			MaterialInstance->SetVectorParameterValue("BackColor", RenderTarget->ClearColor);
+		}
 
-	static FName ParabolaDistortionName(TEXT("ParabolaDistortion"));
+		static FName ParabolaDistortionName(TEXT("ParabolaDistortion"));
 
-	float CurrentParabolaValue;
-	if ( MaterialInstance->GetScalarParameterValue(ParabolaDistortionName, CurrentParabolaValue) && CurrentParabolaValue != ParabolaDistortion )
-	{
-		MaterialInstance->SetScalarParameterValue(ParabolaDistortionName, ParabolaDistortion);
-		bRenderStateDirty = true;
-	}
+		float CurrentParabolaValue;
+		if ( MaterialInstance->GetScalarParameterValue(ParabolaDistortionName, CurrentParabolaValue) && CurrentParabolaValue != ParabolaDistortion )
+		{
+			MaterialInstance->SetScalarParameterValue(ParabolaDistortionName, ParabolaDistortion);
+			bRenderStateDirty = true;
+		}
 
-	if ( bRenderStateDirty )
-	{
-		MarkRenderStateDirty();
+		if ( bRenderStateDirty )
+		{
+			MarkRenderStateDirty();
+		}
 	}
 }
 
@@ -1195,7 +1210,14 @@ FVector2D UWidgetComponent::GetDrawSize() const
 
 void UWidgetComponent::SetDrawSize(FVector2D Size)
 {
-	DrawSize = FIntPoint((int32)Size.X, (int32)Size.Y);
+	FIntPoint NewDrawSize((int32)Size.X, (int32)Size.Y);
+
+	if ( NewDrawSize != DrawSize )
+	{
+		DrawSize = NewDrawSize;
+		MarkRenderStateDirty();
+		RecreatePhysicsState();
+	}
 }
 
 float UWidgetComponent::GetMaxInteractionDistance() const
