UMG - Now defaulting to sRGB space in the color picker for most colors on widgets and slate brushes.

diff --git a/Engine/Source/Runtime/SlateCore/Public/Styling/SlateBrush.h b/Engine/Source/Runtime/SlateCore/Public/Styling/SlateBrush.h
--- a/Engine/Source/Runtime/SlateCore/Public/Styling/SlateBrush.h
+++ b/Engine/Source/Runtime/SlateCore/Public/Styling/SlateBrush.h
@@ -122,7 +122,7 @@ struct SLATECORE_API FSlateBrush
 	FLinearColor Tint_DEPRECATED;
 
 	/** Tinting applied to the image. */
-	UPROPERTY(EditAnywhere, BlueprintReadWrite, Category=Brush, meta=( DisplayName="Tint" ))
+	UPROPERTY(EditAnywhere, BlueprintReadWrite, Category=Brush, meta=( DisplayName="Tint", sRGB="true" ))
 	FSlateColor TintColor;
 
 	/** How to tile the image in Image mode */
diff --git a/Engine/Source/Runtime/UMG/Public/Components/Border.h b/Engine/Source/Runtime/UMG/Public/Components/Border.h
--- a/Engine/Source/Runtime/UMG/Public/Components/Border.h
+++ b/Engine/Source/Runtime/UMG/Public/Components/Border.h
@@ -21,7 +21,7 @@ class UMG_API UBorder : public UContentWidget
 public:
 
 	/** Color and opacity multiplier of content in the border */
-	UPROPERTY(EditAnywhere, BlueprintReadOnly, Category="Content")
+	UPROPERTY(EditAnywhere, BlueprintReadOnly, Category="Content", meta=( sRGB="true" ))
 	FLinearColor ContentColorAndOpacity;
 
 	/** A bindable delegate for the ContentColorAndOpacity. */
@@ -49,7 +49,7 @@ public:
 	FGetSlateBrush BackgroundDelegate;
 
 	/** Color and opacity of the actual border image */
-	UPROPERTY(EditAnywhere, BlueprintReadOnly, Category=Appearance)
+	UPROPERTY(EditAnywhere, BlueprintReadOnly, Category=Appearance, meta=( sRGB="true" ))
 	FLinearColor BrushColor;
 
 	/** A bindable delegate for the BrushColor. */
diff --git a/Engine/Source/Runtime/UMG/Public/Components/Button.h b/Engine/Source/Runtime/UMG/Public/Components/Button.h
--- a/Engine/Source/Runtime/UMG/Public/Components/Button.h
+++ b/Engine/Source/Runtime/UMG/Public/Components/Button.h
@@ -33,11 +33,11 @@ public:
 	FButtonStyle WidgetStyle;
 	
 	/** The color multiplier for the button content */
-	UPROPERTY(EditAnywhere, BlueprintReadOnly, Category="Appearance")
+	UPROPERTY(EditAnywhere, BlueprintReadOnly, Category="Appearance", meta=( sRGB="true" ))
 	FLinearColor ColorAndOpacity;
 	
 	/** The color multiplier for the button background */
-	UPROPERTY(EditAnywhere, BlueprintReadOnly, Category="Appearance")
+	UPROPERTY(EditAnywhere, BlueprintReadOnly, Category="Appearance", meta=( sRGB="true" ))
 	FLinearColor BackgroundColor;
 
 	/** The type of mouse action required by the user to trigger the buttons 'Click' */
diff --git a/Engine/Source/Runtime/UMG/Public/Components/Image.h b/Engine/Source/Runtime/UMG/Public/Components/Image.h
--- a/Engine/Source/Runtime/UMG/Public/Components/Image.h
+++ b/Engine/Source/Runtime/UMG/Public/Components/Image.h
@@ -31,7 +31,7 @@ public:
 	FGetSlateBrush BrushDelegate;
 
 	/** Color and opacity */
-	UPROPERTY(EditAnywhere, BlueprintReadOnly, Category=Appearance)
+	UPROPERTY(EditAnywhere, BlueprintReadOnly, Category=Appearance, meta=( sRGB="true") )
 	FLinearColor ColorAndOpacity;
 
 	/** A bindable delegate for the ColorAndOpacity. */
