UMG - The cursor property on widgets now works as it should.  There are also some functions to enable changing it at runtime.

diff --git a/Engine/Source/Runtime/UMG/Private/Components/Widget.cpp b/Engine/Source/Runtime/UMG/Private/Components/Widget.cpp
--- a/Engine/Source/Runtime/UMG/Private/Components/Widget.cpp
+++ b/Engine/Source/Runtime/UMG/Private/Components/Widget.cpp
@@ -105,9 +105,7 @@ UWidget::UWidget(const FObjectInitializer& ObjectInitializer)
 #endif
 	Visiblity_DEPRECATED = Visibility = ESlateVisibility::Visible;	
 	RenderTransformPivot = FVector2D(0.5f, 0.5f);
-
-	//TODO UMG ToolTipWidget
-	//TODO UMG Cursor doesn't work yet, the underlying slate version needs it to be TOptional.
+	Cursor = EMouseCursor::Default;
 }
 
 void UWidget::SetRenderTransform(FWidgetTransform Transform)
@@ -184,6 +182,29 @@ void UWidget::SetIsEnabled(bool bInIsEnabled)
 	}
 }
 
+void UWidget::SetCursor(EMouseCursor::Type InCursor)
+{
+	bOverride_Cursor = true;
+	Cursor = InCursor;
+
+	TSharedPtr<SWidget> SafeWidget = GetCachedWidget();
+	if ( SafeWidget.IsValid() )
+	{
+		SafeWidget->SetCursor(Cursor);
+	}
+}
+
+void UWidget::ResetCursor()
+{
+	bOverride_Cursor = false;
+
+	TSharedPtr<SWidget> SafeWidget = GetCachedWidget();
+	if ( SafeWidget.IsValid() )
+	{
+		SafeWidget->SetCursor(TOptional<EMouseCursor::Type>());
+	}
+}
+
 bool UWidget::IsVisible() const
 {
 	TSharedPtr<SWidget> SafeWidget = GetCachedWidget();
@@ -735,6 +756,11 @@ void UWidget::SynchronizeProperties()
 	else
 #endif
 	{
+		if ( bOverride_Cursor /*|| CursorDelegate.IsBound()*/ )
+		{
+			SafeWidget->SetCursor(Cursor);// GAME_SAFE_OPTIONAL_BINDING(EMouseCursor::Type, Cursor));
+		}
+
 		SafeWidget->SetEnabled(GAME_SAFE_OPTIONAL_BINDING( bool, bIsEnabled ));
 		SafeWidget->SetVisibility(OPTIONAL_BINDING_CONVERT(ESlateVisibility, Visibility, EVisibility, ConvertVisibility));
 	}
diff --git a/Engine/Source/Runtime/UMG/Public/Components/Widget.h b/Engine/Source/Runtime/UMG/Public/Components/Widget.h
--- a/Engine/Source/Runtime/UMG/Public/Components/Widget.h
+++ b/Engine/Source/Runtime/UMG/Public/Components/Widget.h
@@ -169,13 +169,17 @@ public:
 	UPROPERTY()
 	FGetSlateVisibility VisibilityDelegate;
 
+	/**  */
+	UPROPERTY()
+	uint32 bOverride_Cursor : 1;
+
 	/** The cursor to show when the mouse is over the widget */
-	UPROPERTY(EditAnywhere, BlueprintReadOnly, Category="Behavior", AdvancedDisplay)
+	UPROPERTY(EditAnywhere, BlueprintReadOnly, Category="Behavior", AdvancedDisplay, meta=( editcondition="bOverride_Cursor" ))
 	TEnumAsByte<EMouseCursor::Type> Cursor;
 
 	/** A bindable delegate for Cursor */
-	UPROPERTY()
-	FGetMouseCursor CursorDelegate;
+	//UPROPERTY()
+	//FGetMouseCursor CursorDelegate;
 
 protected:
 
@@ -267,6 +271,14 @@ public:
 	UFUNCTION(BlueprintCallable, Category="Widget")
 	void SetToolTip(UWidget* Widget);
 
+	/** Sets the cursor to show over the widget. */
+	UFUNCTION(BlueprintCallable, Category="Widget")
+	void SetCursor(EMouseCursor::Type InCursor);
+
+	/** Resets the cursor to use on the widget, removing any customization for it. */
+	UFUNCTION(BlueprintCallable, Category="Widget")
+	void ResetCursor();
+
 	/** @return true if the widget is Visible, HitTestInvisible or SelfHitTestInvisible. */
 	UFUNCTION(BlueprintCallable, Category="Widget")
 	bool IsVisible() const;
@@ -574,4 +586,5 @@ private:
 private:
 	GAME_SAFE_BINDING_IMPLEMENTATION(FText, ToolTipText)
 	GAME_SAFE_BINDING_IMPLEMENTATION(bool, bIsEnabled)
+	//GAME_SAFE_BINDING_IMPLEMENTATION(EMouseCursor::Type, Cursor)
 };
