UMG/Sequencer - Fix slot animation handling.

diff --git a/Engine/Source/Runtime/UMG/Private/Animation/WidgetAnimation.cpp b/Engine/Source/Runtime/UMG/Private/Animation/WidgetAnimation.cpp
--- a/Engine/Source/Runtime/UMG/Private/Animation/WidgetAnimation.cpp
+++ b/Engine/Source/Runtime/UMG/Private/Animation/WidgetAnimation.cpp
@@ -160,7 +160,7 @@ FGuid UWidgetAnimation::FindObjectId(UObject& Object) const
 	if (Slot != nullptr)
 	{
 		// slot guids are tracked by their content.
-		return PreviewObjectToIds.FindRef(Slot->Content);
+		return SlotContentPreviewObjectToIds.FindRef(Slot->Content);
 	}
 
 	return PreviewObjectToIds.FindRef(&Object);
@@ -180,7 +180,11 @@ UObject* UWidgetAnimation::FindObject(const FGuid& ObjectId) const
 
 	if (SlotContentPreviewObject.IsValid())
 	{
-		return SlotContentPreviewObject.Get();
+		UWidget* Widget = Cast<UWidget>(SlotContentPreviewObject.Get());
+		if ( Widget != nullptr )
+		{
+			return Widget->Slot;
+		}
 	}
 	
 	return nullptr;
@@ -232,7 +236,14 @@ bool UWidgetAnimation::TryGetObjectDisplayName(const FGuid& ObjectId, FText& Out
 			? FText::FromString(SlotContent->Slot->Parent->GetName())
 			: LOCTEXT("InvalidPanel", "Invalid Panel");
 		FText ContentName = FText::FromString(SlotContent->GetName());
-		OutDisplayName = FText::Format(LOCTEXT("SlotObject", "{0} ({1} Slot)"), ContentName, PanelName);
+		if ( PreviewObjectToIds.Contains( SlotContent ) )
+		{
+			OutDisplayName = FText::Format( LOCTEXT( "SlotObjectWithParent", "{0} Slot" ), PanelName );
+		}
+		else
+		{
+			OutDisplayName = FText::Format(LOCTEXT("SlotObject", "{0} ({1} Slot)"), ContentName, PanelName);
+		}
 
 		return true;
 	}
