Convert handful of functions that still used HidePin/DefaultToSelf instead of WorldContext for identifying the WorldContextObject pin on blueprint callable functions

diff --git a/Engine/Source/Runtime/Engine/Classes/Engine/CanvasRenderTarget2D.h b/Engine/Source/Runtime/Engine/Classes/Engine/CanvasRenderTarget2D.h
--- a/Engine/Source/Runtime/Engine/Classes/Engine/CanvasRenderTarget2D.h
+++ b/Engine/Source/Runtime/Engine/Classes/Engine/CanvasRenderTarget2D.h
@@ -49,7 +49,7 @@ public:
 	 *
 	 * @return						Returns the instanced render target.
 	 */
-	UFUNCTION(BlueprintCallable, Category="Canvas Render Target 2D",meta=(HidePin="WorldContextObject", DefaultToSelf="WorldContextObject" ))
+	UFUNCTION(BlueprintCallable, Category="Canvas Render Target 2D",meta=(WorldContext="WorldContextObject" ))
 	static UCanvasRenderTarget2D* CreateCanvasRenderTarget2D(UObject* WorldContextObject, TSubclassOf<UCanvasRenderTarget2D> CanvasRenderTarget2DClass, int32 Width = 1024, int32 Height = 1024);
  
 	/**
diff --git a/Engine/Source/Runtime/UMG/Public/Blueprint/SlateBlueprintLibrary.h b/Engine/Source/Runtime/UMG/Public/Blueprint/SlateBlueprintLibrary.h
--- a/Engine/Source/Runtime/UMG/Public/Blueprint/SlateBlueprintLibrary.h
+++ b/Engine/Source/Runtime/UMG/Public/Blueprint/SlateBlueprintLibrary.h
@@ -50,7 +50,7 @@ public:
 	 * to add another widget to the viewport at the same position in viewport space as this location, this is
 	 * what you would use.
 	 */
-	UFUNCTION(BlueprintPure, Category="User Interface|Geometry", meta=( HidePin="WorldContextObject", DefaultToSelf="WorldContextObject" ))
+	UFUNCTION(BlueprintPure, Category="User Interface|Geometry", meta=( WorldContext="WorldContextObject" ))
 	static void LocalToViewport(UObject* WorldContextObject, const FGeometry& Geometry, FVector2D LocalCoordinate, FVector2D& PixelPosition, FVector2D& ViewportPosition);
 
 	/**
@@ -62,6 +62,6 @@ public:
 	 * to add another widget to the viewport at the same position in viewport space as this location, this is
 	 * what you would use.
 	 */
-	UFUNCTION(BlueprintPure, Category="User Interface|Geometry", meta=( HidePin="WorldContextObject", DefaultToSelf="WorldContextObject" ))
+	UFUNCTION(BlueprintPure, Category="User Interface|Geometry", meta=( WorldContext="WorldContextObject" ))
 	static void AbsoluteToViewport(UObject* WorldContextObject, FVector2D AbsoluteDesktopCoordinate, FVector2D& PixelPosition, FVector2D& ViewportPosition);
 };
diff --git a/Engine/Source/Runtime/UMG/Public/Blueprint/WidgetLayoutLibrary.h b/Engine/Source/Runtime/UMG/Public/Blueprint/WidgetLayoutLibrary.h
--- a/Engine/Source/Runtime/UMG/Public/Blueprint/WidgetLayoutLibrary.h
+++ b/Engine/Source/Runtime/UMG/Public/Blueprint/WidgetLayoutLibrary.h
@@ -33,13 +33,13 @@ public:
 	/**
 	 * Gets the current DPI Scale being applied to the viewport and all the Widgets.
 	 */
-	UFUNCTION(BlueprintPure, BlueprintCosmetic, Category="Viewport", meta=( HidePin="WorldContextObject", DefaultToSelf="WorldContextObject" ))
+	UFUNCTION(BlueprintPure, BlueprintCosmetic, Category="Viewport", meta=( WorldContext="WorldContextObject" ))
 	static float GetViewportScale(UObject* WorldContextObject);
 
 	/**
 	 * Gets the size of the game viewport.
 	 */
-	UFUNCTION(BlueprintPure, BlueprintCosmetic, Category="Viewport", meta=( HidePin="WorldContextObject", DefaultToSelf="WorldContextObject" ))
+	UFUNCTION(BlueprintPure, BlueprintCosmetic, Category="Viewport", meta=( WorldContext="WorldContextObject" ))
 	static FVector2D GetViewportSize(UObject* WorldContextObject);
 
 	/**
@@ -96,6 +96,6 @@ public:
 	/**
 	 * Removes all widgets from the viewport.
 	 */
-	UFUNCTION(BlueprintCallable, BlueprintCosmetic, Category="Viewport", meta=( HidePin="WorldContextObject", DefaultToSelf="WorldContextObject" ))
+	UFUNCTION(BlueprintCallable, BlueprintCosmetic, Category="Viewport", meta=( WorldContext="WorldContextObject" ))
 	static void RemoveAllWidgets(UObject* WorldContextObject);
 };
