DynamicPropertyPath now checks to make sure the UObject isn't unreachable before trying to process an event with it.

diff --git a/Engine/Source/Runtime/UMG/Public/Binding/DynamicPropertyPath.h b/Engine/Source/Runtime/UMG/Public/Binding/DynamicPropertyPath.h
--- a/Engine/Source/Runtime/UMG/Public/Binding/DynamicPropertyPath.h
+++ b/Engine/Source/Runtime/UMG/Public/Binding/DynamicPropertyPath.h
@@ -331,7 +331,7 @@ private:
 							if ( IsConcreteTypeCompatibleWithReflectedType<T>(ReturnProperty) )
 							{
 								// Ensure that the element sizes are the same, prevents the user from doing something terribly wrong.
-								if ( ReturnProperty->ElementSize == sizeof(T) )
+								if ( ReturnProperty->ElementSize == sizeof(T) && !ContainerObject->HasAnyFlags( RF_Unreachable ) )
 								{
 									ContainerObject->ProcessEvent(Function, &OutValue);
 									return true;
