UMG - Exposing a function to get if any animation is currently playing.  Also includes the canTick/Paint members.

diff --git a/Engine/Source/Runtime/UMG/Public/Blueprint/UserWidget.h b/Engine/Source/Runtime/UMG/Public/Blueprint/UserWidget.h
--- a/Engine/Source/Runtime/UMG/Public/Blueprint/UserWidget.h
+++ b/Engine/Source/Runtime/UMG/Public/Blueprint/UserWidget.h
@@ -761,6 +761,7 @@ public:
 	// End UObject interface
 
 	/** Are we currently playing any animations? */
+	UFUNCTION(BlueprintCallable, Category="User Interface|Animation")
 	bool IsPlayingAnimation() const { return ActiveSequencePlayers.Num() > 0; }
 
 #if WITH_EDITOR
@@ -840,6 +841,12 @@ public:
 
 #endif
 
+	/** If a widget doesn't ever need to tick the blueprint, setting this to false is an optimization. */
+	bool bCanEverTick : 1;
+
+	/** If a widget doesn't ever need to do custom painting in the blueprint, setting this to false is an optimization. */
+	bool bCanEverPaint : 1;
+
 protected:
 
 	/** Adds the widget to the screen, either to the viewport or to the player's screen depending on if the LocalPlayer is null. */
