#UE4

- Null guard in SetBrushFromTexture

--------
Integrated using branch UE4-To-UE4-Orion (reversed) of change#2651458 by [REDACTED_PERSON] on 2015/08/11 14:10:25.

diff --git a/Engine/Source/Runtime/UMG/Private/Components/Image.cpp b/Engine/Source/Runtime/UMG/Private/Components/Image.cpp
--- a/Engine/Source/Runtime/UMG/Private/Components/Image.cpp
+++ b/Engine/Source/Runtime/UMG/Private/Components/Image.cpp
@@ -100,7 +100,7 @@ void UImage::SetBrushFromTexture(UTexture2D* Texture, bool bMatchSize)
 {
 	Brush.SetResourceObject(Texture);
 
-	if (bMatchSize)
+	if (bMatchSize && Texture)
 	{
 		Brush.ImageSize.X = Texture->GetSizeX();
 		Brush.ImageSize.Y = Texture->GetSizeY();
