- Automatic wrapping of textblocks with invalidation panels

--------
Integrated using branch UE4-To-UE4-Orion (reversed) of change#2661622 by [REDACTED_PERSON] on 2015/08/19 16:37:06.

diff --git a/Engine/Source/Runtime/UMG/Private/Components/TextBlock.cpp b/Engine/Source/Runtime/UMG/Private/Components/TextBlock.cpp
--- a/Engine/Source/Runtime/UMG/Private/Components/TextBlock.cpp
+++ b/Engine/Source/Runtime/UMG/Private/Components/TextBlock.cpp
@@ -1,6 +1,7 @@
 // Copyright 1998-2015 Epic Games, Inc. All Rights Reserved.
 
 #include "UMGPrivatePCH.h"
+#include "SInvalidationPanel.h"
 
 #define LOCTEXT_NAMESPACE "UMG"
 
@@ -11,7 +12,7 @@ UTextBlock::UTextBlock(const FObjectInitializer& ObjectInitializer)
 	: Super(ObjectInitializer)
 {
 	bIsVariable = false;
-
+	bWrapWithInvalidationPanel = false;
 	ShadowOffset = FVector2D(1.0f, 1.0f);
 	ColorAndOpacity = FLinearColor::White;
 	ShadowColorAndOpacity = FLinearColor::Transparent;
@@ -83,10 +84,19 @@ void UTextBlock::SetJustification( ETextJustify::Type InJustification )
 
 TSharedRef<SWidget> UTextBlock::RebuildWidget()
 {
-	MyTextBlock = SNew(STextBlock);
-//		.TextStyle(&WidgetStyle);
-
-	return MyTextBlock.ToSharedRef();
+ 	if (bWrapWithInvalidationPanel && !IsDesignTime())
+ 	{
+ 		TSharedPtr<SWidget> RetWidget = SNew(SInvalidationPanel)
+ 		[
+ 			SAssignNew(MyTextBlock, STextBlock)
+ 		];
+ 		return RetWidget.ToSharedRef();
+ 	}
+ 	else
+	{
+		MyTextBlock = SNew(STextBlock);
+		return MyTextBlock.ToSharedRef();
+	}
 }
 
 void UTextBlock::OnBindingChanged(const FName& Property)
diff --git a/Engine/Source/Runtime/UMG/Public/Components/TextBlock.h b/Engine/Source/Runtime/UMG/Public/Components/TextBlock.h
--- a/Engine/Source/Runtime/UMG/Public/Components/TextBlock.h
+++ b/Engine/Source/Runtime/UMG/Public/Components/TextBlock.h
@@ -124,6 +124,10 @@ public:
 	UPROPERTY(EditAnywhere, BlueprintReadOnly, Category=Appearance, AdvancedDisplay)
 	float LineHeightPercentage;
 
+	/** If true, it will automatically wrap this text widget with an invalidation panel */
+	UPROPERTY(EditAnywhere, BlueprintReadOnly, Category = Performance, AdvancedDisplay)
+	bool bWrapWithInvalidationPanel;
+
 	///** Called when this text is double clicked */
 	//SLATE_EVENT(FOnClicked, OnDoubleClicked)
 
