Sequencer: Compile fixes.

diff --git a/Engine/Source/Runtime/MovieSceneTracks/Private/Sections/MovieSceneColorSection.cpp b/Engine/Source/Runtime/MovieSceneTracks/Private/Sections/MovieSceneColorSection.cpp
--- a/Engine/Source/Runtime/MovieSceneTracks/Private/Sections/MovieSceneColorSection.cpp
+++ b/Engine/Source/Runtime/MovieSceneTracks/Private/Sections/MovieSceneColorSection.cpp
@@ -90,25 +90,25 @@ void UMovieSceneColorSection::AddKey( float Time, const FColorKey& Key, FKeyPara
 	bool bAlphaKeyExists = AlphaCurve.IsKeyHandleValid(AlphaCurve.FindKey(Time));
 
 	if ( (CurveName == NAME_None || CurveName == RedName) &&
-			KeyParams.bAddKeyEvenIfUnchanged || !(!bRedKeyExists && !KeyParams.bAutoKeying && RedCurve.GetNumKeys() > 0) )
+			(KeyParams.bAddKeyEvenIfUnchanged || !(!bRedKeyExists && !KeyParams.bAutoKeying && RedCurve.GetNumKeys() > 0) ) )
 	{
 		AddKeyToCurve(RedCurve, Time, Key.Value.R, KeyParams);
 	}
 
 	if ( (CurveName == NAME_None || CurveName == GreenName) &&
-			KeyParams.bAddKeyEvenIfUnchanged || !(!bGreenKeyExists && !KeyParams.bAutoKeying && GreenCurve.GetNumKeys() > 0) )
+			(KeyParams.bAddKeyEvenIfUnchanged || !(!bGreenKeyExists && !KeyParams.bAutoKeying && GreenCurve.GetNumKeys() > 0) ) )
 	{
 		AddKeyToCurve(GreenCurve, Time, Key.Value.G, KeyParams);
 	}
 
 	if ( (CurveName == NAME_None || CurveName == BlueName) &&
-			KeyParams.bAddKeyEvenIfUnchanged || !(!bBlueKeyExists && !KeyParams.bAutoKeying && BlueCurve.GetNumKeys() > 0) )
+			(KeyParams.bAddKeyEvenIfUnchanged || !(!bBlueKeyExists && !KeyParams.bAutoKeying && BlueCurve.GetNumKeys() > 0) ) )
 	{
 		AddKeyToCurve(BlueCurve, Time, Key.Value.B, KeyParams);
 	}
 
 	if ( (CurveName == NAME_None || CurveName == AlphaName) &&
-			KeyParams.bAddKeyEvenIfUnchanged || !(!bAlphaKeyExists && !KeyParams.bAutoKeying && AlphaCurve.GetNumKeys() > 0) )
+			(KeyParams.bAddKeyEvenIfUnchanged || !(!bAlphaKeyExists && !KeyParams.bAutoKeying && AlphaCurve.GetNumKeys() > 0) ) )
 	{
 		AddKeyToCurve(AlphaCurve, Time, Key.Value.A, KeyParams);
 	}
diff --git a/Engine/Source/Runtime/UMG/Private/Animation/MovieScene2DTransformSection.cpp b/Engine/Source/Runtime/UMG/Private/Animation/MovieScene2DTransformSection.cpp
--- a/Engine/Source/Runtime/UMG/Private/Animation/MovieScene2DTransformSection.cpp
+++ b/Engine/Source/Runtime/UMG/Private/Animation/MovieScene2DTransformSection.cpp
@@ -193,25 +193,25 @@ void UMovieScene2DTransformSection::AddKey( float Time, const struct F2DTransfor
 	bool bSyKeyExists = Scale[1].IsKeyHandleValid(Scale[1].FindKey(Time));
 
 	if ( (CurveName == NAME_None || CurveName == TranslationName) &&
-			KeyParams.bAddKeyEvenIfUnchanged || !(!bTxKeyExists && !KeyParams.bAutoKeying && Translation[0].GetNumKeys() > 0) )
+			(KeyParams.bAddKeyEvenIfUnchanged || !(!bTxKeyExists && !KeyParams.bAutoKeying && Translation[0].GetNumKeys() > 0) ) )
 	{
 		AddKeyToCurve(Translation[0], Time, TransformKey.Value.Translation.X, KeyParams);
 	}
 
 	if ( (CurveName == NAME_None || CurveName == TranslationName) &&
-			KeyParams.bAddKeyEvenIfUnchanged || !(!bTyKeyExists && !KeyParams.bAutoKeying && Translation[1].GetNumKeys() > 0) )
+			(KeyParams.bAddKeyEvenIfUnchanged || !(!bTyKeyExists && !KeyParams.bAutoKeying && Translation[1].GetNumKeys() > 0) ) )
 	{
 		AddKeyToCurve(Translation[1], Time, TransformKey.Value.Translation.Y, KeyParams);
 	}
 
 	if ( (CurveName == NAME_None || CurveName == ScaleName) &&
-			KeyParams.bAddKeyEvenIfUnchanged || !(!bSxKeyExists && !KeyParams.bAutoKeying && Scale[0].GetNumKeys() > 0) )
+			(KeyParams.bAddKeyEvenIfUnchanged || !(!bSxKeyExists && !KeyParams.bAutoKeying && Scale[0].GetNumKeys() > 0) ) )
 	{
 		AddKeyToCurve(Scale[0], Time, TransformKey.Value.Scale.X, KeyParams);
 	}
 
 	if ( (CurveName == NAME_None || CurveName == ScaleName) &&
-			KeyParams.bAddKeyEvenIfUnchanged || !(!bSyKeyExists && !KeyParams.bAutoKeying && Scale[1].GetNumKeys() > 0) )
+			(KeyParams.bAddKeyEvenIfUnchanged || !(!bSyKeyExists && !KeyParams.bAutoKeying && Scale[1].GetNumKeys() > 0) ) )
 	{
 		AddKeyToCurve(Scale[1], Time, TransformKey.Value.Scale.Y, KeyParams);
 	}
diff --git a/Engine/Source/Runtime/UMG/Private/Animation/MovieSceneMarginSection.cpp b/Engine/Source/Runtime/UMG/Private/Animation/MovieSceneMarginSection.cpp
--- a/Engine/Source/Runtime/UMG/Private/Animation/MovieSceneMarginSection.cpp
+++ b/Engine/Source/Runtime/UMG/Private/Animation/MovieSceneMarginSection.cpp
@@ -91,25 +91,25 @@ void UMovieSceneMarginSection::AddKey( float Time, const FMarginKey& MarginKey,
 	bool bBottomKeyExists = BottomCurve.IsKeyHandleValid(BottomCurve.FindKey(Time));
 
 	if ( (CurveName == NAME_None || CurveName == LeftName) &&
-			KeyParams.bAddKeyEvenIfUnchanged || !(!bLeftKeyExists && !KeyParams.bAutoKeying && LeftCurve.GetNumKeys() > 0) )
+			(KeyParams.bAddKeyEvenIfUnchanged || !(!bLeftKeyExists && !KeyParams.bAutoKeying && LeftCurve.GetNumKeys() > 0) ) )
 	{
 		AddKeyToCurve(LeftCurve, Time, MarginKey.Value.Left, KeyParams);
 	}
 
 	if ( (CurveName == NAME_None || CurveName == TopName) &&
-			KeyParams.bAddKeyEvenIfUnchanged || !(!bTopKeyExists && !KeyParams.bAutoKeying && TopCurve.GetNumKeys() > 0) )
+			(KeyParams.bAddKeyEvenIfUnchanged || !(!bTopKeyExists && !KeyParams.bAutoKeying && TopCurve.GetNumKeys() > 0) ) )
 	{
 		AddKeyToCurve(TopCurve, Time, MarginKey.Value.Top, KeyParams);
 	}
 
 	if ( (CurveName == NAME_None || CurveName == RightName) &&
-			KeyParams.bAddKeyEvenIfUnchanged || !(!bRightKeyExists && !KeyParams.bAutoKeying && RightCurve.GetNumKeys() > 0) )
+			(KeyParams.bAddKeyEvenIfUnchanged || !(!bRightKeyExists && !KeyParams.bAutoKeying && RightCurve.GetNumKeys() > 0) ) )
 	{
 		AddKeyToCurve(RightCurve, Time, MarginKey.Value.Right, KeyParams);
 	}
 
 	if ( (CurveName == NAME_None || CurveName == BottomName) &&
-			KeyParams.bAddKeyEvenIfUnchanged || !(!bBottomKeyExists && !KeyParams.bAutoKeying && BottomCurve.GetNumKeys() > 0) )
+			(KeyParams.bAddKeyEvenIfUnchanged || !(!bBottomKeyExists && !KeyParams.bAutoKeying && BottomCurve.GetNumKeys() > 0) ) )
 	{
 		AddKeyToCurve(BottomCurve, Time, MarginKey.Value.Bottom, KeyParams);
 	}
diff --git a/Engine/Source/Runtime/UMG/Public/Animation/MovieScene2DTransformTrack.h b/Engine/Source/Runtime/UMG/Public/Animation/MovieScene2DTransformTrack.h
--- a/Engine/Source/Runtime/UMG/Public/Animation/MovieScene2DTransformTrack.h
+++ b/Engine/Source/Runtime/UMG/Public/Animation/MovieScene2DTransformTrack.h
@@ -4,6 +4,7 @@
 
 #include "MovieScenePropertyTrack.h"
 #include "Slate/WidgetTransform.h"
+#include "KeyParams.h"
 #include "MovieScene2DTransformTrack.generated.h"
 
 struct F2DTransformKey
diff --git a/Engine/Source/Runtime/UMG/Public/Animation/MovieSceneMarginTrack.h b/Engine/Source/Runtime/UMG/Public/Animation/MovieSceneMarginTrack.h
--- a/Engine/Source/Runtime/UMG/Public/Animation/MovieSceneMarginTrack.h
+++ b/Engine/Source/Runtime/UMG/Public/Animation/MovieSceneMarginTrack.h
@@ -3,6 +3,7 @@
 #pragma once
 
 #include "MovieScenePropertyTrack.h"
+#include "KeyParams.h"
 #include "MovieSceneMarginTrack.generated.h"
 
 struct FMarginKey
