Sequencer: Missed one section in compile fix.

diff --git a/Engine/Source/Runtime/UMG/Private/Animation/MovieScene2DTransformSection.cpp b/Engine/Source/Runtime/UMG/Private/Animation/MovieScene2DTransformSection.cpp
--- a/Engine/Source/Runtime/UMG/Private/Animation/MovieScene2DTransformSection.cpp
+++ b/Engine/Source/Runtime/UMG/Private/Animation/MovieScene2DTransformSection.cpp
@@ -217,19 +217,19 @@ void UMovieScene2DTransformSection::AddKey( float Time, const struct F2DTransfor
 	}
 		
 	if ( (CurveName == NAME_None || CurveName == ShearName) &&
-			KeyParams.bAddKeyEvenIfUnchanged || !(!bHxKeyExists && !KeyParams.bAutoKeying && Shear[0].GetNumKeys() > 0) )
+			(KeyParams.bAddKeyEvenIfUnchanged || !(!bHxKeyExists && !KeyParams.bAutoKeying && Shear[0].GetNumKeys() > 0) ) )
 	{
 		AddKeyToCurve(Shear[0], Time, TransformKey.Value.Shear.X, KeyParams);
 	}
 
 	if ( (CurveName == NAME_None || CurveName == ShearName) &&
-			KeyParams.bAddKeyEvenIfUnchanged || !(!bHyKeyExists && !KeyParams.bAutoKeying && Shear[1].GetNumKeys() > 0) )
+			(KeyParams.bAddKeyEvenIfUnchanged || !(!bHyKeyExists && !KeyParams.bAutoKeying && Shear[1].GetNumKeys() > 0) ) )
 	{
 		AddKeyToCurve(Shear[1], Time, TransformKey.Value.Shear.Y, KeyParams);
 	}
 
 	if ( (CurveName == NAME_None || CurveName == AngleName) &&
-			KeyParams.bAddKeyEvenIfUnchanged || !(!bRKeyExists && !KeyParams.bAutoKeying && Rotation.GetNumKeys() > 0) )
+			(KeyParams.bAddKeyEvenIfUnchanged || !(!bRKeyExists && !KeyParams.bAutoKeying && Rotation.GetNumKeys() > 0) ) )
 	{
 		AddKeyToCurve(Rotation, Time, TransformKey.Value.Angle, KeyParams);
 	}
