SafeZone should default to SelfHitTestInvisible, IsTitleSafe should start as true because that's kinda the whole point of the widget.

diff --git a/Engine/Source/Runtime/UMG/Private/Components/SafeZone.cpp b/Engine/Source/Runtime/UMG/Private/Components/SafeZone.cpp
--- a/Engine/Source/Runtime/UMG/Private/Components/SafeZone.cpp
+++ b/Engine/Source/Runtime/UMG/Private/Components/SafeZone.cpp
@@ -11,6 +11,7 @@
 USafeZone::USafeZone()
 {
 	bCanHaveMultipleChildren = false;
+	Visibility = ESlateVisibility::SelfHitTestInvisible;	
 }
 
 #if WITH_EDITOR
diff --git a/Engine/Source/Runtime/UMG/Private/Components/SafeZoneSlot.cpp b/Engine/Source/Runtime/UMG/Private/Components/SafeZoneSlot.cpp
--- a/Engine/Source/Runtime/UMG/Private/Components/SafeZoneSlot.cpp
+++ b/Engine/Source/Runtime/UMG/Private/Components/SafeZoneSlot.cpp
@@ -7,7 +7,7 @@
 
 USafeZoneSlot::USafeZoneSlot()
 {
-
+	bIsTitleSafe = true;
 }
 
 void USafeZoneSlot::SynchronizeProperties()
