Added Virtual Keyboard support to the EditableText widget.

diff --git a/Engine/Source/Runtime/UMG/Private/Components/EditableText.cpp b/Engine/Source/Runtime/UMG/Private/Components/EditableText.cpp
--- a/Engine/Source/Runtime/UMG/Private/Components/EditableText.cpp
+++ b/Engine/Source/Runtime/UMG/Private/Components/EditableText.cpp
@@ -39,17 +39,17 @@ void UEditableText::ReleaseSlateResources(bool bReleaseChildren)
 
 TSharedRef<SWidget> UEditableText::RebuildWidget()
 {
-	MyEditableText = SNew(SEditableText)
-	.Style(&WidgetStyle)
-	.MinDesiredWidth(MinimumDesiredWidth)
-	.IsCaretMovedWhenGainFocus(IsCaretMovedWhenGainFocus)
-	.SelectAllTextWhenFocused(SelectAllTextWhenFocused)
-	.RevertTextOnEscape(RevertTextOnEscape)
-	.ClearKeyboardFocusOnCommit(ClearKeyboardFocusOnCommit)
-	.SelectAllTextOnCommit(SelectAllTextOnCommit)
-	.OnTextChanged(BIND_UOBJECT_DELEGATE(FOnTextChanged, HandleOnTextChanged))
-	.OnTextCommitted(BIND_UOBJECT_DELEGATE(FOnTextCommitted, HandleOnTextCommitted))
-	;
+	MyEditableText = SNew( SEditableText )
+		.Style( &WidgetStyle )
+		.MinDesiredWidth( MinimumDesiredWidth )
+		.IsCaretMovedWhenGainFocus( IsCaretMovedWhenGainFocus )
+		.SelectAllTextWhenFocused( SelectAllTextWhenFocused )
+		.RevertTextOnEscape( RevertTextOnEscape )
+		.ClearKeyboardFocusOnCommit( ClearKeyboardFocusOnCommit )
+		.SelectAllTextOnCommit( SelectAllTextOnCommit )
+		.OnTextChanged( BIND_UOBJECT_DELEGATE( FOnTextChanged, HandleOnTextChanged ) )
+		.OnTextCommitted( BIND_UOBJECT_DELEGATE( FOnTextCommitted, HandleOnTextCommitted ) )
+		.VirtualKeyboardType( EVirtualKeyboardType::AsKeyboardType( KeyboardType.GetValue() ) );
 	
 	return BuildDesignTimeWidget( MyEditableText.ToSharedRef() );
 }
diff --git a/Engine/Source/Runtime/UMG/Public/Components/EditableText.h b/Engine/Source/Runtime/UMG/Public/Components/EditableText.h
--- a/Engine/Source/Runtime/UMG/Public/Components/EditableText.h
+++ b/Engine/Source/Runtime/UMG/Public/Components/EditableText.h
@@ -98,6 +98,10 @@ public:
 	UPROPERTY(EditAnywhere, Category=Behavior, AdvancedDisplay)
 	bool SelectAllTextOnCommit;
 
+	/** If we're on a platform that requires a virtual keyboard, what kind of keyboard should this widget use? */
+	UPROPERTY( EditAnywhere, Category = Behavior, AdvancedDisplay )
+	TEnumAsByte< EVirtualKeyboardType::Type > KeyboardType;
+
 public:
 
 	/** Called whenever the text is changed interactively by the user */
diff --git a/Engine/Source/Runtime/UMG/Public/Components/SlateWrapperTypes.h b/Engine/Source/Runtime/UMG/Public/Components/SlateWrapperTypes.h
--- a/Engine/Source/Runtime/UMG/Public/Components/SlateWrapperTypes.h
+++ b/Engine/Source/Runtime/UMG/Public/Components/SlateWrapperTypes.h
@@ -84,3 +84,42 @@ struct FSlateChildSize
 	{
 	}
 };
+
+
+UENUM( BlueprintType )
+namespace EVirtualKeyboardType
+{
+	enum Type
+	{
+		Default,
+		Number,
+		Web,
+		Email,
+		Password,
+		AlphaNumeric
+	};
+}
+
+namespace EVirtualKeyboardType
+{
+	static EKeyboardType AsKeyboardType( Type InType )
+	{
+		switch ( InType )
+		{
+		case Type::Default:
+			return EKeyboardType::Keyboard_Default;
+		case Type::Number:
+			return EKeyboardType::Keyboard_Number;
+		case Type::Web:
+			return EKeyboardType::Keyboard_Web;
+		case Type::Email:
+			return EKeyboardType::Keyboard_Email;
+		case Type::Password:
+			return EKeyboardType::Keyboard_Password;
+		case Type::AlphaNumeric:
+			return EKeyboardType::Keyboard_AlphaNumeric;
+		}
+
+		return EKeyboardType::Keyboard_Default;
+	}
+}
\ No newline at end of file
