Preparing UMG for converting BP-Widgets into native code.

diff --git a/Engine/Source/Editor/UMGEditor/Private/UMGEditorModule.cpp b/Engine/Source/Editor/UMGEditor/Private/UMGEditorModule.cpp
--- a/Engine/Source/Editor/UMGEditor/Private/UMGEditorModule.cpp
+++ b/Engine/Source/Editor/UMGEditor/Private/UMGEditorModule.cpp
@@ -107,7 +107,9 @@ public:
 
 	void PreCompile(UBlueprint* Blueprint, const FKismetCompilerOptions& CompileOptions) override
 	{
-		if ( ReRegister == nullptr && CanCompile(Blueprint) && CompileOptions.CompileType == EKismetCompileType::Full )
+		if ( ReRegister == nullptr 
+			&& CanCompile(Blueprint) 
+			&& (CompileOptions.CompileType == EKismetCompileType::Full || CompileOptions.CompileType == EKismetCompileType::Cpp))
 		{
 			ReRegister = new TComponentReregisterContext<UWidgetComponent>();
 		}
diff --git a/Engine/Source/Runtime/UMG/Private/UserWidget.cpp b/Engine/Source/Runtime/UMG/Private/UserWidget.cpp
--- a/Engine/Source/Runtime/UMG/Private/UserWidget.cpp
+++ b/Engine/Source/Runtime/UMG/Private/UserWidget.cpp
@@ -49,6 +49,10 @@ bool UUserWidget::Initialize()
 		{
 			BGClass->InitializeWidget(this);
 		}
+		else
+		{
+			CustomNativeInitilize();
+		}
 
 		if ( WidgetTree == nullptr )
 		{
@@ -686,6 +690,20 @@ FVector2D UUserWidget::GetFullScreenAlignment() const
 	return ViewportAlignment;
 }
 
+void UUserWidget::RemoveObsoleteBindings(const TArray<FName>& NamedSlots)
+{
+	for (int32 BindingIndex = 0; BindingIndex < NamedSlotBindings.Num(); BindingIndex++)
+	{
+		const FNamedSlotBinding& Binding = NamedSlotBindings[BindingIndex];
+
+		if (!NamedSlots.Contains(Binding.Name))
+		{
+			NamedSlotBindings.RemoveAt(BindingIndex);
+			BindingIndex--;
+		}
+	}
+}
+
 void UUserWidget::PreSave()
 {
 	Super::PreSave();
@@ -693,16 +711,7 @@ void UUserWidget::PreSave()
 	// Remove bindings that are no longer contained in the class.
 	if ( UWidgetBlueprintGeneratedClass* BGClass = Cast<UWidgetBlueprintGeneratedClass>(GetClass()) )
 	{
-		for ( int32 BindingIndex = 0; BindingIndex < NamedSlotBindings.Num(); BindingIndex++ )
-		{
-			const FNamedSlotBinding& Binding = NamedSlotBindings[BindingIndex];
-
-			if ( !BGClass->NamedSlots.Contains(Binding.Name) )
-			{
-				NamedSlotBindings.RemoveAt(BindingIndex);
-				BindingIndex--;
-			}
-		}
+		RemoveObsoleteBindings(BGClass->NamedSlots);
 	}
 }
 
diff --git a/Engine/Source/Runtime/UMG/Private/WidgetBlueprintGeneratedClass.cpp b/Engine/Source/Runtime/UMG/Private/WidgetBlueprintGeneratedClass.cpp
--- a/Engine/Source/Runtime/UMG/Private/WidgetBlueprintGeneratedClass.cpp
+++ b/Engine/Source/Runtime/UMG/Private/WidgetBlueprintGeneratedClass.cpp
@@ -15,24 +15,33 @@ UWidgetBlueprintGeneratedClass::UWidgetBlueprintGeneratedClass(const FObjectInit
 {
 }
 
-void UWidgetBlueprintGeneratedClass::InitializeWidget(UUserWidget* UserWidget) const
+void UWidgetBlueprintGeneratedClass::InitializeWidgetStatic(UUserWidget* UserWidget
+	, const UClass* InClass
+	, bool InCanEverTick
+	, bool InCanEverPaint
+	, UWidgetTree* InWidgetTree
+	, const TArray< UWidgetAnimation* >& InAnimations
+	, const TArray< FDelegateRuntimeBinding >& InBindings
+	, UWidgetTree* InDesignerWidgetTree)
 {
-	UserWidget->bCanEverTick = bCanEverTick;
-	UserWidget->bCanEverPaint = bCanEverPaint;
+	check(InClass);
+
+	UserWidget->bCanEverTick = InCanEverTick;
+	UserWidget->bCanEverPaint = InCanEverPaint;
 
 #if WITH_EDITORONLY_DATA
-	UWidgetTree* ClonedTree = DuplicateObject<UWidgetTree>(DesignerWidgetTree ? DesignerWidgetTree : WidgetTree, UserWidget);
+	UWidgetTree* ClonedTree = DuplicateObject<UWidgetTree>(InDesignerWidgetTree ? InDesignerWidgetTree : InWidgetTree, UserWidget);
 #else
-	UWidgetTree* ClonedTree = DuplicateObject<UWidgetTree>(WidgetTree, UserWidget);
+	UWidgetTree* ClonedTree = DuplicateObject<UWidgetTree>(InWidgetTree, UserWidget);
 #endif
 
-	UserWidget->WidgetGeneratedByClass = this;
+	UserWidget->WidgetGeneratedByClass = InClass;
 
 #if WITH_EDITOR
-	UserWidget->WidgetGeneratedBy = ClassGeneratedBy;
+	UserWidget->WidgetGeneratedBy = InClass->ClassGeneratedBy;
 #endif
 
-	if ( ClonedTree )
+	if (ClonedTree)
 	{
 		ClonedTree->SetFlags(RF_Transactional);
 
@@ -40,32 +49,32 @@ void UWidgetBlueprintGeneratedClass::InitializeWidget(UUserWidget* UserWidget) c
 
 		UClass* WidgetBlueprintClass = UserWidget->GetClass();
 
-		for(UWidgetAnimation* Animation : Animations)
+		for (UWidgetAnimation* Animation : InAnimations)
 		{
-			UWidgetAnimation* Anim = DuplicateObject<UWidgetAnimation>( Animation, UserWidget );
+			UWidgetAnimation* Anim = DuplicateObject<UWidgetAnimation>(Animation, UserWidget);
 
-			if( Anim->GetMovieScene() )
+			if (Anim->GetMovieScene())
 			{
 				// Find property with the same name as the template and assign the new widget to it.
 				UObjectPropertyBase* Prop = FindField<UObjectPropertyBase>(WidgetBlueprintClass, Anim->GetMovieScene()->GetFName());
-				if(Prop)
+				if (Prop)
 				{
 					Prop->SetObjectPropertyValue_InContainer(UserWidget, Anim);
 				}
 			}
 		}
 
-		ClonedTree->ForEachWidget([&] (UWidget* Widget) {
+		ClonedTree->ForEachWidget([&](UWidget* Widget) {
 			// Not fatal if NULL, but shouldn't happen
-			if ( !ensure(Widget != nullptr) )
+			if (!ensure(Widget != nullptr))
 			{
 				return;
 			}
 
-			Widget->WidgetGeneratedByClass = this;
+			Widget->WidgetGeneratedByClass = InClass;
 
 #if WITH_EDITOR
-			Widget->WidgetGeneratedBy = ClassGeneratedBy;
+			Widget->WidgetGeneratedBy = InClass->ClassGeneratedBy;
 #endif
 
 			// TODO UMG Make this an FName
@@ -75,7 +84,7 @@ void UWidgetBlueprintGeneratedClass::InitializeWidget(UUserWidget* UserWidget) c
 
 			// Find property with the same name as the template and assign the new widget to it.
 			UObjectPropertyBase* Prop = FindField<UObjectPropertyBase>(WidgetBlueprintClass, *VariableName);
-			if ( Prop )
+			if (Prop)
 			{
 				Prop->SetObjectPropertyValue_InContainer(UserWidget, Widget);
 				UObject* Value = Prop->GetObjectPropertyValue_InContainer(UserWidget);
@@ -83,22 +92,22 @@ void UWidgetBlueprintGeneratedClass::InitializeWidget(UUserWidget* UserWidget) c
 			}
 
 			// Perform binding
-			for ( const FDelegateRuntimeBinding& Binding : Bindings )
+			for (const FDelegateRuntimeBinding& Binding : InBindings)
 			{
 				//TODO UMG Make this faster.
-				if ( Binding.ObjectName == VariableName )
+				if (Binding.ObjectName == VariableName)
 				{
-					UDelegateProperty* DelegateProperty = FindField<UDelegateProperty>(Widget->GetClass(), FName(*( Binding.PropertyName.ToString() + TEXT("Delegate") )));
-					if ( !DelegateProperty )
+					UDelegateProperty* DelegateProperty = FindField<UDelegateProperty>(Widget->GetClass(), FName(*(Binding.PropertyName.ToString() + TEXT("Delegate"))));
+					if (!DelegateProperty)
 					{
 						DelegateProperty = FindField<UDelegateProperty>(Widget->GetClass(), Binding.PropertyName);
 					}
 
-					if ( DelegateProperty )
+					if (DelegateProperty)
 					{
 						bool bSourcePathBound = false;
 
-						if ( Binding.SourcePath.IsValid() )
+						if (Binding.SourcePath.IsValid())
 						{
 							bSourcePathBound = Widget->AddBinding(DelegateProperty, UserWidget, Binding.SourcePath);
 						}
@@ -106,10 +115,10 @@ void UWidgetBlueprintGeneratedClass::InitializeWidget(UUserWidget* UserWidget) c
 						// If no native binder is found then the only possibility is that the binding is for
 						// a delegate that doesn't match the known native binders available and so we
 						// fallback to just attempting to bind to the function directly.
-						if ( bSourcePathBound == false )
+						if (bSourcePathBound == false)
 						{
 							FScriptDelegate* ScriptDelegate = DelegateProperty->GetPropertyValuePtr_InContainer(Widget);
-							if ( ScriptDelegate )
+							if (ScriptDelegate)
 							{
 								ScriptDelegate->BindUFunction(UserWidget, Binding.FunctionName);
 							}
@@ -124,18 +133,29 @@ void UWidgetBlueprintGeneratedClass::InitializeWidget(UUserWidget* UserWidget) c
 				Widget->Navigation->ResolveExplictRules(ClonedTree);
 			}
 
-	#if WITH_EDITOR
+#if WITH_EDITOR
 			Widget->ConnectEditorData();
-	#endif
+#endif
 		});
 
 		// Bind any delegates on widgets
-		BindDynamicDelegates(this, UserWidget);
+		UBlueprintGeneratedClass::BindDynamicDelegates(InClass, UserWidget);
 
 		//TODO UMG Add OnWidgetInitialized?
 	}
 }
 
+void UWidgetBlueprintGeneratedClass::InitializeWidget(UUserWidget* UserWidget) const
+{
+	InitializeWidgetStatic(UserWidget, this, bCanEverTick, bCanEverPaint, WidgetTree, Animations, Bindings, 
+#if WITH_EDITOR
+		DesignerWidgetTree
+#else // WITH_EDITOR
+		nullptr
+#endif  // WITH_EDITOR
+		);
+}
+
 void UWidgetBlueprintGeneratedClass::PostLoad()
 {
 	Super::PostLoad();
diff --git a/Engine/Source/Runtime/UMG/Public/Blueprint/UserWidget.h b/Engine/Source/Runtime/UMG/Public/Blueprint/UserWidget.h
--- a/Engine/Source/Runtime/UMG/Public/Blueprint/UserWidget.h
+++ b/Engine/Source/Runtime/UMG/Public/Blueprint/UserWidget.h
@@ -155,6 +155,7 @@ public:
 	// End of UObject interface
 
 	virtual bool Initialize();
+	virtual void CustomNativeInitilize() {}
 
 	//UVisual interface
 	virtual void ReleaseSlateResources(bool bReleaseChildren) override;
@@ -920,6 +921,8 @@ protected:
 	 */
 	void TickActionsAndAnimation(const FGeometry& MyGeometry, float InDeltaTime);
 
+	void RemoveObsoleteBindings(const TArray<FName>& NamedSlots);
+
 private:
 	FAnchors ViewportAnchors;
 	FMargin ViewportOffsets;
diff --git a/Engine/Source/Runtime/UMG/Public/Blueprint/WidgetBlueprintGeneratedClass.h b/Engine/Source/Runtime/UMG/Public/Blueprint/WidgetBlueprintGeneratedClass.h
--- a/Engine/Source/Runtime/UMG/Public/Blueprint/WidgetBlueprintGeneratedClass.h
+++ b/Engine/Source/Runtime/UMG/Public/Blueprint/WidgetBlueprintGeneratedClass.h
@@ -97,4 +97,13 @@ public:
 	 * binding and wiring necessary to have the user's widget perform as desired.
 	 */
 	void InitializeWidget(UUserWidget* UserWidget) const;
+
+	static void InitializeWidgetStatic(UUserWidget* UserWidget
+		, const UClass* InClass
+		, bool InCanEverTick
+		, bool InCanEverPaint
+		, UWidgetTree* InWidgetTree
+		, const TArray< UWidgetAnimation* >& InAnimations
+		, const TArray< FDelegateRuntimeBinding >& InBindings
+		, UWidgetTree* InDesignerWidgetTree);
 };
\ No newline at end of file
