ScrollBox in UMG will now scroll to the selected descendent, or back to the begining when clicking on children in the layout tree.

diff --git a/Engine/Source/Runtime/UMG/Private/Components/ScrollBox.cpp b/Engine/Source/Runtime/UMG/Private/Components/ScrollBox.cpp
--- a/Engine/Source/Runtime/UMG/Private/Components/ScrollBox.cpp
+++ b/Engine/Source/Runtime/UMG/Private/Components/ScrollBox.cpp
@@ -179,6 +179,37 @@ const FText UScrollBox::GetPaletteCategory()
 	return LOCTEXT("Panel", "Panel");
 }
 
+void UScrollBox::OnDescendantSelected( UWidget* DescendantWidget )
+{
+	UWidget* SelectedChild = UWidget::FindChildContainingDescendant( this, DescendantWidget );
+	if ( SelectedChild )
+	{
+		ScrollWidgetIntoView( SelectedChild, true );
+
+		if ( TickHandle.IsValid() )
+		{
+			FTicker::GetCoreTicker().RemoveTicker( TickHandle );
+			TickHandle.Reset();
+		}
+	}
+}
+
+void UScrollBox::OnDescendantDeselected( UWidget* DescendantWidget )
+{
+	if ( TickHandle.IsValid() )
+	{
+		FTicker::GetCoreTicker().RemoveTicker( TickHandle );
+		TickHandle.Reset();
+	}
+
+	// because we get a deselect before we get a select, we need to delay this call until we're sure we didn't scroll to another widget.
+	TickHandle = FTicker::GetCoreTicker().AddTicker( FTickerDelegate::CreateLambda( [=]( float ) -> bool
+	                                                                                {
+		                                                                                this->ScrollToStart();
+		                                                                                return false;
+		                                                                            } ) );
+}
+
 #endif
 
 /////////////////////////////////////////////////////
diff --git a/Engine/Source/Runtime/UMG/Public/Components/ScrollBox.h b/Engine/Source/Runtime/UMG/Public/Components/ScrollBox.h
--- a/Engine/Source/Runtime/UMG/Public/Components/ScrollBox.h
+++ b/Engine/Source/Runtime/UMG/Public/Components/ScrollBox.h
@@ -96,6 +96,8 @@ public:
 	// UWidget interface
 	virtual const FSlateBrush* GetEditorIcon() override;
 	virtual const FText GetPaletteCategory() override;
+	virtual void OnDescendantSelected( UWidget* DescendantWidget ) override;
+	virtual void OnDescendantDeselected( UWidget* DescendantWidget ) override;
 	// End UWidget interface
 #endif
 
@@ -117,4 +119,8 @@ protected:
 	// UWidget interface
 	virtual TSharedRef<SWidget> RebuildWidget() override;
 	// End of UWidget interface
+
+#if WITH_EDITOR
+	FDelegateHandle TickHandle;
+#endif //WITH_EDITOR
 };
