UMG - Value property of slider now clamped to [0,1].
[UE-20565] Slider widget handle slides off of slider bar

diff --git a/Engine/Source/Runtime/UMG/Public/Components/Slider.h b/Engine/Source/Runtime/UMG/Public/Components/Slider.h
--- a/Engine/Source/Runtime/UMG/Public/Components/Slider.h
+++ b/Engine/Source/Runtime/UMG/Public/Components/Slider.h
@@ -20,7 +20,7 @@ class UMG_API USlider : public UWidget
 
 public:
 	/** The volume value to display. */
-	UPROPERTY(EditAnywhere, Category=Appearance)
+	UPROPERTY(EditAnywhere, Category=Appearance, meta=( ClampMin="0", ClampMax="1", UIMin="0", UIMax="1"))
 	float Value;
 
 	/** A bindable delegate to allow logic to drive the value of the widget */
