Add support for input actions in UserWidget

diff --git a/Engine/Source/Runtime/UMG/Private/UserWidget.cpp b/Engine/Source/Runtime/UMG/Private/UserWidget.cpp
--- a/Engine/Source/Runtime/UMG/Private/UserWidget.cpp
+++ b/Engine/Source/Runtime/UMG/Private/UserWidget.cpp
@@ -442,7 +442,7 @@ void UUserWidget::SetContentForSlot(FName SlotName, UWidget* Content)
 	}
 
 	// Dynamically insert the new widget into the hierarchy if it exists.
-	if (WidgetTree)
+	if ( WidgetTree )
 	{
 		UNamedSlot* NamedSlot = Cast<UNamedSlot>(WidgetTree->FindWidget(SlotName));
 		if (NamedSlot)
@@ -757,6 +757,7 @@ void UUserWidget::NativeConstruct()
 
 void UUserWidget::NativeDestruct()
 {
+	StopListeningForAllInputActions();
 	Destruct();
 }
 
@@ -814,12 +815,113 @@ void UUserWidget::TickActionsAndAnimation(const FGeometry& MyGeometry, float InD
 	}
 }
 
-void UUserWidget::NativePaint( FPaintContext& InContext ) const 
+void UUserWidget::ListenForInputAction( FName ActionName, FOnInputAction Callback )
+{
+	if ( !InputComponent )
+	{
+		if ( APlayerController* Controller = GetOwningPlayer() )
+		{
+			InputComponent = NewObject< UInputComponent >( this, NAME_None, RF_Transient );
+			InputComponent->bBlockInput = bStopAction;
+			InputComponent->Priority = Priority;
+			Controller->PushInputComponent( InputComponent );
+		}
+		else
+		{
+			FMessageLog( "PIE" ).Error( LOCTEXT( "GameNull", "Unable to listen to input actions without a player controller." ) );
+		}
+	}
+
+	if ( InputComponent )
+	{
+		FInputActionBinding NewBinding( ActionName, IE_Pressed );
+		NewBinding.ActionDelegate.GetDelegateForManualSet().BindUObject( this, &ThisClass::OnInputAction, Callback );
+
+		InputComponent->AddActionBinding( NewBinding );
+	}
+}
+
+void UUserWidget::StopListeningForInputAction( FName ActionName )
+{
+	if ( InputComponent )
+	{
+		for ( int32 ExistingIndex = InputComponent->GetNumActionBindings() - 1; ExistingIndex >= 0; --ExistingIndex )
+		{
+			const FInputActionBinding& ExistingBind = InputComponent->GetActionBinding( ExistingIndex );
+			if ( ExistingBind.ActionName == ActionName )
+			{
+				InputComponent->RemoveActionBinding( ExistingIndex );
+			}
+		}
+	}
+}
+
+void UUserWidget::StopListeningForAllInputActions()
+{
+	if ( InputComponent )
+	{
+		if ( APlayerController* Controller = GetOwningPlayer() )
+		{
+			Controller->PopInputComponent( InputComponent );
+		}
+		else
+		{
+			FMessageLog( "PIE" ).Error( LOCTEXT( "GameNull", "Unable to listen to input actions without a player controller." ) );
+		}
+
+		InputComponent->ClearActionBindings();
+		InputComponent = nullptr;
+	}
+}
+
+bool UUserWidget::IsListeningForInputAction( FName ActionName ) const
 {
-	if ( bCanEverPaint )
+	bool bResult = false;
+	if ( InputComponent )
 	{
-		OnPaint(InContext);
+		for ( int32 ExistingIndex = InputComponent->GetNumActionBindings() - 1; ExistingIndex >= 0; --ExistingIndex )
+		{
+			const FInputActionBinding& ExistingBind = InputComponent->GetActionBinding( ExistingIndex );
+			if ( ExistingBind.ActionName == ActionName )
+			{
+				bResult = true;
+				break;
+			}
+		}
 	}
+
+	return bResult;
+}
+
+void UUserWidget::SetInputActionPriority( int32 NewPriority )
+{
+	if ( InputComponent )
+	{
+		Priority = NewPriority;
+		InputComponent->Priority = Priority;
+	}
+}
+
+void UUserWidget::SetInputActionBlocking( bool bShouldBlock )
+{
+	if ( InputComponent )
+	{
+		bStopAction = bShouldBlock;
+		InputComponent->bBlockInput = bStopAction;
+	}
+}
+
+void UUserWidget::OnInputAction( FOnInputAction Callback )
+{
+	if ( GetIsEnabled() )
+	{
+		Callback.ExecuteIfBound();
+	}
+}
+
+void UUserWidget::NativePaint( FPaintContext& InContext ) const
+{
+	OnPaint( InContext );
 }
 
 bool UUserWidget::NativeIsInteractable() const
diff --git a/Engine/Source/Runtime/UMG/Public/Blueprint/UserWidget.h b/Engine/Source/Runtime/UMG/Public/Blueprint/UserWidget.h
--- a/Engine/Source/Runtime/UMG/Public/Blueprint/UserWidget.h
+++ b/Engine/Source/Runtime/UMG/Public/Blueprint/UserWidget.h
@@ -134,6 +134,8 @@ enum class EDesignPreviewSizeMode : uint8
 
 DECLARE_DYNAMIC_MULTICAST_DELEGATE(FOnConstructEvent);
 
+DECLARE_DYNAMIC_DELEGATE( FOnInputAction );
+
 /**
  * The user widget is extensible by users through the WidgetBlueprint.
  */
@@ -821,6 +823,12 @@ public:
 	UPROPERTY()
 	TArray<FNamedSlotBinding> NamedSlotBindings;
 
+	UPROPERTY(EditAnywhere, BlueprintReadWrite, Category = "Input")
+	bool bStopAction;
+
+	UPROPERTY( EditAnywhere, BlueprintReadOnly, Category = "Input" )
+	int32 Priority;
+
 #if WITH_EDITORONLY_DATA
 
 	/** Stores the design time desired size of the user widget */
@@ -922,6 +930,29 @@ protected:
 	void TickActionsAndAnimation(const FGeometry& MyGeometry, float InDeltaTime);
 
 	void RemoveObsoleteBindings(const TArray<FName>& NamedSlots);
+	
+	UFUNCTION( BlueprintCallable, Category = "Input", meta = ( BlueprintProtected = "true" ) )
+	void ListenForInputAction( FName ActionName, FOnInputAction Callback );
+
+	UFUNCTION( BlueprintCallable, Category = "Input", meta = ( BlueprintProtected = "true" ) )
+	void StopListeningForInputAction( FName ActionName );
+
+	UFUNCTION( BlueprintCallable, Category = "Input", meta = ( BlueprintProtected = "true" ) )
+	void StopListeningForAllInputActions();
+
+	UFUNCTION( BlueprintCallable, Category = "Input", meta = ( BlueprintProtected = "true" ) )
+	bool IsListeningForInputAction( FName ActionName ) const;
+
+	UFUNCTION( BlueprintCallable, Category = "Input", meta = ( BlueprintProtected = "true" ) )
+	void SetInputActionPriority( int32 NewPriority );
+
+	UFUNCTION( BlueprintCallable, Category = "Input", meta = ( BlueprintProtected = "true" ) )
+	void SetInputActionBlocking( bool bShouldBlock );
+
+	void OnInputAction( FOnInputAction Callback );
+
+	UPROPERTY( transient )
+	class UInputComponent* InputComponent;
 
 private:
 	FAnchors ViewportAnchors;
