PR #1560: Minor Fix: Parameter 'Position' (Contributed by [REDACTED_PERSON]); Also fixed an addition case not in the PR with the same problem.

diff --git a/Engine/Source/Runtime/UMG/Private/WidgetBlueprintLibrary.cpp b/Engine/Source/Runtime/UMG/Private/WidgetBlueprintLibrary.cpp
--- a/Engine/Source/Runtime/UMG/Private/WidgetBlueprintLibrary.cpp
+++ b/Engine/Source/Runtime/UMG/Private/WidgetBlueprintLibrary.cpp
@@ -167,7 +167,7 @@ void UWidgetBlueprintLibrary::DrawText(UPARAM(ref) FPaintContext& Context, const
 	FSlateDrawElement::MakeText(
 		Context.OutDrawElements,
 		Context.MaxLayer,
-		Context.AllottedGeometry.ToPaintGeometry(),
+		Context.AllottedGeometry.ToOffsetPaintGeometry(Position),
 		InString,
 		FontInfo,
 		Context.MyClippingRect,
@@ -187,7 +187,7 @@ void UWidgetBlueprintLibrary::DrawTextFormatted(UPARAM(ref) FPaintContext& Conte
 		FSlateDrawElement::MakeText(
 			Context.OutDrawElements,
 			Context.MaxLayer,
-			Context.AllottedGeometry.ToPaintGeometry(),
+			Context.AllottedGeometry.ToOffsetPaintGeometry(Position),
 			Text,
 			FontInfo,
 			Context.MyClippingRect,
