UE-20719, UE-20212, UE-20777 - UMG/Slate - Fixing the Left menu anchor position to offset the correct amount.  Slider Locked/HandleIndent properties are now respected.  DrawTextFormatted is now correctly static.  The WidgetComponent will now correctly hit test once again.

diff --git a/Engine/Source/Runtime/UMG/Private/Components/Slider.cpp b/Engine/Source/Runtime/UMG/Private/Components/Slider.cpp
--- a/Engine/Source/Runtime/UMG/Private/Components/Slider.cpp
+++ b/Engine/Source/Runtime/UMG/Private/Components/Slider.cpp
@@ -39,6 +39,8 @@ void USlider::SynchronizeProperties()
 	MySlider->SetSliderBarColor(SliderBarColor);
 	MySlider->SetSliderHandleColor(SliderHandleColor);
 	MySlider->SetValue(ValueBinding);
+	MySlider->SetLocked(Locked);
+	MySlider->SetIndentHandle(IndentHandle);
 }
 
 void USlider::HandleOnValueChanged(float InValue)
@@ -75,6 +77,24 @@ void USlider::SetValue(float InValue)
 	}
 }
 
+void USlider::SetIndentHandle(bool InIndentHandle)
+{
+	IndentHandle = InIndentHandle;
+	if ( MySlider.IsValid() )
+	{
+		MySlider->SetIndentHandle(InIndentHandle);
+	}
+}
+
+void USlider::SetLocked(bool InLocked)
+{
+	Locked = InLocked;
+	if ( MySlider.IsValid() )
+	{
+		MySlider->SetLocked(InLocked);
+	}
+}
+
 #if WITH_EDITOR
 
 const FSlateBrush* USlider::GetEditorIcon()
diff --git a/Engine/Source/Runtime/UMG/Private/Components/WidgetComponent.cpp b/Engine/Source/Runtime/UMG/Private/Components/WidgetComponent.cpp
--- a/Engine/Source/Runtime/UMG/Private/Components/WidgetComponent.cpp
+++ b/Engine/Source/Runtime/UMG/Private/Components/WidgetComponent.cpp
@@ -358,12 +358,14 @@ public:
 					{
 						if ( UWidgetComponent* WidgetComponent = Cast<UWidgetComponent>(HitComponent) )
 						{
+							// Get the "forward" vector based on the current rotation system.
+							const FVector ForwardVector = WidgetComponent->IsUsingLegacyRotation() ? WidgetComponent->GetUpVector() : WidgetComponent->GetForwardVector();
+
 							// Make sure the player is interacting with the front of the widget
-							// For widget components, the "front" faces the Z (or Up vector) direction
-							if ( FVector::DotProduct(WidgetComponent->GetUpVector(), CachedHitResult.ImpactPoint - CachedHitResult.TraceStart) < 0.f )
+							if ( FVector::DotProduct(ForwardVector, CachedHitResult.ImpactPoint - CachedHitResult.TraceStart) < 0.f )
 							{
 								// Make sure the player is close enough to the widget to interact with it
-								if ( FVector::DistSquared(CachedHitResult.TraceStart, CachedHitResult.ImpactPoint) <= FMath::Square(WidgetComponent->GetMaxInteractionDistance()) )
+								if ( FVector::DistSquared(CachedHitResult.TraceStart, WidgetComponent->GetComponentLocation()) <= FMath::Square(WidgetComponent->GetMaxInteractionDistance()) )
 								{
 									return WidgetComponent->GetHitWidgetPath(CachedHitResult.Location, bIgnoreEnabledStatus);
 								}
@@ -760,7 +762,7 @@ void UWidgetComponent::TickComponent(float DeltaTime, enum ELevelTick TickType,
 
 					// Paint the window
 					MaxLayerId = SlateWidget->Paint(
-						FPaintArgs(SlateWidget.ToSharedRef(), *HitTestGrid, FVector2D::ZeroVector, FApp::GetCurrentTime(), DeltaTime),
+						FPaintArgs(*SlateWidget.Get(), *HitTestGrid, FVector2D::ZeroVector, FApp::GetCurrentTime(), DeltaTime),
 						WindowGeometry, WindowGeometry.GetClippingRect(),
 						WindowElementList,
 						0,
@@ -1158,7 +1160,17 @@ void UWidgetComponent::UpdateBodySetup( bool bDrawSizeChanged )
 void UWidgetComponent::GetLocalHitLocation(FVector WorldHitLocation, FVector2D& OutLocalWidgetHitLocation) const
 {
 	// Find the hit location on the component
-	OutLocalWidgetHitLocation = FVector2D(ComponentToWorld.InverseTransformPosition(WorldHitLocation));
+	FVector ComponentHitLocation = ComponentToWorld.InverseTransformPosition(WorldHitLocation);
+
+	// Convert the 3D position of component space, into the 2D equivalent
+	if ( bUseLegacyRotation )
+	{
+		OutLocalWidgetHitLocation = FVector2D(ComponentHitLocation.X, ComponentHitLocation.Y);
+	}
+	else
+	{
+		OutLocalWidgetHitLocation = FVector2D(ComponentHitLocation.Y, ComponentHitLocation.Z);
+	}
 
 	// Offset the position by the pivot to get the position in widget space.
 	OutLocalWidgetHitLocation.X += DrawSize.X * Pivot.X;
diff --git a/Engine/Source/Runtime/UMG/Public/Blueprint/WidgetBlueprintLibrary.h b/Engine/Source/Runtime/UMG/Public/Blueprint/WidgetBlueprintLibrary.h
--- a/Engine/Source/Runtime/UMG/Public/Blueprint/WidgetBlueprintLibrary.h
+++ b/Engine/Source/Runtime/UMG/Public/Blueprint/WidgetBlueprintLibrary.h
@@ -87,7 +87,7 @@ public:
 	 * @param Tint			Color to render the line.
 	 */
 	UFUNCTION(BlueprintCallable, Category="Painting", meta=(DisplayName="Draw Text"))
-	void DrawTextFormatted(UPARAM(ref) FPaintContext& Context, const FText& Text, FVector2D Position, UFont* Font, int32 FontSize = 16, FName FontTypeFace = FName(TEXT("Regular")), FLinearColor Tint = FLinearColor::White);
+	static void DrawTextFormatted(UPARAM(ref) FPaintContext& Context, const FText& Text, FVector2D Position, UFont* Font, int32 FontSize = 16, FName FontTypeFace = FName(TEXT("Regular")), FLinearColor Tint = FLinearColor::White);
 
 	/**
 	 * The event reply to use when you choose to handle an event.  This will prevent the event 
diff --git a/Engine/Source/Runtime/UMG/Public/Components/Slider.h b/Engine/Source/Runtime/UMG/Public/Components/Slider.h
--- a/Engine/Source/Runtime/UMG/Public/Components/Slider.h
+++ b/Engine/Source/Runtime/UMG/Public/Components/Slider.h
@@ -74,6 +74,14 @@ public:
 	/** Sets the current value of the slider. */
 	UFUNCTION(BlueprintCallable, Category="Behavior")
 	void SetValue(float InValue);
+
+	/** Sets if the slidable area should be indented to fit the handle */
+	UFUNCTION(BlueprintCallable, Category="Behavior")
+	void SetIndentHandle(bool InValue);
+
+	/** Sets the handle to be interactive or fixed */
+	UFUNCTION(BlueprintCallable, Category="Behavior")
+	void SetLocked(bool InValue);
 	
 	// UWidget interface
 	virtual void SynchronizeProperties() override;
