Added boolean to ListenForInputAction that allows optional consumption of input.

diff --git a/Engine/Source/Runtime/UMG/Private/UserWidget.cpp b/Engine/Source/Runtime/UMG/Private/UserWidget.cpp
--- a/Engine/Source/Runtime/UMG/Private/UserWidget.cpp
+++ b/Engine/Source/Runtime/UMG/Private/UserWidget.cpp
@@ -815,7 +815,7 @@ void UUserWidget::TickActionsAndAnimation(const FGeometry& MyGeometry, float InD
 	}
 }
 
-void UUserWidget::ListenForInputAction( FName ActionName, FOnInputAction Callback )
+void UUserWidget::ListenForInputAction( FName ActionName, bool bConsume, FOnInputAction Callback )
 {
 	if ( !InputComponent )
 	{
@@ -835,6 +835,7 @@ void UUserWidget::ListenForInputAction( FName ActionName, FOnInputAction Callbac
 	if ( InputComponent )
 	{
 		FInputActionBinding NewBinding( ActionName, IE_Pressed );
+		NewBinding.bConsumeInput = bConsume;
 		NewBinding.ActionDelegate.GetDelegateForManualSet().BindUObject( this, &ThisClass::OnInputAction, Callback );
 
 		InputComponent->AddActionBinding( NewBinding );
diff --git a/Engine/Source/Runtime/UMG/Public/Blueprint/UserWidget.h b/Engine/Source/Runtime/UMG/Public/Blueprint/UserWidget.h
--- a/Engine/Source/Runtime/UMG/Public/Blueprint/UserWidget.h
+++ b/Engine/Source/Runtime/UMG/Public/Blueprint/UserWidget.h
@@ -932,7 +932,7 @@ protected:
 	void RemoveObsoleteBindings(const TArray<FName>& NamedSlots);
 	
 	UFUNCTION( BlueprintCallable, Category = "Input", meta = ( BlueprintProtected = "true" ) )
-	void ListenForInputAction( FName ActionName, FOnInputAction Callback );
+	void ListenForInputAction( FName ActionName, bool bConsume, FOnInputAction Callback );
 
 	UFUNCTION( BlueprintCallable, Category = "Input", meta = ( BlueprintProtected = "true" ) )
 	void StopListeningForInputAction( FName ActionName );
