Button now has OnRollOver and OnRollOut events.

diff --git a/Engine/Source/Runtime/Slate/Private/Widgets/Input/SButton.cpp b/Engine/Source/Runtime/Slate/Private/Widgets/Input/SButton.cpp
--- a/Engine/Source/Runtime/Slate/Private/Widgets/Input/SButton.cpp
+++ b/Engine/Source/Runtime/Slate/Private/Widgets/Input/SButton.cpp
@@ -55,6 +55,8 @@ void SButton::Construct( const FArguments& InArgs )
 	OnClicked = InArgs._OnClicked;
 	OnPressed = InArgs._OnPressed;
 	OnReleased = InArgs._OnReleased;
+	OnRollOver = InArgs._OnRollOver;
+	OnRollOut = InArgs._OnRollOut;
 
 	ClickMethod = InArgs._ClickMethod;
 	TouchMethod = InArgs._TouchMethod;
@@ -306,6 +308,8 @@ void SButton::OnMouseEnter( const FGeometry& MyGeometry, const FPointerEvent& Mo
 	
 	SBorder::OnMouseEnter( MyGeometry, MouseEvent );
 
+	OnRollOver.ExecuteIfBound();
+
 	Invalidate(EInvalidateWidget::Layout);
 }
 
@@ -321,6 +325,8 @@ void SButton::OnMouseLeave( const FPointerEvent& MouseEvent )
 		Release();
 	}
 
+	OnRollOut.ExecuteIfBound();
+
 	Invalidate(EInvalidateWidget::Layout);
 }
 
diff --git a/Engine/Source/Runtime/Slate/Public/Widgets/Input/SButton.h b/Engine/Source/Runtime/Slate/Public/Widgets/Input/SButton.h
--- a/Engine/Source/Runtime/Slate/Public/Widgets/Input/SButton.h
+++ b/Engine/Source/Runtime/Slate/Public/Widgets/Input/SButton.h
@@ -59,6 +59,10 @@ public:
 		/** Called when the button is released */
 		SLATE_EVENT( FSimpleDelegate, OnReleased )
 
+		SLATE_EVENT( FSimpleDelegate, OnRollOver )
+
+		SLATE_EVENT( FSimpleDelegate, OnRollOut )
+
 		/** Sets the rules to use for determining whether the button was clicked.  This is an advanced setting and generally should be left as the default. */
 		SLATE_ARGUMENT( EButtonClickMethod::Type, ClickMethod )
 
@@ -180,6 +184,10 @@ protected:
 	/** The delegate to execute when the button is released */
 	FSimpleDelegate OnReleased;
 
+	FSimpleDelegate OnRollOver;
+
+	FSimpleDelegate OnRollOut;
+
 	/** Style resource for the button */
 	const FButtonStyle* Style;
 
diff --git a/Engine/Source/Runtime/UMG/Private/Components/Button.cpp b/Engine/Source/Runtime/UMG/Private/Components/Button.cpp
--- a/Engine/Source/Runtime/UMG/Private/Components/Button.cpp
+++ b/Engine/Source/Runtime/UMG/Private/Components/Button.cpp
@@ -35,6 +35,8 @@ TSharedRef<SWidget> UButton::RebuildWidget()
 		.OnClicked(BIND_UOBJECT_DELEGATE(FOnClicked, SlateHandleClicked))
 		.OnPressed(BIND_UOBJECT_DELEGATE(FSimpleDelegate, SlateHandlePressed))
 		.OnReleased(BIND_UOBJECT_DELEGATE(FSimpleDelegate, SlateHandleReleased))
+		.OnRollOver_UObject( this, &ThisClass::SlateHandleRollOver )
+		.OnRollOut_UObject( this, &ThisClass::SlateHandleRollOut )
 		.ButtonStyle(&WidgetStyle)
 		.ClickMethod(ClickMethod)
 		.TouchMethod(TouchMethod)
@@ -166,6 +168,16 @@ void UButton::SlateHandleReleased()
 	OnReleased.Broadcast();
 }
 
+void UButton::SlateHandleRollOver()
+{
+	OnRollOver.Broadcast();
+}
+
+void UButton::SlateHandleRollOut()
+{
+	OnRollOut.Broadcast();
+}
+
 #if WITH_EDITOR
 
 const FSlateBrush* UButton::GetEditorIcon()
diff --git a/Engine/Source/Runtime/UMG/Public/Components/Button.h b/Engine/Source/Runtime/UMG/Public/Components/Button.h
--- a/Engine/Source/Runtime/UMG/Public/Components/Button.h
+++ b/Engine/Source/Runtime/UMG/Public/Components/Button.h
@@ -9,6 +9,7 @@
 DECLARE_DYNAMIC_MULTICAST_DELEGATE(FOnButtonClickedEvent);
 DECLARE_DYNAMIC_MULTICAST_DELEGATE(FOnButtonPressedEvent);
 DECLARE_DYNAMIC_MULTICAST_DELEGATE(FOnButtonReleasedEvent);
+DECLARE_DYNAMIC_MULTICAST_DELEGATE(FOnButtonHoverEvent);
 
 /**
  * The button is a click-able primitive widget to enable basic interaction, you
@@ -66,6 +67,12 @@ public:
 	UPROPERTY(BlueprintAssignable, Category="Button|Event")
 	FOnButtonReleasedEvent OnReleased;
 
+	UPROPERTY( BlueprintAssignable, Category = "Button|Event" )
+	FOnButtonHoverEvent OnRollOver;
+
+	UPROPERTY( BlueprintAssignable, Category = "Button|Event" )
+	FOnButtonHoverEvent OnRollOut;
+
 public:
 	
 	/** Sets the color multiplier for the button background */
@@ -120,6 +127,8 @@ protected:
 	FReply SlateHandleClicked();
 	void SlateHandlePressed();
 	void SlateHandleReleased();
+	void SlateHandleRollOver();
+	void SlateHandleRollOut();
 
 protected:
 	//~ Begin UWidget Interface
