Sequencer: Fix non autokeying behavior for 2d transform and margin sections.

diff --git a/Engine/Source/Runtime/UMG/Private/Animation/MovieScene2DTransformSection.cpp b/Engine/Source/Runtime/UMG/Private/Animation/MovieScene2DTransformSection.cpp
--- a/Engine/Source/Runtime/UMG/Private/Animation/MovieScene2DTransformSection.cpp
+++ b/Engine/Source/Runtime/UMG/Private/Animation/MovieScene2DTransformSection.cpp
@@ -141,30 +141,9 @@ bool UMovieScene2DTransformSection::NewKeyIsNewData( float Time, const FWidgetTr
 
 	bHasEmptyKeys |= Rotation.GetNumKeys() == 0;
 
-	if (bHasEmptyKeys || Eval(Time, Transform) != Transform)
+	if (bHasEmptyKeys || (KeyParams.bAutoKeying && Eval(Time, Transform) != Transform))
 	{
-		// Don't add a keyframe if there are existing keys and auto key is not enabled.
-		bool bTxKeyExists = Translation[0].IsKeyHandleValid(Translation[0].FindKey(Time));
-		bool bTyKeyExists = Translation[1].IsKeyHandleValid(Translation[1].FindKey(Time));
-
-		bool bRKeyExists = Rotation.IsKeyHandleValid(Rotation.FindKey(Time));
-
-		bool bHxKeyExists = Shear[0].IsKeyHandleValid(Shear[0].FindKey(Time));
-		bool bHyKeyExists = Shear[1].IsKeyHandleValid(Shear[1].FindKey(Time));
-
-		bool bSxKeyExists = Scale[0].IsKeyHandleValid(Scale[0].FindKey(Time));
-		bool bSyKeyExists = Scale[1].IsKeyHandleValid(Scale[1].FindKey(Time));
-
-		if ( !(!bTxKeyExists && !KeyParams.bAutoKeying && Translation[0].GetNumKeys() > 0) ||
-			 !(!bTyKeyExists && !KeyParams.bAutoKeying && Translation[1].GetNumKeys() > 0) ||
-			 !(!bRKeyExists && !KeyParams.bAutoKeying && Rotation.GetNumKeys() > 0) ||
-			 !(!bHxKeyExists && !KeyParams.bAutoKeying && Shear[0].GetNumKeys() > 0) ||
-			 !(!bHyKeyExists && !KeyParams.bAutoKeying && Shear[1].GetNumKeys() > 0) ||
-			 !(!bSxKeyExists && !KeyParams.bAutoKeying && Scale[0].GetNumKeys() > 0) ||
-			 !(!bSyKeyExists && !KeyParams.bAutoKeying && Scale[1].GetNumKeys() > 0) )
-		{
-			return true;
-		}
+		return true;
 	}
 
 	return false;
diff --git a/Engine/Source/Runtime/UMG/Private/Animation/MovieSceneMarginSection.cpp b/Engine/Source/Runtime/UMG/Private/Animation/MovieSceneMarginSection.cpp
--- a/Engine/Source/Runtime/UMG/Private/Animation/MovieSceneMarginSection.cpp
+++ b/Engine/Source/Runtime/UMG/Private/Animation/MovieSceneMarginSection.cpp
@@ -123,21 +123,9 @@ bool UMovieSceneMarginSection::NewKeyIsNewData(float Time, const FMargin& Value,
 		BottomCurve.GetNumKeys() == 0 ||
 		(Eval(Time,Value) != Value);
 
-	if (bHasEmptyKeys)
+	if (bHasEmptyKeys || (KeyParams.bAutoKeying && Eval(Time, Value) != Value))
 	{
-		// Don't add a keyframe if there are existing keys and auto key is not enabled.
-		bool bTopKeyExists = TopCurve.IsKeyHandleValid(TopCurve.FindKey(Time));
-		bool bLeftKeyExists = LeftCurve.IsKeyHandleValid(LeftCurve.FindKey(Time));
-		bool bRightKeyExists = RightCurve.IsKeyHandleValid(RightCurve.FindKey(Time));
-		bool bBottomKeyExists = BottomCurve.IsKeyHandleValid(BottomCurve.FindKey(Time));
-
-		if ( !(!bTopKeyExists && !KeyParams.bAutoKeying && TopCurve.GetNumKeys() > 0) ||
-			 !(!bLeftKeyExists && !KeyParams.bAutoKeying && LeftCurve.GetNumKeys() > 0) ||
-			 !(!bRightKeyExists && !KeyParams.bAutoKeying && RightCurve.GetNumKeys() > 0) ||
-			 !(!bBottomKeyExists && !KeyParams.bAutoKeying && BottomCurve.GetNumKeys() > 0) )
-		{
-			return true;
-		}
+		return true;
 	}
 	return false;
 }
