Renamed OnRollOver/OnRollOut to OnHovered/OnUnhovered

diff --git a/Engine/Source/Runtime/Slate/Private/Widgets/Input/SButton.cpp b/Engine/Source/Runtime/Slate/Private/Widgets/Input/SButton.cpp
--- a/Engine/Source/Runtime/Slate/Private/Widgets/Input/SButton.cpp
+++ b/Engine/Source/Runtime/Slate/Private/Widgets/Input/SButton.cpp
@@ -55,8 +55,8 @@ void SButton::Construct( const FArguments& InArgs )
 	OnClicked = InArgs._OnClicked;
 	OnPressed = InArgs._OnPressed;
 	OnReleased = InArgs._OnReleased;
-	OnRollOver = InArgs._OnRollOver;
-	OnRollOut = InArgs._OnRollOut;
+	OnHovered = InArgs._OnHovered;
+	OnUnhovered = InArgs._OnUnhovered;
 
 	ClickMethod = InArgs._ClickMethod;
 	TouchMethod = InArgs._TouchMethod;
@@ -308,7 +308,7 @@ void SButton::OnMouseEnter( const FGeometry& MyGeometry, const FPointerEvent& Mo
 	
 	SBorder::OnMouseEnter( MyGeometry, MouseEvent );
 
-	OnRollOver.ExecuteIfBound();
+	OnHovered.ExecuteIfBound();
 
 	Invalidate(EInvalidateWidget::Layout);
 }
@@ -325,7 +325,7 @@ void SButton::OnMouseLeave( const FPointerEvent& MouseEvent )
 		Release();
 	}
 
-	OnRollOut.ExecuteIfBound();
+	OnUnhovered.ExecuteIfBound();
 
 	Invalidate(EInvalidateWidget::Layout);
 }
diff --git a/Engine/Source/Runtime/Slate/Public/Widgets/Input/SButton.h b/Engine/Source/Runtime/Slate/Public/Widgets/Input/SButton.h
--- a/Engine/Source/Runtime/Slate/Public/Widgets/Input/SButton.h
+++ b/Engine/Source/Runtime/Slate/Public/Widgets/Input/SButton.h
@@ -59,9 +59,9 @@ public:
 		/** Called when the button is released */
 		SLATE_EVENT( FSimpleDelegate, OnReleased )
 
-		SLATE_EVENT( FSimpleDelegate, OnRollOver )
+		SLATE_EVENT( FSimpleDelegate, OnHovered )
 
-		SLATE_EVENT( FSimpleDelegate, OnRollOut )
+		SLATE_EVENT( FSimpleDelegate, OnUnhovered )
 
 		/** Sets the rules to use for determining whether the button was clicked.  This is an advanced setting and generally should be left as the default. */
 		SLATE_ARGUMENT( EButtonClickMethod::Type, ClickMethod )
@@ -184,9 +184,9 @@ protected:
 	/** The delegate to execute when the button is released */
 	FSimpleDelegate OnReleased;
 
-	FSimpleDelegate OnRollOver;
+	FSimpleDelegate OnHovered;
 
-	FSimpleDelegate OnRollOut;
+	FSimpleDelegate OnUnhovered;
 
 	/** Style resource for the button */
 	const FButtonStyle* Style;
diff --git a/Engine/Source/Runtime/UMG/Private/Components/Button.cpp b/Engine/Source/Runtime/UMG/Private/Components/Button.cpp
--- a/Engine/Source/Runtime/UMG/Private/Components/Button.cpp
+++ b/Engine/Source/Runtime/UMG/Private/Components/Button.cpp
@@ -35,8 +35,8 @@ TSharedRef<SWidget> UButton::RebuildWidget()
 		.OnClicked(BIND_UOBJECT_DELEGATE(FOnClicked, SlateHandleClicked))
 		.OnPressed(BIND_UOBJECT_DELEGATE(FSimpleDelegate, SlateHandlePressed))
 		.OnReleased(BIND_UOBJECT_DELEGATE(FSimpleDelegate, SlateHandleReleased))
-		.OnRollOver_UObject( this, &ThisClass::SlateHandleRollOver )
-		.OnRollOut_UObject( this, &ThisClass::SlateHandleRollOut )
+		.OnHovered_UObject( this, &ThisClass::SlateHandleHovered )
+		.OnUnhovered_UObject( this, &ThisClass::SlateHandleUnhovered )
 		.ButtonStyle(&WidgetStyle)
 		.ClickMethod(ClickMethod)
 		.TouchMethod(TouchMethod)
@@ -168,14 +168,14 @@ void UButton::SlateHandleReleased()
 	OnReleased.Broadcast();
 }
 
-void UButton::SlateHandleRollOver()
+void UButton::SlateHandleHovered()
 {
-	OnRollOver.Broadcast();
+	OnHovered.Broadcast();
 }
 
-void UButton::SlateHandleRollOut()
+void UButton::SlateHandleUnhovered()
 {
-	OnRollOut.Broadcast();
+	OnUnhovered.Broadcast();
 }
 
 #if WITH_EDITOR
diff --git a/Engine/Source/Runtime/UMG/Public/Components/Button.h b/Engine/Source/Runtime/UMG/Public/Components/Button.h
--- a/Engine/Source/Runtime/UMG/Public/Components/Button.h
+++ b/Engine/Source/Runtime/UMG/Public/Components/Button.h
@@ -68,10 +68,10 @@ public:
 	FOnButtonReleasedEvent OnReleased;
 
 	UPROPERTY( BlueprintAssignable, Category = "Button|Event" )
-	FOnButtonHoverEvent OnRollOver;
+	FOnButtonHoverEvent OnHovered;
 
 	UPROPERTY( BlueprintAssignable, Category = "Button|Event" )
-	FOnButtonHoverEvent OnRollOut;
+	FOnButtonHoverEvent OnUnhovered;
 
 public:
 	
@@ -127,8 +127,8 @@ protected:
 	FReply SlateHandleClicked();
 	void SlateHandlePressed();
 	void SlateHandleReleased();
-	void SlateHandleRollOver();
-	void SlateHandleRollOut();
+	void SlateHandleHovered();
+	void SlateHandleUnhovered();
 
 protected:
 	//~ Begin UWidget Interface
