Added support for event type on UserWidget input action listeners

diff --git a/Engine/Source/Runtime/UMG/Private/UserWidget.cpp b/Engine/Source/Runtime/UMG/Private/UserWidget.cpp
--- a/Engine/Source/Runtime/UMG/Private/UserWidget.cpp
+++ b/Engine/Source/Runtime/UMG/Private/UserWidget.cpp
@@ -815,7 +815,7 @@ void UUserWidget::TickActionsAndAnimation(const FGeometry& MyGeometry, float InD
 	}
 }
 
-void UUserWidget::ListenForInputAction( FName ActionName, bool bConsume, FOnInputAction Callback )
+void UUserWidget::ListenForInputAction( FName ActionName, TEnumAsByte< EInputEvent > EventType, bool bConsume, FOnInputAction Callback )
 {
 	if ( !InputComponent )
 	{
@@ -834,7 +834,7 @@ void UUserWidget::ListenForInputAction( FName ActionName, bool bConsume, FOnInpu
 
 	if ( InputComponent )
 	{
-		FInputActionBinding NewBinding( ActionName, IE_Pressed );
+		FInputActionBinding NewBinding( ActionName, EventType.GetValue() );
 		NewBinding.bConsumeInput = bConsume;
 		NewBinding.ActionDelegate.GetDelegateForManualSet().BindUObject( this, &ThisClass::OnInputAction, Callback );
 
@@ -842,14 +842,14 @@ void UUserWidget::ListenForInputAction( FName ActionName, bool bConsume, FOnInpu
 	}
 }
 
-void UUserWidget::StopListeningForInputAction( FName ActionName )
+void UUserWidget::StopListeningForInputAction( FName ActionName, TEnumAsByte< EInputEvent > EventType )
 {
 	if ( InputComponent )
 	{
 		for ( int32 ExistingIndex = InputComponent->GetNumActionBindings() - 1; ExistingIndex >= 0; --ExistingIndex )
 		{
 			const FInputActionBinding& ExistingBind = InputComponent->GetActionBinding( ExistingIndex );
-			if ( ExistingBind.ActionName == ActionName )
+			if ( ExistingBind.ActionName == ActionName && ExistingBind.KeyEvent == EventType )
 			{
 				InputComponent->RemoveActionBinding( ExistingIndex );
 			}
diff --git a/Engine/Source/Runtime/UMG/Public/Blueprint/UserWidget.h b/Engine/Source/Runtime/UMG/Public/Blueprint/UserWidget.h
--- a/Engine/Source/Runtime/UMG/Public/Blueprint/UserWidget.h
+++ b/Engine/Source/Runtime/UMG/Public/Blueprint/UserWidget.h
@@ -932,10 +932,10 @@ protected:
 	void RemoveObsoleteBindings(const TArray<FName>& NamedSlots);
 	
 	UFUNCTION( BlueprintCallable, Category = "Input", meta = ( BlueprintProtected = "true" ) )
-	void ListenForInputAction( FName ActionName, bool bConsume, FOnInputAction Callback );
+	void ListenForInputAction( FName ActionName, TEnumAsByte< EInputEvent > EventType, bool bConsume, FOnInputAction Callback );
 
 	UFUNCTION( BlueprintCallable, Category = "Input", meta = ( BlueprintProtected = "true" ) )
-	void StopListeningForInputAction( FName ActionName );
+	void StopListeningForInputAction( FName ActionName, TEnumAsByte< EInputEvent > EventType );
 
 	UFUNCTION( BlueprintCallable, Category = "Input", meta = ( BlueprintProtected = "true" ) )
 	void StopListeningForAllInputActions();
