Fix unecessary error spam in userwidget shutdown

diff --git a/Engine/Source/Runtime/UMG/Private/UserWidget.cpp b/Engine/Source/Runtime/UMG/Private/UserWidget.cpp
--- a/Engine/Source/Runtime/UMG/Private/UserWidget.cpp
+++ b/Engine/Source/Runtime/UMG/Private/UserWidget.cpp
@@ -865,10 +865,6 @@ void UUserWidget::StopListeningForAllInputActions()
 		{
 			Controller->PopInputComponent( InputComponent );
 		}
-		else
-		{
-			FMessageLog( "PIE" ).Error( LOCTEXT( "GameNull", "Unable to listen to input actions without a player controller." ) );
-		}
 
 		InputComponent->ClearActionBindings();
 		InputComponent = nullptr;
