#slate
- virtual keyboard type selection for editable text box

--------
Integrated using branch UE4-Orion-To-UE4 of change#2701042 by [REDACTED_PERSON] on 2015/09/22 14:27:40.

diff --git a/Engine/Source/Runtime/Slate/Private/Widgets/Input/SEditableTextBox.cpp b/Engine/Source/Runtime/Slate/Private/Widgets/Input/SEditableTextBox.cpp
--- a/Engine/Source/Runtime/Slate/Private/Widgets/Input/SEditableTextBox.cpp
+++ b/Engine/Source/Runtime/Slate/Private/Widgets/Input/SEditableTextBox.cpp
@@ -57,6 +57,7 @@ void SEditableTextBox::Construct( const FArguments& InArgs )
 					.MinDesiredWidth( InArgs._MinDesiredWidth )
 					.SelectAllTextOnCommit( InArgs._SelectAllTextOnCommit )
 					.OnKeyDownHandler( InArgs._OnKeyDownHandler )
+					.VirtualKeyboardType( InArgs._VirtualKeyboardType )
 				]
 			]
 		]
diff --git a/Engine/Source/Runtime/Slate/Public/Widgets/Input/SEditableTextBox.h b/Engine/Source/Runtime/Slate/Public/Widgets/Input/SEditableTextBox.h
--- a/Engine/Source/Runtime/Slate/Public/Widgets/Input/SEditableTextBox.h
+++ b/Engine/Source/Runtime/Slate/Public/Widgets/Input/SEditableTextBox.h
@@ -93,6 +93,9 @@ public:
 		/** Provide a alternative mechanism for error reporting. */
 		SLATE_ARGUMENT( TSharedPtr<class IErrorReportingWidget>, ErrorReporting )
 
+		/** The type of virtual keyboard to use on mobile devices */
+		SLATE_ATTRIBUTE(EKeyboardType, VirtualKeyboardType)
+
 	SLATE_END_ARGS()
 	
 	/**
diff --git a/Engine/Source/Runtime/UMG/Private/Components/EditableTextBox.cpp b/Engine/Source/Runtime/UMG/Private/Components/EditableTextBox.cpp
--- a/Engine/Source/Runtime/UMG/Private/Components/EditableTextBox.cpp
+++ b/Engine/Source/Runtime/UMG/Private/Components/EditableTextBox.cpp
@@ -52,7 +52,7 @@ TSharedRef<SWidget> UEditableTextBox::RebuildWidget()
 		.SelectAllTextOnCommit(SelectAllTextOnCommit)
 		.OnTextChanged(BIND_UOBJECT_DELEGATE(FOnTextChanged, HandleOnTextChanged))
 		.OnTextCommitted(BIND_UOBJECT_DELEGATE(FOnTextCommitted, HandleOnTextCommitted))
-		;
+		.VirtualKeyboardType(EVirtualKeyboardType::AsKeyboardType(KeyboardType.GetValue()));
 
 	return MyEditableTextBlock.ToSharedRef();
 }
diff --git a/Engine/Source/Runtime/UMG/Public/Components/EditableTextBox.h b/Engine/Source/Runtime/UMG/Public/Components/EditableTextBox.h
--- a/Engine/Source/Runtime/UMG/Public/Components/EditableTextBox.h
+++ b/Engine/Source/Runtime/UMG/Public/Components/EditableTextBox.h
@@ -99,6 +99,10 @@ public:
 	UPROPERTY(EditAnywhere, BlueprintReadOnly, Category=Behavior, AdvancedDisplay)
 	bool SelectAllTextOnCommit;
 
+	/** If we're on a platform that requires a virtual keyboard, what kind of keyboard should this widget use? */
+	UPROPERTY(EditAnywhere, Category = Behavior, AdvancedDisplay)
+	TEnumAsByte< EVirtualKeyboardType::Type > KeyboardType;
+
 public:
 
 	/** Called whenever the text is changed interactively by the user */
