PR #1588: Add GridSlot SetLayer to BluePrintable API (Contributed by [REDACTED_PERSON])

Additionally, fixed a bug where the layer was not being set initially if configured at design time.

diff --git a/Engine/Source/Runtime/Slate/Public/Widgets/Layout/SGridPanel.h b/Engine/Source/Runtime/Slate/Public/Widgets/Layout/SGridPanel.h
--- a/Engine/Source/Runtime/Slate/Public/Widgets/Layout/SGridPanel.h
+++ b/Engine/Source/Runtime/Slate/Public/Widgets/Layout/SGridPanel.h
@@ -109,6 +109,19 @@ public:
 				return *this;
 			}
 
+			/** Positive values offset this cell to be hit-tested and drawn on top of others. Default is 0; i.e. no offset. */
+			FSlot& Layer(int32 Layer)
+			{
+				LayerParam = Layer;
+
+				if (Panel.IsValid())
+				{
+					Panel.Pin()->NotifySlotChanged(this);
+				}
+
+				return *this;
+			}
+
 			/** Offset this slot's content by some amount; positive values offset to lower right */
 			FSlot& Nudge( const FVector2D& Nudge )
 			{
diff --git a/Engine/Source/Runtime/UMG/Private/Components/GridSlot.cpp b/Engine/Source/Runtime/UMG/Private/Components/GridSlot.cpp
--- a/Engine/Source/Runtime/UMG/Private/Components/GridSlot.cpp
+++ b/Engine/Source/Runtime/UMG/Private/Components/GridSlot.cpp
@@ -26,7 +26,7 @@ void UGridSlot::ReleaseSlateResources(bool bReleaseChildren)
 // COMPILE PLEASE
 void UGridSlot::BuildSlot(TSharedRef<SGridPanel> GridPanel)
 {
-	Slot = &GridPanel->AddSlot(Column, Row)
+	Slot = &GridPanel->AddSlot(Column, Row, SGridPanel::Layer(Layer))
 		.HAlign(HorizontalAlignment)
 		.VAlign(VerticalAlignment)
 		.RowSpan(RowSpan)
@@ -73,6 +73,15 @@ void UGridSlot::SetColumnSpan(int32 InColumnSpan)
 	}
 }
 
+void UGridSlot::SetLayer(int32 InLayer)
+{
+	Layer = InLayer;
+	if (Slot)
+	{
+		Slot->Layer(InLayer);
+	}
+}
+
 void UGridSlot::SetHorizontalAlignment(EHorizontalAlignment InHorizontalAlignment)
 {
 	HorizontalAlignment = InHorizontalAlignment;
diff --git a/Engine/Source/Runtime/UMG/Public/Components/GridSlot.h b/Engine/Source/Runtime/UMG/Public/Components/GridSlot.h
--- a/Engine/Source/Runtime/UMG/Public/Components/GridSlot.h
+++ b/Engine/Source/Runtime/UMG/Public/Components/GridSlot.h
@@ -65,6 +65,10 @@ public:
 	UFUNCTION(BlueprintCallable, Category="Layout|Grid Slot")
 	void SetColumnSpan(int32 InColumnSpan);
 
+	/** Sets positive values offset this cell to be hit-tested and drawn on top of others. */
+	UFUNCTION(BlueprintCallable, Category = "Layout|Grid Slot")
+	void SetLayer(int32 InLayer);
+
 	/**  */
 	UFUNCTION(BlueprintCallable, Category="Layout|Grid Slot")
 	void SetHorizontalAlignment(EHorizontalAlignment InHorizontalAlignment);
