UMG - Spacer's Visibility now defaults to SelfHitTestInvisibile as was intended from the start.

diff --git a/Engine/Source/Runtime/UMG/Private/Components/Spacer.cpp b/Engine/Source/Runtime/UMG/Private/Components/Spacer.cpp
--- a/Engine/Source/Runtime/UMG/Private/Components/Spacer.cpp
+++ b/Engine/Source/Runtime/UMG/Private/Components/Spacer.cpp
@@ -12,6 +12,7 @@ USpacer::USpacer(const FObjectInitializer& ObjectInitializer)
 	, Size(1.0f, 1.0f)
 {
 	bIsVariable = false;
+	Visibility = ESlateVisibility::SelfHitTestInvisible;
 }
 
 void USpacer::ReleaseSlateResources(bool bReleaseChildren)
