UE-21493 - ScrollBox: ScrollWidgetIntoView no longer crashes if you pass it a null widget.

diff --git a/Engine/Source/Runtime/UMG/Private/Components/ScrollBox.cpp b/Engine/Source/Runtime/UMG/Private/Components/ScrollBox.cpp
--- a/Engine/Source/Runtime/UMG/Private/Components/ScrollBox.cpp
+++ b/Engine/Source/Runtime/UMG/Private/Components/ScrollBox.cpp
@@ -130,10 +130,13 @@ void UScrollBox::ScrollToEnd()
 
 void UScrollBox::ScrollWidgetIntoView(UWidget* WidgetToFind, bool AnimateScroll)
 {
-	TSharedPtr<SWidget> SlateWidgetToFind = WidgetToFind->GetCachedWidget();
-	if ( MyScrollBox.IsValid() )
+	if ( WidgetToFind )
 	{
-		MyScrollBox->ScrollDescendantIntoView(SlateWidgetToFind, AnimateScroll);
+		TSharedPtr<SWidget> SlateWidgetToFind = WidgetToFind->GetCachedWidget();
+		if ( MyScrollBox.IsValid() )
+		{
+			MyScrollBox->ScrollDescendantIntoView(SlateWidgetToFind, AnimateScroll);
+		}
 	}
 }
 
