UMG - Exposing the bCanCache option for the Invalidation panel to blueprints and the details panel to make it easier to debug/examine invalidation problems without removing it from the stack, you can just turn off invalidation.

diff --git a/Engine/Source/Runtime/UMG/Private/Components/InvalidationBox.cpp b/Engine/Source/Runtime/UMG/Private/Components/InvalidationBox.cpp
--- a/Engine/Source/Runtime/UMG/Private/Components/InvalidationBox.cpp
+++ b/Engine/Source/Runtime/UMG/Private/Components/InvalidationBox.cpp
@@ -12,6 +12,7 @@
 UInvalidationBox::UInvalidationBox(const FObjectInitializer& ObjectInitializer)
 	: Super(ObjectInitializer)
 {
+	bCanCache = true;
 	Visibility = ESlateVisibility::SelfHitTestInvisible;
 }
 
@@ -32,6 +33,10 @@ TSharedRef<SWidget> UInvalidationBox::RebuildWidget()
 	{
 		MyInvalidationPanel->SetCanCache(false);
 	}
+	else
+	{
+		MyInvalidationPanel->SetCanCache(bCanCache);
+	}
 
 	if ( GetChildrenCount() > 0 )
 	{
@@ -67,6 +72,25 @@ void UInvalidationBox::InvalidateCache()
 	}
 }
 
+bool UInvalidationBox::GetCanCache() const
+{
+	if ( MyInvalidationPanel.IsValid() )
+	{
+		return MyInvalidationPanel->GetCanCache();
+	}
+
+	return bCanCache;
+}
+
+void UInvalidationBox::SetCanCache(bool CanCache)
+{
+	bCanCache = CanCache;
+	if ( MyInvalidationPanel.IsValid() )
+	{
+		return MyInvalidationPanel->SetCanCache(bCanCache);
+	}
+}
+
 #if WITH_EDITOR
 
 const FSlateBrush* UInvalidationBox::GetEditorIcon()
diff --git a/Engine/Source/Runtime/UMG/Public/Components/InvalidationBox.h b/Engine/Source/Runtime/UMG/Public/Components/InvalidationBox.h
--- a/Engine/Source/Runtime/UMG/Public/Components/InvalidationBox.h
+++ b/Engine/Source/Runtime/UMG/Public/Components/InvalidationBox.h
@@ -17,18 +17,22 @@ class UMG_API UInvalidationBox : public UContentWidget
 public:
 
 	/**
-	 * Caches the locations for child draw elements relative to the invalidation box,
-	 * this adds extra overhead to drawing them every frame.  However, in cases where
-	 * the position of the invalidation boxes changes every frame this can be a big savings.
+	 * 
 	 */
-	UPROPERTY(EditAnywhere, BlueprintReadOnly, Category="Caching")
-	bool CacheRelativeTransforms;
+	UFUNCTION(BlueprintCallable, Category="Invalidation Box")
+	void InvalidateCache();
 
 	/**
 	 * 
 	 */
 	UFUNCTION(BlueprintCallable, Category="Invalidation Box")
-	void InvalidateCache();
+	bool GetCanCache() const;
+
+	/**
+	 * 
+	 */
+	UFUNCTION(BlueprintCallable, Category="Invalidation Box")
+	void SetCanCache(bool CanCache);
 
 	virtual void ReleaseSlateResources(bool bReleaseChildren) override;
 
@@ -48,6 +52,22 @@ protected:
 	virtual TSharedRef<SWidget> RebuildWidget() override;
 	// End of UWidget interface
 
+protected:
+	/**
+	 * Should the invalidation panel cache the widgets?  Making this false makes it so the invalidation
+	 * panel stops acting like an invalidation panel, just becomes a simple container widget.
+	 */
+	UPROPERTY(EditAnywhere, Category="Caching")
+	bool bCanCache;
+
+	/**
+	 * Caches the locations for child draw elements relative to the invalidation box,
+	 * this adds extra overhead to drawing them every frame.  However, in cases where
+	 * the position of the invalidation boxes changes every frame this can be a big savings.
+	 */
+	UPROPERTY(EditAnywhere, BlueprintReadOnly, Category="Caching")
+	bool CacheRelativeTransforms;
+
 protected:
 	TSharedPtr<class SInvalidationPanel> MyInvalidationPanel;
 };
