UMG - Breaking out several bits of the WidgetComponent into other classes that can be used elsewhere.  First, the bit that actually renders widgets has been brought out into the FWidgetRenderer, useful if you want to just draw some widgets to a render target it can do that for you.  Second, the surface that manages the screen position of widgets based on the projected location of the component in the world.  That code has been brought out into a publically usable widget that works for any Component.  It does a few special things for widget components for DrawSize/Pivot, but for non-widget components it uses the DrawSize and Pivot you provide on the surface.  Will probably wrap the SWorldWidgetScreenLayer in a UWidget in the future to allow users to get the same useful functionality of screen positioned widgets without needing to apply widget components to everything, which can be difficult when you're doing screen space widgets for several users in a split screen situation.

diff --git a/Engine/Source/Runtime/UMG/Private/Components/WidgetComponent.cpp b/Engine/Source/Runtime/UMG/Private/Components/WidgetComponent.cpp
--- a/Engine/Source/Runtime/UMG/Private/Components/WidgetComponent.cpp
+++ b/Engine/Source/Runtime/UMG/Private/Components/WidgetComponent.cpp
@@ -12,138 +12,20 @@
 #include "WidgetLayoutLibrary.h"
 #include "PhysicsEngine/BodySetup.h"
 #include "SGameLayerManager.h"
-
-extern SLATECORE_API int32 bFoldTick;
+#include "Slate/WidgetRenderer.h"
+#include "Slate/SWorldWidgetScreenLayer.h"
 
 DECLARE_CYCLE_STAT(TEXT("3DHitTesting"), STAT_Slate3DHitTesting, STATGROUP_Slate);
 
 static const FName SharedLayerName(TEXT("WidgetComponentScreenLayer"));
 
-class SWorldWidgetScreenLayer : public SCompoundWidget
+class FWorldWidgetScreenLayer : public IGameLayer
 {
-	SLATE_BEGIN_ARGS(SWorldWidgetScreenLayer)
-	{
-		_Visibility = EVisibility::SelfHitTestInvisible;
-	}
-	SLATE_END_ARGS()
-
 public:
-	void Construct(const FArguments& InArgs, const TArray<TWeakObjectPtr<UWidgetComponent>>& InComponents)
+	FWorldWidgetScreenLayer(APlayerController* InOwningPlayer)
 	{
-		Components = InComponents;
-
-		ChildSlot
-		[
-			SAssignNew(Canvas, SConstraintCanvas)
-		];
 	}
 
-	virtual void Tick(const FGeometry& AllottedGeometry, const double InCurrentTime, const float InDeltaTime) override
-	{
-		TArray<UWidgetComponent*, TInlineAllocator<1>> DeadComponents;
-
-		for ( TWeakObjectPtr<UWidgetComponent> Component : Components )
-		{
-			if ( UWidgetComponent* WidgetComponent = Component.Get() )
-			{
-				if ( ULocalPlayer* LocalPlayer = WidgetComponent->GetOwnerPlayer() )
-				{
-					if ( APlayerController* PlayerController = LocalPlayer->PlayerController )
-					{
-						FVector WorldLocation = WidgetComponent->GetComponentLocation();
-
-						FVector ScreenPosition;
-						const bool bProjected = UWidgetLayoutLibrary::ProjectWorldLocationToWidgetPositionWithDistance(PlayerController, WorldLocation, ScreenPosition);
-
-						if ( bProjected )
-						{
-							if (UUserWidget* Widget = WidgetComponent->GetUserWidgetObject())
-							{
-								Widget->SetVisibility(ESlateVisibility::SelfHitTestInvisible);
-							}
-							
-							if ( SConstraintCanvas::FSlot* CanvasSlot = ComponentToSlot.FindRef(WidgetComponent) )
-							{
-								FVector2D DrawSize = WidgetComponent->GetDrawSize();
-								FVector2D Pivot = WidgetComponent->GetPivot();
-
-								CanvasSlot->AutoSize(DrawSize.IsZero());
-								CanvasSlot->Offset(FMargin(ScreenPosition.X, ScreenPosition.Y, DrawSize.X, DrawSize.Y));
-								CanvasSlot->Anchors(FAnchors(0, 0, 0, 0));
-								CanvasSlot->Alignment(Pivot);
-								CanvasSlot->ZOrder(-ScreenPosition.Z);
-							}
-						}
-						else
-						{
-							if (UUserWidget* Widget = WidgetComponent->GetUserWidgetObject())
-							{
-								Widget->SetVisibility(ESlateVisibility::Hidden);
-							}
-						}
-					}
-				}
-			}
-			else
-			{
-				DeadComponents.Add(WidgetComponent);
-			}
-		}
-
-		// Normally components should be removed by someone calling remove component, but just in case it was 
-		// deleted in a way where they didn't happen, this is our backup solution to enure we remove stale widgets.
-		for ( int32 Index = 0; Index < DeadComponents.Num(); Index++ )
-		{
-			RemoveComponent(DeadComponents[Index]);
-		}
-	}
-
-	void AddComponent(UWidgetComponent* Component)
-	{
-		Components.AddUnique(Component);
-
-		if (UUserWidget* Widget = Component->GetUserWidgetObject())
-		{
-			SConstraintCanvas::FSlot* CanvasSlot;
-
-			Canvas->AddSlot()
-			.Expose(CanvasSlot)
-			[
-				Widget->TakeWidget()
-			];
-
-			ComponentToSlot.Add(Component, CanvasSlot);
-		}
-	}
-
-	void RemoveComponent(UWidgetComponent* Component)
-	{
-		if (Component)
-		{
-			Components.RemoveSwap(Component);
-
-			if (UUserWidget* Widget = Component->GetUserWidgetObject())
-			{
-				TSharedPtr<SWidget> CachedWidget = Widget->GetCachedWidget();
-				if (CachedWidget.IsValid())
-				{
-					Canvas->RemoveSlot(CachedWidget.ToSharedRef());
-				}
-
-				ComponentToSlot.Remove(Component);
-			}
-		}
-	}
-
-private:
-	TArray<TWeakObjectPtr<UWidgetComponent>> Components;
-	TMap<UWidgetComponent*, SConstraintCanvas::FSlot*> ComponentToSlot;
-	TSharedPtr<SConstraintCanvas> Canvas;
-};
-
-class FWorldWidgetScreenLayer : public IGameLayer
-{
-public:
 	virtual ~FWorldWidgetScreenLayer()
 	{
 		// empty virtual destructor to help clang warning
@@ -154,7 +36,10 @@ public:
 		Components.AddUnique(Component);
 		if ( ScreenLayer.IsValid() )
 		{
-			ScreenLayer.Pin()->AddComponent(Component);
+			if ( UUserWidget* UserWidget = Component->GetUserWidgetObject() )
+			{
+				ScreenLayer.Pin()->AddComponent(Component, UserWidget->TakeWidget());
+			}
 		}
 	}
 
@@ -175,38 +60,32 @@ public:
 			return ScreenLayer.Pin().ToSharedRef();
 		}
 
-		TSharedRef<SWorldWidgetScreenLayer> NewScreenLayer = SNew(SWorldWidgetScreenLayer, Components);
+		TSharedRef<SWorldWidgetScreenLayer> NewScreenLayer = SNew(SWorldWidgetScreenLayer);
+		NewScreenLayer->SetOwningPlayer(OwningPlayer.Get());
 		ScreenLayer = NewScreenLayer;
 
+		// Add all the pending user widgets to the surface
+		for ( TWeakObjectPtr<UWidgetComponent>& WeakComponent : Components )
+		{
+			if ( UWidgetComponent* Component = WeakComponent.Get() )
+			{
+				if ( UUserWidget* UserWidget = Component->GetUserWidgetObject() )
+				{
+					NewScreenLayer->AddComponent(Component, UserWidget->TakeWidget());
+				}
+			}
+		}
+
 		return NewScreenLayer;
 	}
 
 private:
+	TWeakObjectPtr<APlayerController> OwningPlayer;
 	TWeakPtr<SWorldWidgetScreenLayer> ScreenLayer;
 	TArray<TWeakObjectPtr<UWidgetComponent>> Components;
 };
 
 
-class SVirtualWindow : public SWindow
-{
-	SLATE_BEGIN_ARGS( SVirtualWindow )
-		: _Size( FVector2D(100,100) )
-	{}
-
-		SLATE_ARGUMENT( FVector2D, Size )
-	SLATE_END_ARGS()
-
-public:
-	void Construct( const FArguments& InArgs )
-	{
-		bIsPopupWindow = true;
-		SetCachedSize( InArgs._Size );
-		SetNativeWindow( MakeShareable( new FGenericWindow() ) );
-		SetContent( SNullWidget::NullWidget );
-	}
-};
-
-
 /** Represents a billboard sprite to the scene manager. */
 class FWidget3DSceneProxy : public FPrimitiveSceneProxy
 {
@@ -557,9 +436,9 @@ void UWidgetComponent::EndPlay(const EEndPlayReason::Type EndPlayReason)
 
 FPrimitiveSceneProxy* UWidgetComponent::CreateSceneProxy()
 {
-	if ( Space != EWidgetSpace::Screen && Renderer.IsValid() )
+	if ( Space != EWidgetSpace::Screen && WidgetRenderer.IsValid() )
 	{
-		return new FWidget3DSceneProxy(this, *Renderer);
+		return new FWidget3DSceneProxy(this, *WidgetRenderer->GetSlateRenderer());
 	}
 	
 	return nullptr;
@@ -671,9 +550,9 @@ void UWidgetComponent::OnRegister()
 
 		if ( Space != EWidgetSpace::Screen )
 		{
-			if ( !Renderer.IsValid() && !GUsingNullRHI )
+			if ( !WidgetRenderer.IsValid() && !GUsingNullRHI )
 			{
-				Renderer = FModuleManager::Get().LoadModuleChecked<ISlateRHIRendererModule>("SlateRHIRenderer").CreateSlate3DRenderer();
+				WidgetRenderer = MakeShareable(new FWidgetRenderer());
 			}
 		}
 	}
@@ -728,7 +607,7 @@ void UWidgetComponent::ReleaseResources()
 		Widget = nullptr;
 	}
 
-	Renderer.Reset();
+	WidgetRenderer.Reset();
 	SlateWidget.Reset();
 	HitTestGrid.Reset();
 }
@@ -769,49 +648,14 @@ void UWidgetComponent::TickComponent(float DeltaTime, enum ELevelTick TickType,
 			// If we don't tick when off-screen, don't bother ticking if it hasn't been rendered recently
 			if ( TickWhenOffscreen || GetWorld()->TimeSince(LastRenderTime) <= RenderTimeThreshold )
 			{
-				SlateWidget->SlatePrepass();
-
-				FGeometry WindowGeometry = FGeometry::MakeRoot(DrawSize, FSlateLayoutTransform());
-
-				if ( !bFoldTick )
-				{
-					SlateWidget->TickWidgetsRecursively(WindowGeometry, FApp::GetCurrentTime(), DeltaTime);
-				}
-
-				// Ticking can cause geometry changes.  Recompute
-				SlateWidget->SlatePrepass();
-
-				// Get the free buffer & add our virtual window
-				FSlateDrawBuffer& DrawBuffer = Renderer->GetDrawBuffer();
-				FSlateWindowElementList& WindowElementList = DrawBuffer.AddWindowElementList(SlateWidget.ToSharedRef());
-
-				int32 MaxLayerId = 0;
-				{
-					// Prepare the test grid 
-					HitTestGrid->ClearGridForNewFrame(WindowGeometry.GetClippingRect());
-
-					// Paint the window
-					MaxLayerId = SlateWidget->Paint(
-						FPaintArgs(*SlateWidget.Get(), *HitTestGrid, FVector2D::ZeroVector, FApp::GetCurrentTime(), DeltaTime),
-						WindowGeometry, WindowGeometry.GetClippingRect(),
-						WindowElementList,
-						0,
-						FWidgetStyle(),
-						SlateWidget->IsEnabled());
-				}
-
-				Renderer->DrawWindow_GameThread(DrawBuffer);
-
 				UpdateRenderTarget();
 
-				// Enqueue a command to unlock the draw buffer after all windows have been drawn
-				ENQUEUE_UNIQUE_RENDER_COMMAND_THREEPARAMETER(UWidgetComponentRenderToTexture,
-					FSlateDrawBuffer&, InDrawBuffer, DrawBuffer,
-					UTextureRenderTarget2D*, InRenderTarget, RenderTarget,
-					ISlate3DRenderer*, InRenderer, Renderer.Get(),
-					{
-						InRenderer->DrawWindowToTarget_RenderThread(RHICmdList, InRenderTarget, InDrawBuffer);
-					});
+				WidgetRenderer->DrawWindow(
+					RenderTarget,
+					HitTestGrid.ToSharedRef(),
+					SlateWidget.ToSharedRef(),
+					DrawSize,
+					DeltaTime);
 			}
 		}
 	}
@@ -840,7 +684,7 @@ void UWidgetComponent::TickComponent(float DeltaTime, enum ELevelTick TickType,
 								TSharedPtr<IGameLayer> Layer = LayerManager->FindLayerForPlayer(TargetPlayer, SharedLayerName);
 								if ( !Layer.IsValid() )
 								{
-									TSharedRef<FWorldWidgetScreenLayer> NewScreenLayer = MakeShareable(new FWorldWidgetScreenLayer());
+									TSharedRef<FWorldWidgetScreenLayer> NewScreenLayer = MakeShareable(new FWorldWidgetScreenLayer(PlayerController));
 									LayerManager->AddLayerForPlayer(TargetPlayer, SharedLayerName, NewScreenLayer, LayerZOrder);
 									ScreenLayer = NewScreenLayer;
 								}
diff --git a/Engine/Source/Runtime/UMG/Private/Slate/SWorldWidgetScreenLayer.cpp b/Engine/Source/Runtime/UMG/Private/Slate/SWorldWidgetScreenLayer.cpp
new file mode 100644
--- /dev/null
+++ b/Engine/Source/Runtime/UMG/Private/Slate/SWorldWidgetScreenLayer.cpp
@@ -0,0 +1,129 @@
+// Copyright 1998-2015 Epic Games, Inc. All Rights Reserved.
+
+#include "UMGPrivatePCH.h"
+
+#include "SWorldWidgetScreenLayer.h"
+
+void SWorldWidgetScreenLayer::Construct(const FArguments& InArgs)
+{
+	bCanSupportFocus = false;
+	DrawSize = FVector2D(0, 0);
+	Pivot = FVector2D(0.5f, 0.5f);
+
+	ChildSlot
+	[
+		SAssignNew(Canvas, SConstraintCanvas)
+	];
+}
+
+void SWorldWidgetScreenLayer::SetOwningPlayer(APlayerController* Controller)
+{
+	PlayerController = Controller;
+}
+
+void SWorldWidgetScreenLayer::SetWidgetDrawSize(FVector2D InDrawSize)
+{
+	DrawSize = InDrawSize;
+}
+
+void SWorldWidgetScreenLayer::SetWidgetPivot(FVector2D InPivot)
+{
+	Pivot = InPivot;
+}
+
+void SWorldWidgetScreenLayer::AddComponent(USceneComponent* Component, TSharedPtr<SWidget> Widget)
+{
+	if ( Widget.IsValid() )
+	{
+		FComponentEntry& Entry = ComponentMap.FindOrAdd(Component);
+		Entry.Component = Component;
+		Entry.WidgetComponent = Cast<UWidgetComponent>(Component);
+		Entry.Widget = Widget;
+
+		Canvas->AddSlot()
+		.Expose(Entry.Slot)
+		[
+			SAssignNew(Entry.ContainerWidget, SBox)
+			[
+				Widget.ToSharedRef()
+			]
+		];
+	}
+}
+
+void SWorldWidgetScreenLayer::RemoveComponent(USceneComponent* Component)
+{
+	if ( FComponentEntry* Entry = ComponentMap.Find(Component) )
+	{
+		if ( Entry->Widget.IsValid() )
+		{
+			Canvas->RemoveSlot(Entry->ContainerWidget.ToSharedRef());
+		}
+
+		ComponentMap.Remove(Component);
+	}
+}
+
+void SWorldWidgetScreenLayer::Tick(const FGeometry& AllottedGeometry, const double InCurrentTime, const float InDeltaTime)
+{
+	TArray<USceneComponent*, TInlineAllocator<1>> DeadComponents;
+
+	if ( APlayerController* Player = PlayerController.Get() )
+	{
+		for ( auto It = ComponentMap.CreateIterator(); It; ++It )
+		{
+			FComponentEntry& Entry = It.Value();
+
+			if ( USceneComponent* SceneComponent = Entry.Component.Get() )
+			{
+				FVector WorldLocation = SceneComponent->GetComponentLocation();
+
+				FVector ScreenPosition;
+				const bool bProjected = UWidgetLayoutLibrary::ProjectWorldLocationToWidgetPositionWithDistance(Player, WorldLocation, ScreenPosition);
+
+				if ( bProjected )
+				{
+					Entry.ContainerWidget->SetVisibility(EVisibility::SelfHitTestInvisible);
+
+					if ( SConstraintCanvas::FSlot* CanvasSlot = Entry.Slot )
+					{
+						if ( Entry.WidgetComponent )
+						{
+							FVector2D ComponentDrawSize = Entry.WidgetComponent->GetDrawSize();
+							FVector2D ComponentPivot = Entry.WidgetComponent->GetPivot();
+
+							CanvasSlot->AutoSize(ComponentDrawSize.IsZero());
+							CanvasSlot->Offset(FMargin(ScreenPosition.X, ScreenPosition.Y, ComponentDrawSize.X, ComponentDrawSize.Y));
+							CanvasSlot->Anchors(FAnchors(0, 0, 0, 0));
+							CanvasSlot->Alignment(ComponentPivot);
+							CanvasSlot->ZOrder(-ScreenPosition.Z);
+						}
+						else
+						{
+							CanvasSlot->AutoSize(DrawSize.IsZero());
+							CanvasSlot->Offset(FMargin(ScreenPosition.X, ScreenPosition.Y, DrawSize.X, DrawSize.Y));
+							CanvasSlot->Anchors(FAnchors(0, 0, 0, 0));
+							CanvasSlot->Alignment(Pivot);
+							CanvasSlot->ZOrder(-ScreenPosition.Z);
+						}
+					}
+				}
+				else
+				{
+					Entry.ContainerWidget->SetVisibility(EVisibility::Collapsed);
+				}
+			}
+			else
+			{
+				DeadComponents.Add(SceneComponent);
+			}
+		}
+	}
+
+	// Normally components should be removed by someone calling remove component, but just in case it was 
+	// deleted in a way where they didn't happen, this is our backup solution to enure we remove stale widgets.
+	for ( int32 Index = 0; Index < DeadComponents.Num(); Index++ )
+	{
+		RemoveComponent(DeadComponents[Index]);
+	}
+}
diff --git a/Engine/Source/Runtime/UMG/Private/Slate/WidgetRenderer.cpp b/Engine/Source/Runtime/UMG/Private/Slate/WidgetRenderer.cpp
new file mode 100644
--- /dev/null
+++ b/Engine/Source/Runtime/UMG/Private/Slate/WidgetRenderer.cpp
@@ -0,0 +1,93 @@
+// Copyright 1998-2015 Epic Games, Inc. All Rights Reserved.
+
+#include "UMGPrivatePCH.h"
+#include "WidgetRenderer.h"
+
+extern SLATECORE_API int32 bFoldTick;
+
+void SVirtualWindow::Construct(const FArguments& InArgs)
+{
+	bIsPopupWindow = true;
+	SetCachedSize(InArgs._Size);
+	SetNativeWindow(MakeShareable(new FGenericWindow()));
+	SetContent(SNullWidget::NullWidget);
+}
+
+FWidgetRenderer::FWidgetRenderer()
+{
+	Renderer = FModuleManager::Get().LoadModuleChecked<ISlateRHIRendererModule>("SlateRHIRenderer").CreateSlate3DRenderer();
+}
+
+FWidgetRenderer::~FWidgetRenderer()
+{
+}
+
+ISlate3DRenderer* FWidgetRenderer::GetSlateRenderer()
+{
+	return Renderer.Get();
+}
+
+UTextureRenderTarget2D* FWidgetRenderer::DrawWidget(TSharedRef<SWidget>& Widget, FVector2D DrawSize)
+{
+	UTextureRenderTarget2D* RenderTarget = NewObject<UTextureRenderTarget2D>();
+	RenderTarget->ClearColor = FLinearColor::Transparent;
+	RenderTarget->InitCustomFormat(DrawSize.X, DrawSize.Y, PF_B8G8R8A8, false);
+
+	DrawWidget(RenderTarget, Widget, DrawSize, 0);
+
+	return RenderTarget;
+}
+
+void FWidgetRenderer::DrawWidget(UTextureRenderTarget2D* RenderTarget, TSharedRef<SWidget>& Widget, FVector2D DrawSize, float DeltaTime)
+{
+	TSharedRef<SVirtualWindow> Window = SNew(SVirtualWindow).Size(DrawSize);
+	TSharedRef<FHittestGrid> HitTestGrid = MakeShareable(new FHittestGrid);
+
+	Window->SetContent(Widget);
+	Window->Resize(DrawSize);
+
+	DrawWindow(RenderTarget, HitTestGrid, Window, DrawSize, DeltaTime);
+}
+
+void FWidgetRenderer::DrawWindow(UTextureRenderTarget2D* RenderTarget, TSharedRef<FHittestGrid> HitTestGrid, TSharedRef<SWindow> Window, FVector2D DrawSize, float DeltaTime)
+{
+	FGeometry WindowGeometry = FGeometry::MakeRoot(DrawSize, FSlateLayoutTransform());
+
+	if ( !bFoldTick )
+	{
+		Window->TickWidgetsRecursively(WindowGeometry, FApp::GetCurrentTime(), DeltaTime);
+	}
+
+	// Ticking can cause geometry changes.  Recompute
+	Window->SlatePrepass();
+
+	// Prepare the test grid 
+	HitTestGrid->ClearGridForNewFrame(WindowGeometry.GetClippingRect());
+
+	// Get the free buffer & add our virtual window
+	FSlateDrawBuffer& DrawBuffer = Renderer->GetDrawBuffer();
+	FSlateWindowElementList& WindowElementList = DrawBuffer.AddWindowElementList(Window);
+
+	int32 MaxLayerId = 0;
+	{
+		// Paint the window
+		MaxLayerId = Window->Paint(
+			FPaintArgs(Window.Get(), HitTestGrid.Get(), FVector2D::ZeroVector, FApp::GetCurrentTime(), DeltaTime),
+			WindowGeometry, WindowGeometry.GetClippingRect(),
+			WindowElementList,
+			0,
+			FWidgetStyle(),
+			Window->IsEnabled());
+	}
+
+	Renderer->DrawWindow_GameThread(DrawBuffer);
+
+	// Enqueue a command to unlock the draw buffer after all windows have been drawn
+	ENQUEUE_UNIQUE_RENDER_COMMAND_THREEPARAMETER(UWidgetComponentRenderToTexture,
+		FSlateDrawBuffer&, InDrawBuffer, DrawBuffer,
+		UTextureRenderTarget2D*, InRenderTarget, RenderTarget,
+		ISlate3DRenderer*, InRenderer, Renderer.Get(),
+		{
+			InRenderer->DrawWindowToTarget_RenderThread(RHICmdList, InRenderTarget, InDrawBuffer);
+		});
+}
diff --git a/Engine/Source/Runtime/UMG/Public/Components/WidgetComponent.h b/Engine/Source/Runtime/UMG/Public/Components/WidgetComponent.h
--- a/Engine/Source/Runtime/UMG/Public/Components/WidgetComponent.h
+++ b/Engine/Source/Runtime/UMG/Public/Components/WidgetComponent.h
@@ -262,9 +262,6 @@ protected:
 
 	/** The grid used to find actual hit actual widgets once input has been translated to the components local space */
 	TSharedPtr<class FHittestGrid> HitTestGrid;
-
-	/** The slate 3D renderer used to render the user slate widget */
-	TSharedPtr<class ISlate3DRenderer> Renderer;
 	
 	/** The slate window that contains the user widget content */
 	TSharedPtr<class SVirtualWindow> SlateWidget;
@@ -274,4 +271,7 @@ protected:
 
 	/** The hit tester to use for this component */
 	static TSharedPtr<class FWidget3DHitTester> WidgetHitTester;
+
+	/** Helper class for drawing widgets to a render target. */
+	TSharedPtr<class FWidgetRenderer> WidgetRenderer;
 };
diff --git a/Engine/Source/Runtime/UMG/Public/Slate/SWorldWidgetScreenLayer.h b/Engine/Source/Runtime/UMG/Public/Slate/SWorldWidgetScreenLayer.h
new file mode 100644
--- /dev/null
+++ b/Engine/Source/Runtime/UMG/Public/Slate/SWorldWidgetScreenLayer.h
@@ -0,0 +1,54 @@
+// Copyright 1998-2015 Epic Games, Inc. All Rights Reserved.
+
+#pragma once
+
+class SWorldWidgetScreenLayer : public SCompoundWidget
+{
+	SLATE_BEGIN_ARGS(SWorldWidgetScreenLayer)
+	{
+		_Visibility = EVisibility::SelfHitTestInvisible;
+	}
+	SLATE_END_ARGS()
+
+public:
+	void Construct(const FArguments& InArgs);
+
+	virtual void Tick(const FGeometry& AllottedGeometry, const double InCurrentTime, const float InDeltaTime) override;
+
+	void SetWidgetDrawSize(FVector2D DrawSize);
+
+	void SetWidgetPivot(FVector2D Pivot);
+
+	void SetOwningPlayer(APlayerController* Controller);
+
+	void AddComponent(USceneComponent* Component, TSharedPtr<SWidget> Widget);
+
+	void RemoveComponent(USceneComponent* Component);
+
+private:
+	TWeakObjectPtr<APlayerController> PlayerController;
+
+	FVector2D DrawSize;
+	FVector2D Pivot;
+
+	class FComponentEntry
+	{
+	public:
+		FComponentEntry()
+			: Slot(nullptr)
+		{
+		}
+
+	public:
+
+		TWeakObjectPtr<USceneComponent> Component;
+		UWidgetComponent* WidgetComponent;
+
+		TSharedPtr<SWidget> ContainerWidget;
+		TSharedPtr<SWidget> Widget;
+		SConstraintCanvas::FSlot* Slot;
+	};
+
+	TMap<USceneComponent*, FComponentEntry> ComponentMap;
+	TSharedPtr<SConstraintCanvas> Canvas;
+};
\ No newline at end of file
diff --git a/Engine/Source/Runtime/UMG/Public/Slate/WidgetRenderer.h b/Engine/Source/Runtime/UMG/Public/Slate/WidgetRenderer.h
new file mode 100644
--- /dev/null
+++ b/Engine/Source/Runtime/UMG/Public/Slate/WidgetRenderer.h
@@ -0,0 +1,37 @@
+// Copyright 1998-2015 Epic Games, Inc. All Rights Reserved.
+
+#pragma once
+
+class SVirtualWindow : public SWindow
+{
+	SLATE_BEGIN_ARGS(SVirtualWindow)
+		: _Size(FVector2D(100, 100))
+	{}
+
+	SLATE_ARGUMENT(FVector2D, Size)
+	
+	SLATE_END_ARGS()
+
+public:
+	void Construct(const FArguments& InArgs);
+};
+
+/**
+ * 
+ */
+class FWidgetRenderer
+{
+public:
+	FWidgetRenderer();
+	~FWidgetRenderer();
+
+	ISlate3DRenderer* GetSlateRenderer();
+
+	UTextureRenderTarget2D* DrawWidget(TSharedRef<SWidget>& Widget, FVector2D DrawSize);
+	void DrawWidget(UTextureRenderTarget2D* RenderTarget, TSharedRef<SWidget>& Widget, FVector2D DrawSize, float DeltaTime);
+	void DrawWindow(UTextureRenderTarget2D* RenderTarget, TSharedRef<class FHittestGrid> HitTestGrid, TSharedRef<class SWindow> Window, FVector2D DrawSize, float DeltaTime);
+
+private:
+	/** The slate 3D renderer used to render the user slate widget */
+	TSharedPtr<class ISlate3DRenderer> Renderer;
+};
