UMG/Slate - Adding some additional checks to the SWorldWidgetScreenLayer to ignore null components add/remove.

diff --git a/Engine/Source/Runtime/UMG/Private/Slate/SWorldWidgetScreenLayer.cpp b/Engine/Source/Runtime/UMG/Private/Slate/SWorldWidgetScreenLayer.cpp
--- a/Engine/Source/Runtime/UMG/Private/Slate/SWorldWidgetScreenLayer.cpp
+++ b/Engine/Source/Runtime/UMG/Private/Slate/SWorldWidgetScreenLayer.cpp
@@ -33,7 +33,7 @@ void SWorldWidgetScreenLayer::SetWidgetPivot(FVector2D InPivot)
 
 void SWorldWidgetScreenLayer::AddComponent(USceneComponent* Component, TSharedPtr<SWidget> Widget)
 {
-	if ( Widget.IsValid() )
+	if ( Component && Widget.IsValid() )
 	{
 		FComponentEntry& Entry = ComponentMap.FindOrAdd(Component);
 		Entry.Component = Component;
@@ -53,14 +53,17 @@ void SWorldWidgetScreenLayer::AddComponent(USceneComponent* Component, TSharedPt
 
 void SWorldWidgetScreenLayer::RemoveComponent(USceneComponent* Component)
 {
-	if ( FComponentEntry* Entry = ComponentMap.Find(Component) )
+	if ( Component )
 	{
-		if ( Entry->Widget.IsValid() )
+		if ( FComponentEntry* Entry = ComponentMap.Find(Component) )
 		{
-			Canvas->RemoveSlot(Entry->ContainerWidget.ToSharedRef());
-		}
+			if ( Entry->Widget.IsValid() )
+			{
+				Canvas->RemoveSlot(Entry->ContainerWidget.ToSharedRef());
+			}
 
-		ComponentMap.Remove(Component);
+			ComponentMap.Remove(Component);
+		}
 	}
 }
 
