UMG - Fixing some weird merge problems with widget component.  Should now reappear again for screen space widgets.

diff --git a/Engine/Source/Runtime/UMG/Private/Components/WidgetComponent.cpp b/Engine/Source/Runtime/UMG/Private/Components/WidgetComponent.cpp
--- a/Engine/Source/Runtime/UMG/Private/Components/WidgetComponent.cpp
+++ b/Engine/Source/Runtime/UMG/Private/Components/WidgetComponent.cpp
@@ -24,32 +24,39 @@ class FWorldWidgetScreenLayer : public IGameLayer
 public:
 	FWorldWidgetScreenLayer(APlayerController* InOwningPlayer)
 	{
+		OwningPlayer = InOwningPlayer;
 	}
 
 	virtual ~FWorldWidgetScreenLayer()
 	{
 		// empty virtual destructor to help clang warning
 	}
-	
+
 	void AddComponent(UWidgetComponent* Component)
 	{
-		Components.AddUnique(Component);
-		if ( ScreenLayer.IsValid() )
+		if ( Component )
 		{
-			if ( UUserWidget* UserWidget = Component->GetUserWidgetObject() )
+			Components.AddUnique(Component);
+			if ( ScreenLayer.IsValid() )
 			{
-				ScreenLayer.Pin()->AddComponent(Component, UserWidget->TakeWidget());
+				if ( UUserWidget* UserWidget = Component->GetUserWidgetObject() )
+				{
+					ScreenLayer.Pin()->AddComponent(Component, UserWidget->TakeWidget());
+				}
 			}
 		}
 	}
 
 	void RemoveComponent(UWidgetComponent* Component)
 	{
-		Components.RemoveSwap(Component);
-
-		if ( ScreenLayer.IsValid() )
+		if ( Component )
 		{
-			ScreenLayer.Pin()->RemoveComponent(Component);
+			Components.RemoveSwap(Component);
+
+			if ( ScreenLayer.IsValid() )
+			{
+				ScreenLayer.Pin()->RemoveComponent(Component);
+			}
 		}
 	}
 
