UMG - Further refinement of the player reference passed to the SWorldWidgetScreenLayer.  Now using FLocalPlayerContext to be more robust.

diff --git a/Engine/Source/Runtime/UMG/Private/Components/WidgetComponent.cpp b/Engine/Source/Runtime/UMG/Private/Components/WidgetComponent.cpp
--- a/Engine/Source/Runtime/UMG/Private/Components/WidgetComponent.cpp
+++ b/Engine/Source/Runtime/UMG/Private/Components/WidgetComponent.cpp
@@ -22,9 +22,9 @@ static const FName SharedLayerName(TEXT("WidgetComponentScreenLayer"));
 class FWorldWidgetScreenLayer : public IGameLayer
 {
 public:
-	FWorldWidgetScreenLayer(APlayerController* InOwningPlayer)
+	FWorldWidgetScreenLayer(const FLocalPlayerContext& PlayerContext)
 	{
-		OwningPlayer = InOwningPlayer;
+		OwningPlayer = PlayerContext;
 	}
 
 	virtual ~FWorldWidgetScreenLayer()
@@ -67,8 +67,7 @@ public:
 			return ScreenLayer.Pin().ToSharedRef();
 		}
 
-		TSharedRef<SWorldWidgetScreenLayer> NewScreenLayer = SNew(SWorldWidgetScreenLayer);
-		NewScreenLayer->SetOwningPlayer(OwningPlayer.Get());
+		TSharedRef<SWorldWidgetScreenLayer> NewScreenLayer = SNew(SWorldWidgetScreenLayer, OwningPlayer);
 		ScreenLayer = NewScreenLayer;
 
 		// Add all the pending user widgets to the surface
@@ -87,7 +86,7 @@ public:
 	}
 
 private:
-	TWeakObjectPtr<APlayerController> OwningPlayer;
+	FLocalPlayerContext OwningPlayer;
 	TWeakPtr<SWorldWidgetScreenLayer> ScreenLayer;
 	TArray<TWeakObjectPtr<UWidgetComponent>> Components;
 };
@@ -688,10 +687,12 @@ void UWidgetComponent::TickComponent(float DeltaTime, enum ELevelTick TickType,
 							{
 								TSharedPtr<FWorldWidgetScreenLayer> ScreenLayer;
 
+								FLocalPlayerContext PlayerContext(TargetPlayer, GetWorld());
+
 								TSharedPtr<IGameLayer> Layer = LayerManager->FindLayerForPlayer(TargetPlayer, SharedLayerName);
 								if ( !Layer.IsValid() )
 								{
-									TSharedRef<FWorldWidgetScreenLayer> NewScreenLayer = MakeShareable(new FWorldWidgetScreenLayer(PlayerController));
+									TSharedRef<FWorldWidgetScreenLayer> NewScreenLayer = MakeShareable(new FWorldWidgetScreenLayer(PlayerContext));
 									LayerManager->AddLayerForPlayer(TargetPlayer, SharedLayerName, NewScreenLayer, LayerZOrder);
 									ScreenLayer = NewScreenLayer;
 								}
@@ -702,7 +703,7 @@ void UWidgetComponent::TickComponent(float DeltaTime, enum ELevelTick TickType,
 								
 								bAddedToScreen = true;
 								
-								Widget->SetPlayerContext(FLocalPlayerContext(TargetPlayer, GetWorld()));
+								Widget->SetPlayerContext(PlayerContext);
 								ScreenLayer->AddComponent(this);
 							}
 						}
diff --git a/Engine/Source/Runtime/UMG/Private/Slate/SWorldWidgetScreenLayer.cpp b/Engine/Source/Runtime/UMG/Private/Slate/SWorldWidgetScreenLayer.cpp
--- a/Engine/Source/Runtime/UMG/Private/Slate/SWorldWidgetScreenLayer.cpp
+++ b/Engine/Source/Runtime/UMG/Private/Slate/SWorldWidgetScreenLayer.cpp
@@ -4,8 +4,10 @@
 
 #include "SWorldWidgetScreenLayer.h"
 
-void SWorldWidgetScreenLayer::Construct(const FArguments& InArgs)
+void SWorldWidgetScreenLayer::Construct(const FArguments& InArgs, const FLocalPlayerContext& InPlayerContext)
 {
+	PlayerContext = InPlayerContext;
+
 	bCanSupportFocus = false;
 	DrawSize = FVector2D(0, 0);
 	Pivot = FVector2D(0.5f, 0.5f);
@@ -16,11 +18,6 @@ void SWorldWidgetScreenLayer::Construct(const FArguments& InArgs)
 	];
 }
 
-void SWorldWidgetScreenLayer::SetOwningPlayer(APlayerController* Controller)
-{
-	PlayerController = Controller;
-}
-
 void SWorldWidgetScreenLayer::SetWidgetDrawSize(FVector2D InDrawSize)
 {
 	DrawSize = InDrawSize;
@@ -71,7 +68,7 @@ void SWorldWidgetScreenLayer::Tick(const FGeometry& AllottedGeometry, const doub
 {
 	TArray<USceneComponent*, TInlineAllocator<1>> DeadComponents;
 
-	if ( APlayerController* Player = PlayerController.Get() )
+	if ( APlayerController* PlayerController = PlayerContext.GetPlayerController() )
 	{
 		for ( auto It = ComponentMap.CreateIterator(); It; ++It )
 		{
@@ -82,7 +79,7 @@ void SWorldWidgetScreenLayer::Tick(const FGeometry& AllottedGeometry, const doub
 				FVector WorldLocation = SceneComponent->GetComponentLocation();
 
 				FVector ScreenPosition;
-				const bool bProjected = UWidgetLayoutLibrary::ProjectWorldLocationToWidgetPositionWithDistance(Player, WorldLocation, ScreenPosition);
+				const bool bProjected = UWidgetLayoutLibrary::ProjectWorldLocationToWidgetPositionWithDistance(PlayerController, WorldLocation, ScreenPosition);
 
 				if ( bProjected )
 				{
diff --git a/Engine/Source/Runtime/UMG/Public/Slate/SWorldWidgetScreenLayer.h b/Engine/Source/Runtime/UMG/Public/Slate/SWorldWidgetScreenLayer.h
--- a/Engine/Source/Runtime/UMG/Public/Slate/SWorldWidgetScreenLayer.h
+++ b/Engine/Source/Runtime/UMG/Public/Slate/SWorldWidgetScreenLayer.h
@@ -11,7 +11,7 @@ class SWorldWidgetScreenLayer : public SCompoundWidget
 	SLATE_END_ARGS()
 
 public:
-	void Construct(const FArguments& InArgs);
+	void Construct(const FArguments& InArgs, const FLocalPlayerContext& InPlayerContext);
 
 	virtual void Tick(const FGeometry& AllottedGeometry, const double InCurrentTime, const float InDeltaTime) override;
 
@@ -19,14 +19,12 @@ public:
 
 	void SetWidgetPivot(FVector2D Pivot);
 
-	void SetOwningPlayer(APlayerController* Controller);
-
 	void AddComponent(USceneComponent* Component, TSharedPtr<SWidget> Widget);
 
 	void RemoveComponent(USceneComponent* Component);
 
 private:
-	TWeakObjectPtr<APlayerController> PlayerController;
+	FLocalPlayerContext PlayerContext;
 
 	FVector2D DrawSize;
 	FVector2D Pivot;
