Slate - Fixing some build errors with the FWidgetRenderer.

diff --git a/Engine/Source/Runtime/UMG/Public/Slate/WidgetRenderer.h b/Engine/Source/Runtime/UMG/Public/Slate/WidgetRenderer.h
--- a/Engine/Source/Runtime/UMG/Public/Slate/WidgetRenderer.h
+++ b/Engine/Source/Runtime/UMG/Public/Slate/WidgetRenderer.h
@@ -2,6 +2,10 @@
 
 #pragma once
 
+class ISlate3DRenderer;
+class FHittestGrid;
+class SWindow;
+
 class SVirtualWindow : public SWindow
 {
 	SLATE_BEGIN_ARGS(SVirtualWindow)
@@ -29,9 +33,9 @@ public:
 
 	UTextureRenderTarget2D* DrawWidget(TSharedRef<SWidget>& Widget, FVector2D DrawSize);
 	void DrawWidget(UTextureRenderTarget2D* RenderTarget, TSharedRef<SWidget>& Widget, FVector2D DrawSize, float DeltaTime);
-	void DrawWindow(UTextureRenderTarget2D* RenderTarget, TSharedRef<class FHittestGrid> HitTestGrid, TSharedRef<class SWindow> Window, FVector2D DrawSize, float DeltaTime);
+	void DrawWindow(UTextureRenderTarget2D* RenderTarget, TSharedRef<FHittestGrid> HitTestGrid, TSharedRef<SWindow> Window, FVector2D DrawSize, float DeltaTime);
 
 private:
 	/** The slate 3D renderer used to render the user slate widget */
-	TSharedPtr<class ISlate3DRenderer> Renderer;
+	TSharedPtr<ISlate3DRenderer> Renderer;
 };
