UMG - Including some needed headers in WidgetRenderer.

diff --git a/Engine/Source/Runtime/UMG/Private/Slate/WidgetRenderer.cpp b/Engine/Source/Runtime/UMG/Private/Slate/WidgetRenderer.cpp
--- a/Engine/Source/Runtime/UMG/Private/Slate/WidgetRenderer.cpp
+++ b/Engine/Source/Runtime/UMG/Private/Slate/WidgetRenderer.cpp
@@ -3,6 +3,11 @@
 #include "UMGPrivatePCH.h"
 #include "WidgetRenderer.h"
 
+#if !UE_SERVER
+#include "ISlateRHIRendererModule.h"
+#include "ISlate3DRenderer.h"
+#endif // !UE_SERVER
+
 extern SLATECORE_API int32 bFoldTick;
 
 void SVirtualWindow::Construct(const FArguments& InArgs)
