Fixed server build.

- Not a particularly elegant fix but it avoids pulling in more dependencies.

(Merging CL 2708864 to the main branch).

diff --git a/Engine/Source/Runtime/UMG/Private/Slate/WidgetRenderer.cpp b/Engine/Source/Runtime/UMG/Private/Slate/WidgetRenderer.cpp
--- a/Engine/Source/Runtime/UMG/Private/Slate/WidgetRenderer.cpp
+++ b/Engine/Source/Runtime/UMG/Private/Slate/WidgetRenderer.cpp
@@ -56,6 +56,7 @@ void FWidgetRenderer::DrawWidget(UTextureRenderTarget2D* RenderTarget, TSharedRe
 
 void FWidgetRenderer::DrawWindow(UTextureRenderTarget2D* RenderTarget, TSharedRef<FHittestGrid> HitTestGrid, TSharedRef<SWindow> Window, FVector2D DrawSize, float DeltaTime)
 {
+#if !UE_SERVER
 	FGeometry WindowGeometry = FGeometry::MakeRoot(DrawSize, FSlateLayoutTransform());
 
 	if ( !bFoldTick )
@@ -95,4 +96,5 @@ void FWidgetRenderer::DrawWindow(UTextureRenderTarget2D* RenderTarget, TSharedRe
 		{
 			InRenderer->DrawWindowToTarget_RenderThread(RHICmdList, InRenderTarget, InDrawBuffer);
 		});
+#endif // !UE_SERVER
 }
