fix for Non Unity CIS compile failure

diff --git a/Engine/Source/Runtime/UMG/Private/Slate/SWorldWidgetScreenLayer.cpp b/Engine/Source/Runtime/UMG/Private/Slate/SWorldWidgetScreenLayer.cpp
--- a/Engine/Source/Runtime/UMG/Private/Slate/SWorldWidgetScreenLayer.cpp
+++ b/Engine/Source/Runtime/UMG/Private/Slate/SWorldWidgetScreenLayer.cpp
@@ -2,6 +2,9 @@
 
 #include "UMGPrivatePCH.h"
 
+#include "WidgetComponent.h"
+#include "WidgetLayoutLibrary.h"
+
 #include "SWorldWidgetScreenLayer.h"
 
 void SWorldWidgetScreenLayer::Construct(const FArguments& InArgs, const FLocalPlayerContext& InPlayerContext)
diff --git a/Engine/Source/Runtime/UMG/Private/Slate/WidgetRenderer.cpp b/Engine/Source/Runtime/UMG/Private/Slate/WidgetRenderer.cpp
--- a/Engine/Source/Runtime/UMG/Private/Slate/WidgetRenderer.cpp
+++ b/Engine/Source/Runtime/UMG/Private/Slate/WidgetRenderer.cpp
@@ -2,6 +2,7 @@
 
 #include "UMGPrivatePCH.h"
 #include "WidgetRenderer.h"
+#include "HittestGrid.h"
 
 #if !UE_SERVER
 #include "ISlateRHIRendererModule.h"
@@ -46,7 +47,7 @@ UTextureRenderTarget2D* FWidgetRenderer::DrawWidget(TSharedRef<SWidget>& Widget,
 void FWidgetRenderer::DrawWidget(UTextureRenderTarget2D* RenderTarget, TSharedRef<SWidget>& Widget, FVector2D DrawSize, float DeltaTime)
 {
 	TSharedRef<SVirtualWindow> Window = SNew(SVirtualWindow).Size(DrawSize);
-	TSharedRef<FHittestGrid> HitTestGrid = MakeShareable(new FHittestGrid);
+	TSharedRef<FHittestGrid> HitTestGrid = MakeShareable(new FHittestGrid());
 
 	Window->SetContent(Widget);
 	Window->Resize(DrawSize);
