UMG - UserWidgets are no longer Focusable by default.  Also adding Padding to be settable on a UserWidget.

diff --git a/Engine/Source/Runtime/Core/Public/UObject/ObjectVersion.h b/Engine/Source/Runtime/Core/Public/UObject/ObjectVersion.h
--- a/Engine/Source/Runtime/Core/Public/UObject/ObjectVersion.h
+++ b/Engine/Source/Runtime/Core/Public/UObject/ObjectVersion.h
@@ -574,6 +574,8 @@ enum EUnrealEngineObjectUE4Version
 	VER_UE4_K2NODE_VAR_REFERENCEGUIDS,
 	// Added support for sound concurrency settings structure and overrides
 	VER_UE4_SOUND_CONCURRENCY_PACKAGE,
+	// Changing the default value for focusable user widgets to false
+	VER_UE4_USERWIDGET_DEFAULT_FOCUSABLE_FALSE,
 	// -----<new versions can be added before this line>-------------------------------------------------
 	// - this needs to be the last line (see note below)
 	VER_UE4_AUTOMATIC_VERSION_PLUS_ONE,
diff --git a/Engine/Source/Runtime/UMG/Private/Slate/SObjectWidget.cpp b/Engine/Source/Runtime/UMG/Private/Slate/SObjectWidget.cpp
--- a/Engine/Source/Runtime/UMG/Private/Slate/SObjectWidget.cpp
+++ b/Engine/Source/Runtime/UMG/Private/Slate/SObjectWidget.cpp
@@ -51,6 +51,11 @@ void SObjectWidget::AddReferencedObjects(FReferenceCollector& Collector)
 	Collector.AddReferencedObject(WidgetObject);
 }
 
+void SObjectWidget::SetPadding(const TAttribute<FMargin>& InMargin)
+{
+	ChildSlot.Padding(InMargin);
+}
+
 void SObjectWidget::Tick( const FGeometry& AllottedGeometry, const double InCurrentTime, const float InDeltaTime )
 {
 	if ( CanRouteEvent() )
diff --git a/Engine/Source/Runtime/UMG/Private/UserWidget.cpp b/Engine/Source/Runtime/UMG/Private/UserWidget.cpp
--- a/Engine/Source/Runtime/UMG/Private/UserWidget.cpp
+++ b/Engine/Source/Runtime/UMG/Private/UserWidget.cpp
@@ -23,7 +23,8 @@ UUserWidget::UUserWidget(const FObjectInitializer& ObjectInitializer)
 	Visiblity_DEPRECATED = Visibility = ESlateVisibility::SelfHitTestInvisible;
 
 	bInitialized = false;
-	bSupportsKeyboardFocus = true;
+	bSupportsKeyboardFocus_DEPRECATED = true;
+	bIsFocusable = false;
 	ColorAndOpacity = FLinearColor::White;
 	ForegroundColor = FSlateColor::UseForeground();
 
@@ -138,6 +139,7 @@ void UUserWidget::SynchronizeProperties()
 
 		SafeGCWidget->SetColorAndOpacity(ColorBinding);
 		SafeGCWidget->SetForegroundColor(ForegroundColorBinding);
+		SafeGCWidget->SetPadding(Padding);
 	}
 }
 
@@ -163,6 +165,17 @@ void UUserWidget::SetForegroundColor(FSlateColor InForegroundColor)
 	}
 }
 
+void UUserWidget::SetPadding(FMargin InPadding)
+{
+	Padding = InPadding;
+
+	TSharedPtr<SObjectWidget> SafeGCWidget = MyGCWidget.Pin();
+	if ( SafeGCWidget.IsValid() )
+	{
+		SafeGCWidget->SetPadding(Padding);
+	}
+}
+
 void UUserWidget::PostInitProperties()
 {
 	Super::PostInitProperties();
@@ -928,7 +941,7 @@ bool UUserWidget::NativeIsInteractable() const
 
 bool UUserWidget::NativeSupportsKeyboardFocus() const
 {
-	return bSupportsKeyboardFocus;
+	return bIsFocusable;
 }
 
 FReply UUserWidget::NativeOnFocusReceived( const FGeometry& InGeometry, const FFocusEvent& InFocusEvent )
@@ -1070,6 +1083,11 @@ void UUserWidget::PostLoad()
 {
 	Super::PostLoad();
 
+	if ( GetLinkerUE4Version() < VER_UE4_USERWIDGET_DEFAULT_FOCUSABLE_FALSE )
+	{
+		bIsFocusable = bSupportsKeyboardFocus_DEPRECATED;
+	}
+
 #if WITH_EDITORONLY_DATA
 
 	if ( GetLinkerUE4Version() < VER_UE4_DEPRECATE_USER_WIDGET_DESIGN_SIZE )
diff --git a/Engine/Source/Runtime/UMG/Public/Blueprint/UserWidget.h b/Engine/Source/Runtime/UMG/Public/Blueprint/UserWidget.h
--- a/Engine/Source/Runtime/UMG/Public/Blueprint/UserWidget.h
+++ b/Engine/Source/Runtime/UMG/Public/Blueprint/UserWidget.h
@@ -705,6 +705,9 @@ public:
 	UFUNCTION(BlueprintCallable, BlueprintCosmetic, Category="Appearance")
 	void SetForegroundColor(FSlateColor InForegroundColor);
 
+	UFUNCTION(BlueprintCallable, BlueprintCosmetic, Category="Appearance")
+	void SetPadding(FMargin InPadding);
+
 	/**
 	 * Plays an animation in this widget a specified number of times
 	 * 
@@ -803,9 +806,16 @@ public:
 	UPROPERTY()
 	FGetSlateColor ForegroundColorDelegate;
 
+	/** The padding area around the content. */
+	UPROPERTY(EditAnywhere, BlueprintReadOnly, Category="Appearance")
+	FMargin Padding;
+
+	UPROPERTY()
+	bool bSupportsKeyboardFocus_DEPRECATED;
+
 	/** Setting this flag to true, allows this widget to accept focus when clicked, or when navigated to. */
-	UPROPERTY(EditAnywhere, BlueprintReadWrite, Category="Interaction", meta=(DisplayName="Is Focusable"))
-	bool bSupportsKeyboardFocus;
+	UPROPERTY(EditAnywhere, BlueprintReadWrite, Category="Interaction")
+	bool bIsFocusable;
 
 	/** The widget tree contained inside this user widget initialized by the blueprint */
 	UPROPERTY(Transient)
diff --git a/Engine/Source/Runtime/UMG/Public/Slate/SObjectWidget.h b/Engine/Source/Runtime/UMG/Public/Slate/SObjectWidget.h
--- a/Engine/Source/Runtime/UMG/Public/Slate/SObjectWidget.h
+++ b/Engine/Source/Runtime/UMG/Public/Slate/SObjectWidget.h
@@ -32,6 +32,8 @@ class UMG_API SObjectWidget : public SCompoundWidget, public FGCObject
 
 	UUserWidget* GetWidgetObject() const { return WidgetObject; }
 
+	void SetPadding(const TAttribute<FMargin>& InMargin);
+
 	virtual void Tick( const FGeometry& AllottedGeometry, const double InCurrentTime, const float InDeltaTime ) override;
 	virtual int32 OnPaint(const FPaintArgs& Args, const FGeometry& AllottedGeometry, const FSlateRect& MyClippingRect, FSlateWindowElementList& OutDrawElements, int32 LayerId, const FWidgetStyle& InWidgetStyle, bool bParentEnabled) const override;
 
