Removed Expirimental tag from ExpandableArea

diff --git a/Engine/Source/Runtime/UMG/Private/Components/ExpandableArea.cpp b/Engine/Source/Runtime/UMG/Private/Components/ExpandableArea.cpp
--- a/Engine/Source/Runtime/UMG/Private/Components/ExpandableArea.cpp
+++ b/Engine/Source/Runtime/UMG/Private/Components/ExpandableArea.cpp
@@ -17,6 +17,11 @@ UExpandableArea::UExpandableArea(const FObjectInitializer& ObjectInitializer)
 	, bIsExpanded(false)
 {
 	bIsVariable = false;
+
+	SExpandableArea::FArguments ExpandableDefaults;
+	Style       = *ExpandableDefaults._Style;
+	BorderColor = ExpandableDefaults._BorderBackgroundColor.Get( FLinearColor::White );
+	BorderBrush = *ExpandableDefaults._BorderImage;
 }
 
 bool UExpandableArea::GetIsExpanded() const
@@ -83,6 +88,9 @@ TSharedRef<SWidget> UExpandableArea::RebuildWidget()
 	TSharedRef<SWidget> BodyWidget = BodyContent ? BodyContent->TakeWidget() : SNullWidget::NullWidget;
 
 	MyExpandableArea = SNew(SExpandableArea)
+		.Style(&Style)
+		.BorderImage(&BorderBrush)
+		.BorderBackgroundColor(BorderColor)
 		.MaxHeight(MaxHeight)
 		.Padding(AreaPadding)
 		.OnAreaExpansionChanged(BIND_UOBJECT_DELEGATE(FOnBooleanValueChanged, SlateExpansionChanged))
@@ -111,7 +119,7 @@ void UExpandableArea::SlateExpansionChanged(bool NewState)
 
 	if ( OnExpansionChanged.IsBound() )
 	{
-		OnExpansionChanged.Broadcast(NewState);
+		OnExpansionChanged.Broadcast(this, NewState);
 	}
 }
 
diff --git a/Engine/Source/Runtime/UMG/Public/Components/ExpandableArea.h b/Engine/Source/Runtime/UMG/Public/Components/ExpandableArea.h
--- a/Engine/Source/Runtime/UMG/Public/Components/ExpandableArea.h
+++ b/Engine/Source/Runtime/UMG/Public/Components/ExpandableArea.h
@@ -4,17 +4,27 @@
 
 #include "ExpandableArea.generated.h"
 
-DECLARE_DYNAMIC_MULTICAST_DELEGATE_OneParam(FOnExpandableAreaExpansionChanged, bool, bIsExpanded);
+DECLARE_DYNAMIC_MULTICAST_DELEGATE_TwoParams(FOnExpandableAreaExpansionChanged, UExpandableArea*, Area, bool, bIsExpanded);
 
 /**
  * 
  */
-UCLASS(Experimental)
+UCLASS()
 class UMG_API UExpandableArea : public UWidget, public INamedSlotInterface
 {
 	GENERATED_UCLASS_BODY()
 
 public:
+
+	UPROPERTY(EditAnywhere, BlueprintReadOnly, Category = "Style")
+	FExpandableAreaStyle Style;
+
+	UPROPERTY( EditAnywhere, BlueprintReadOnly, Category = "Style" )
+	FSlateBrush BorderBrush;
+
+	UPROPERTY( EditAnywhere, BlueprintReadOnly, Category = "Style" )
+	FSlateColor BorderColor;
+
 	/**  */
 	UPROPERTY(EditAnywhere, BlueprintReadOnly, Category="Expansion")
 	bool bIsExpanded;
