UMG - fixing non-unity dedicated server build with widget renderer.

diff --git a/Engine/Source/Runtime/UMG/Private/Slate/WidgetRenderer.cpp b/Engine/Source/Runtime/UMG/Private/Slate/WidgetRenderer.cpp
--- a/Engine/Source/Runtime/UMG/Private/Slate/WidgetRenderer.cpp
+++ b/Engine/Source/Runtime/UMG/Private/Slate/WidgetRenderer.cpp
@@ -21,7 +21,9 @@ void SVirtualWindow::Construct(const FArguments& InArgs)
 
 FWidgetRenderer::FWidgetRenderer()
 {
+#if !UE_SERVER
 	Renderer = FModuleManager::Get().LoadModuleChecked<ISlateRHIRendererModule>("SlateRHIRenderer").CreateSlate3DRenderer();
+#endif
 }
 
 FWidgetRenderer::~FWidgetRenderer()
