UMG - Fix an issue where starting a new animation in response to an animation ending would crash due to modifying a ranged for. (Checking in fix made by [REDACTED_PERSON])

diff --git a/Engine/Source/Runtime/UMG/Private/UserWidget.cpp b/Engine/Source/Runtime/UMG/Private/UserWidget.cpp
--- a/Engine/Source/Runtime/UMG/Private/UserWidget.cpp
+++ b/Engine/Source/Runtime/UMG/Private/UserWidget.cpp
@@ -793,10 +793,13 @@ void UUserWidget::TickActionsAndAnimation(const FGeometry& MyGeometry, float InD
 		return;
 	}
 
-	// Update active movie scenes
-	for ( UUMGSequencePlayer* Player : ActiveSequencePlayers )
+	// Update active movie scenes, none will be removed here, but new
+	// ones can be added during the tick, if a player ends and triggers
+	// starting another animation
+	for ( int32 Index = 0; Index < ActiveSequencePlayers.Num(); Index++ )
 	{
-		Player->Tick(InDeltaTime);
+		UUMGSequencePlayer* Player = ActiveSequencePlayers[Index];
+		Player->Tick( InDeltaTime );
 	}
 
 	const bool bWasPlayingAnimation = IsPlayingAnimation();
