Mergine engine changes to //UE4/Main @ 2748007

diff --git a/Engine/Source/Runtime/UMG/Private/Slate/SRetainerWidget.cpp b/Engine/Source/Runtime/UMG/Private/Slate/SRetainerWidget.cpp
new file mode 100644
--- /dev/null
+++ b/Engine/Source/Runtime/UMG/Private/Slate/SRetainerWidget.cpp
@@ -0,0 +1,362 @@
+// Copyright 1998-2015 Epic Games, Inc. All Rights Reserved.
+
+#include "UMGPrivatePCH.h"
+
+#include "Slate/SlateTextures.h"
+#include "SRetainerWidget.h"
+#include "Runtime/SlateRHIRenderer/Public/Interfaces/ISlateRHIRendererModule.h"
+
+#if !(UE_BUILD_SHIPPING || UE_BUILD_TEST)
+
+/** True if we should allow widgets to be cached in the UI at all. */
+TAutoConsoleVariable<int32> EnableRetainedRendering(
+	TEXT("Slate.EnableRetainedRendering"),
+	true,
+	TEXT("Whether to attempt to render things in SRetainerWidgets to render targets first."));
+
+static bool IsRetainedRenderingEnabled()
+{
+	return EnableRetainedRendering.GetValueOnGameThread() == 1;
+}
+
+#else
+
+static bool IsRetainedRenderingEnabled()
+{
+	return true;
+}
+
+#endif
+
+
+
+static FVector2D RoundToInt(const FVector2D& Vec)
+{
+	return FVector2D(FMath::RoundToInt(Vec.X), FMath::RoundToInt(Vec.Y));
+}
+
+/**
+* Used to construct a rotated rect from an aligned clip rect and a set of layout and render transforms from the geometry, snapped to pixel boundaries. Returns a float or float16 version of the rect based on the typedef.
+*/
+static FSlateRotatedClipRectType ToSnappedRotatedRect(const FSlateRect& ClipRectInLayoutWindowSpace, const FSlateLayoutTransform& InverseLayoutTransform, const FSlateRenderTransform& RenderTransform)
+{
+	FSlateRotatedRect RotatedRect = TransformRect(
+		Concatenate(InverseLayoutTransform, RenderTransform),
+		FSlateRotatedRect(ClipRectInLayoutWindowSpace));
+
+	// Pixel snapping is done here by rounding the resulting floats to ints, we do this before
+	// calculating the final extents of the clip box otherwise we'll get a smaller clip rect than a visual
+	// rect where each point is individually snapped.
+	FVector2D SnappedTopLeft = RoundToInt(RotatedRect.TopLeft);
+	FVector2D SnappedTopRight = RoundToInt(RotatedRect.TopLeft + RotatedRect.ExtentX);
+	FVector2D SnappedBottomLeft = RoundToInt(RotatedRect.TopLeft + RotatedRect.ExtentY);
+
+	//NOTE: We explicitly do not re-snap the extent x/y, it wouldn't be correct to snap again in distance space
+	// even if two points are snapped, their distance wont necessarily be a whole number if those points are not
+	// axis aligned.
+	return FSlateRotatedClipRectType(
+		SnappedTopLeft,
+		SnappedTopRight - SnappedTopLeft,
+		SnappedBottomLeft - SnappedTopLeft);
+}
+
+
+
+
+//---------------------------------------------------
+
+SRetainerWidget::SRetainerWidget()
+	: WidgetRenderer( /* bUseGammaCorrection */ true)
+{
+}
+
+SRetainerWidget::~SRetainerWidget()
+{
+	ENQUEUE_UNIQUE_RENDER_COMMAND_ONEPARAMETER(DeleteShaderResource,
+		FSlateRenderTargetRHI*, InRenderTargetRHI, RenderTargetResource,
+		{
+			delete InRenderTargetRHI;
+		});
+}
+
+void SRetainerWidget::Construct(const FArguments& InArgs)
+{
+	RenderTarget = NewObject<UTextureRenderTarget2D>();
+	RenderTarget->SRGB = true;
+	RenderTarget->TargetGamma = 1;
+	RenderTarget->ClearColor = FLinearColor::Transparent;
+
+	RenderTargetBrush.SetResourceObject(RenderTarget);
+
+	Window = SNew(SVirtualWindow);
+	HitTestGrid = MakeShareable(new FHittestGrid());
+
+	WidgetRenderer.SetIsPrepassNeeded(false);
+
+	RenderTargetResource = new FSlateRenderTargetRHI(nullptr, 0, 0);
+
+	MyWidget = InArgs._Content.Widget;
+	Phase = InArgs._Phase;
+	PhaseCount = InArgs._PhaseCount;
+
+	LastDrawTime = FApp::GetCurrentTime();
+	LastTickedFrame = 0;
+
+	bRenderingOffscreenDesire = true;
+	bRenderingOffscreen = false;
+
+	ChildSlot
+	[
+		MyWidget.ToSharedRef()
+	];
+
+	Window->SetContent(MyWidget.ToSharedRef());
+}
+
+FIntPoint SRetainerWidget::GetSize() const
+{
+	FIntPoint Point;
+	Point.X = (int32)RenderTarget->GetSurfaceWidth();
+	Point.Y = (int32)RenderTarget->GetSurfaceHeight();
+
+	return Point;
+}
+
+bool SRetainerWidget::RequiresVsync() const
+{
+	return false;
+}
+
+FSlateShaderResource* SRetainerWidget::GetViewportRenderTargetTexture() const
+{
+	return RenderTargetResource;
+}
+
+bool SRetainerWidget::IsViewportTextureAlphaOnly() const
+{
+	return false;
+}
+
+bool SRetainerWidget::ShouldBeRenderingOffscreen() const
+{
+	return bRenderingOffscreenDesire && IsRetainedRenderingEnabled();
+}
+
+void SRetainerWidget::SetRetainedRendering(bool bRetainRendering)
+{
+	bRenderingOffscreenDesire = bRetainRendering;
+}
+
+void SRetainerWidget::RefreshRenderingMode()
+{
+	const bool bShouldBeRenderingOffscreen = ShouldBeRenderingOffscreen();
+
+	if ( bRenderingOffscreen != bShouldBeRenderingOffscreen )
+	{
+		bRenderingOffscreen = bShouldBeRenderingOffscreen;
+
+		ChildSlot
+		[
+			MyWidget.ToSharedRef()
+		];
+
+		// TODO When parent pointers come into play, we need to make it so the virtual window,
+		// doesn't actually ever own the content, it just sorta kind of owns it.
+		if ( bRenderingOffscreen )
+		{
+			Window->SetContent(MyWidget.ToSharedRef());
+		}
+		else
+		{
+			Window->SetContent(SNullWidget::NullWidget);
+		}
+	}
+}
+
+void SRetainerWidget::SetContent(const TSharedRef< SWidget >& InContent)
+{
+	MyWidget = InContent;
+
+	ChildSlot
+	[
+		InContent
+	];
+
+	if ( bRenderingOffscreen )
+	{
+		Window->SetContent(InContent);
+	}
+}
+
+void SRetainerWidget::AddReferencedObjects(FReferenceCollector& Collector)
+{
+	Collector.AddReferencedObject(RenderTarget);
+}
+
+FChildren* SRetainerWidget::GetChildren()
+{
+	if ( bRenderingOffscreen )
+	{
+		return &EmptyChildSlot;
+	}
+	else
+	{
+		return SCompoundWidget::GetChildren();
+	}
+}
+
+bool SRetainerWidget::ComputeVolatility() const
+{
+	return true;
+}
+
+void SRetainerWidget::Tick(float DeltaTime)
+{
+	if ( bRenderingOffscreen )
+	{
+		if ( LastTickedFrame != GFrameCounter && ( GFrameCounter % PhaseCount ) == Phase )
+		{
+			LastTickedFrame = GFrameCounter;
+			double TimeSinceLastDraw = FApp::GetCurrentTime() - LastDrawTime;
+
+			FPaintGeometry PaintGeometry = CachedAllottedGeometry.ToPaintGeometry();
+
+			// extract the layout transform from the draw element
+			FSlateLayoutTransform InverseLayoutTransform(Inverse(FSlateLayoutTransform(PaintGeometry.DrawScale, PaintGeometry.DrawPosition)));
+			// The clip rect is NOT subject to the rotations specified by MakeRotatedBox.
+			FSlateRotatedClipRectType RenderClipRect = ToSnappedRotatedRect(CachedClippingRect, InverseLayoutTransform, CachedAllottedGeometry.GetAccumulatedRenderTransform());
+
+			const uint32 RenderTargetWidth  = FMath::RoundToInt(RenderClipRect.ExtentX.X);
+			const uint32 RenderTargetHeight = FMath::RoundToInt(RenderClipRect.ExtentY.Y);
+
+			// Need to prepass.
+			Window->SlatePrepass(CachedAllottedGeometry.Scale);
+
+			if ( RenderTargetWidth != 0 && RenderTargetHeight != 0 )
+			{
+				if ( MyWidget->GetVisibility().IsVisible() )
+				{
+					if ( RenderTargetResource->GetWidth() != RenderTargetWidth ||
+						RenderTargetResource->GetHeight() != RenderTargetHeight )
+					{
+						const bool bForceLinearGamma = false;
+						RenderTarget->InitCustomFormat(RenderTargetWidth, RenderTargetHeight, PF_B8G8R8A8, bForceLinearGamma);
+
+						ENQUEUE_UNIQUE_RENDER_COMMAND_FOURPARAMETER(UpdateShaderResource,
+							UTextureRenderTarget2D*, InRenderTarget, RenderTarget,
+							FSlateRenderTargetRHI*, InRenderTargetRHI, RenderTargetResource,
+							uint32, InRenderTargetWidth, RenderTargetWidth,
+							uint32, InRenderTargetHeight, RenderTargetHeight,
+							{
+								FTextureRenderTarget2DResource* CurrentResource = static_cast<FTextureRenderTarget2DResource*>( InRenderTarget->GetRenderTargetResource() );
+								InRenderTargetRHI->SetRHIRef(CurrentResource->GetTextureRHI(), InRenderTargetWidth, InRenderTargetHeight);
+							});
+					}
+
+					// TODO Need to improve widget renderer to allow us to not render deferred rendered widgets into the render target, as they will likely be clipped.
+
+					WidgetRenderer.DrawWindow(
+						RenderTarget,
+						HitTestGrid.ToSharedRef(),
+						Window.ToSharedRef(),
+						CachedAllottedGeometry.Scale,
+						FVector2D(RenderTargetWidth, RenderTargetHeight),
+						TimeSinceLastDraw);
+				}
+			}
+
+			LastDrawTime = FApp::GetCurrentTime();
+		}
+	}
+}
+
+int32 SRetainerWidget::OnPaint(const FPaintArgs& Args, const FGeometry& AllottedGeometry, const FSlateRect& MyClippingRect, FSlateWindowElementList& OutDrawElements, int32 LayerId, const FWidgetStyle& InWidgetStyle, bool bParentEnabled) const
+{
+	SRetainerWidget* MutableThis = const_cast<SRetainerWidget*>( this );
+
+	MutableThis->RefreshRenderingMode();
+
+	if ( bRenderingOffscreen )
+	{
+		CachedAllottedGeometry = AllottedGeometry;
+		CachedClippingRect = MyClippingRect;
+
+		if ( RenderTarget->GetSurfaceWidth() >= 1 && RenderTarget->GetSurfaceHeight() >= 1 )
+		{
+			const FLinearColor FinalColorAndOpacity(InWidgetStyle.GetColorAndOpacityTint() * ColorAndOpacity.Get() * RenderTargetBrush.GetTint(InWidgetStyle));
+
+			const bool bEnableGammaCorrection = false;
+			const bool bEnableBlending = true;
+
+			if ( RenderTargetResource->GetWidth() != 0 && RenderTargetResource->GetWidth() != 0 )
+			{
+				FSlateDrawElement::MakeViewport(
+					OutDrawElements,
+					LayerId,
+					AllottedGeometry.ToPaintGeometry(),
+					SharedThis(this),
+					MyClippingRect,
+					bEnableGammaCorrection,
+					bEnableBlending,
+					ESlateDrawEffect::PreMultipliedAlpha,
+					FinalColorAndOpacity);
+			}
+			
+			TSharedRef<SRetainerWidget> SharedMutableThis = SharedThis(MutableThis);
+			Args.InsertCustomHitTestPath(SharedMutableThis, Args.GetLastHitTestIndex());
+		}
+		
+		return LayerId;
+	}
+	else
+	{
+		return SCompoundWidget::OnPaint(Args, AllottedGeometry, MyClippingRect, OutDrawElements, LayerId, InWidgetStyle, bParentEnabled);
+	}
+}
+
+FVector2D SRetainerWidget::ComputeDesiredSize(float LayoutScaleMuliplier) const
+{
+	if ( bRenderingOffscreen )
+	{
+		return MyWidget->GetDesiredSize();
+	}
+	else
+	{
+		return SCompoundWidget::ComputeDesiredSize(LayoutScaleMuliplier);
+	}
+}
+
+TArray<FWidgetAndPointer> SRetainerWidget::GetBubblePathAndVirtualCursors(const FGeometry& InGeometry, FVector2D DesktopSpaceCoordinate, bool bIgnoreEnabledStatus) const
+{
+	const FVector2D LocalPosition = InGeometry.AbsoluteToLocal(DesktopSpaceCoordinate);
+	const FVector2D LastLocalPosition = InGeometry.AbsoluteToLocal(DesktopSpaceCoordinate);
+
+	TSharedRef<FVirtualPointerPosition> VirtualMouseCoordinate = MakeShareable(new FVirtualPointerPosition(LocalPosition, LastLocalPosition));
+
+	// TODO Where should this come from?
+	const float CursorRadius = 0;
+
+	TArray<FWidgetAndPointer> ArrangedWidgets = 
+		HitTestGrid->GetBubblePath(LocalPosition, CursorRadius, bIgnoreEnabledStatus);
+
+	for ( FWidgetAndPointer& ArrangedWidget : ArrangedWidgets )
+	{
+		ArrangedWidget.PointerPosition = VirtualMouseCoordinate;
+	}
+
+	return ArrangedWidgets;
+}
+
+void SRetainerWidget::ArrangeChildren(FArrangedChildren& ArrangedChildren) const
+{
+	ArrangedChildren.AddWidget(
+		FArrangedWidget(MyWidget.ToSharedRef(), CachedAllottedGeometry));
+}
+
+TSharedPtr<struct FVirtualPointerPosition> SRetainerWidget::TranslateMouseCoordinateFor3DChild(const TSharedRef<SWidget>& ChildWidget, const FGeometry& InGeometry, const FVector2D& ScreenSpaceMouseCoordinate, const FVector2D& LastScreenSpaceMouseCoordinate) const
+{
+	const FVector2D LocalPosition = InGeometry.AbsoluteToLocal(ScreenSpaceMouseCoordinate);
+	const FVector2D LastLocalPosition = InGeometry.AbsoluteToLocal(LastScreenSpaceMouseCoordinate);
+
+	return MakeShareable(new FVirtualPointerPosition(LocalPosition, LastLocalPosition));
+}
diff --git a/Engine/Source/Runtime/UMG/Private/Slate/WidgetRenderer.cpp b/Engine/Source/Runtime/UMG/Private/Slate/WidgetRenderer.cpp
--- a/Engine/Source/Runtime/UMG/Private/Slate/WidgetRenderer.cpp
+++ b/Engine/Source/Runtime/UMG/Private/Slate/WidgetRenderer.cpp
@@ -19,10 +19,12 @@ void SVirtualWindow::Construct(const FArguments& InArgs)
 	SetContent(SNullWidget::NullWidget);
 }
 
-FWidgetRenderer::FWidgetRenderer()
+FWidgetRenderer::FWidgetRenderer(bool bUseGammaCorrection)
+	: bPrepassNeeded(true)
+	, bUseGammaSpace(bUseGammaCorrection)
 {
 #if !UE_SERVER
-	Renderer = FModuleManager::Get().LoadModuleChecked<ISlateRHIRendererModule>("SlateRHIRenderer").CreateSlate3DRenderer();
+	Renderer = FModuleManager::Get().LoadModuleChecked<ISlateRHIRendererModule>("SlateRHIRenderer").CreateSlate3DRenderer(bUseGammaSpace);
 #endif
 }
 
@@ -35,18 +37,26 @@ ISlate3DRenderer* FWidgetRenderer::GetSlateRenderer()
 	return Renderer.Get();
 }
 
-UTextureRenderTarget2D* FWidgetRenderer::DrawWidget(TSharedRef<SWidget>& Widget, FVector2D DrawSize)
+UTextureRenderTarget2D* FWidgetRenderer::DrawWidget(const TSharedRef<SWidget>& Widget, FVector2D DrawSize)
 {
-	UTextureRenderTarget2D* RenderTarget = NewObject<UTextureRenderTarget2D>();
-	RenderTarget->ClearColor = FLinearColor::Transparent;
-	RenderTarget->InitCustomFormat(DrawSize.X, DrawSize.Y, PF_B8G8R8A8, false);
+	UTextureRenderTarget2D* RenderTarget = FWidgetRenderer::CreateTargetFor(DrawSize, TF_Bilinear, bUseGammaSpace);
 
 	DrawWidget(RenderTarget, Widget, DrawSize, 0);
 
 	return RenderTarget;
 }
 
-void FWidgetRenderer::DrawWidget(UTextureRenderTarget2D* RenderTarget, TSharedRef<SWidget>& Widget, FVector2D DrawSize, float DeltaTime)
+UTextureRenderTarget2D* FWidgetRenderer::CreateTargetFor(FVector2D DrawSize, TextureFilter InFilter, bool bUseGammaCorrection)
+{
+	UTextureRenderTarget2D* RenderTarget = NewObject<UTextureRenderTarget2D>();
+	RenderTarget->Filter = InFilter;
+	RenderTarget->ClearColor = FLinearColor::Transparent;
+	const bool bIsLinearSpace = !bUseGammaCorrection;
+	RenderTarget->InitCustomFormat(DrawSize.X, DrawSize.Y, PF_B8G8R8A8, bIsLinearSpace);
+	return RenderTarget;
+}
+
+void FWidgetRenderer::DrawWidget(UTextureRenderTarget2D* RenderTarget, const TSharedRef<SWidget>& Widget, FVector2D DrawSize, float DeltaTime)
 {
 	TSharedRef<SVirtualWindow> Window = SNew(SVirtualWindow).Size(DrawSize);
 	TSharedRef<FHittestGrid> HitTestGrid = MakeShareable(new FHittestGrid());
@@ -54,21 +64,30 @@ void FWidgetRenderer::DrawWidget(UTextureRenderTarget2D* RenderTarget, TSharedRe
 	Window->SetContent(Widget);
 	Window->Resize(DrawSize);
 
-	DrawWindow(RenderTarget, HitTestGrid, Window, DrawSize, DeltaTime);
+	DrawWindow(RenderTarget, HitTestGrid, Window, 1, DrawSize, DeltaTime);
 }
 
-void FWidgetRenderer::DrawWindow(UTextureRenderTarget2D* RenderTarget, TSharedRef<FHittestGrid> HitTestGrid, TSharedRef<SWindow> Window, FVector2D DrawSize, float DeltaTime)
+void FWidgetRenderer::DrawWindow(
+	UTextureRenderTarget2D* RenderTarget,
+	TSharedRef<FHittestGrid> HitTestGrid,
+	TSharedRef<SWindow> Window,
+	float Scale,
+	FVector2D DrawSize,
+	float DeltaTime)
 {
 #if !UE_SERVER
-	FGeometry WindowGeometry = FGeometry::MakeRoot(DrawSize, FSlateLayoutTransform());
+	FGeometry WindowGeometry = FGeometry::MakeRoot(DrawSize * (1 / Scale), FSlateLayoutTransform(Scale));
 
 	if ( !bFoldTick )
 	{
 		Window->TickWidgetsRecursively(WindowGeometry, FApp::GetCurrentTime(), DeltaTime);
 	}
 
-	// Ticking can cause geometry changes.  Recompute
-	Window->SlatePrepass();
+	if ( bPrepassNeeded )
+	{
+		// Ticking can cause geometry changes.  Recompute
+		Window->SlatePrepass(Scale);
+	}
 
 	// Prepare the test grid 
 	HitTestGrid->ClearGridForNewFrame(WindowGeometry.GetClippingRect());
diff --git a/Engine/Source/Runtime/UMG/Public/Components/RetainerBox.h b/Engine/Source/Runtime/UMG/Public/Components/RetainerBox.h
new file mode 100644
--- /dev/null
+++ b/Engine/Source/Runtime/UMG/Public/Components/RetainerBox.h
@@ -0,0 +1,48 @@
+﻿// Copyright 1998-2015 Epic Games, Inc. All Rights Reserved.
+
+#pragma once
+
+#include "RetainerBox.generated.h"
+
+/**
+ * Invalidate
+ * * Single Child
+ * * Caching / Performance
+ */
+UCLASS()
+class UMG_API URetainerBox : public UContentWidget
+{
+	GENERATED_UCLASS_BODY()
+
+public:
+	UPROPERTY(EditAnywhere, Category="Phasing", meta=( UIMin=0, ClampMin=0 ))
+	int32 Phase;
+
+	UPROPERTY(EditAnywhere, Category="Phasing", meta=( UIMin=1, ClampMin=1 ))
+	int32 PhaseCount;
+
+public:
+
+	virtual void ReleaseSlateResources(bool bReleaseChildren) override;
+
+#if WITH_EDITOR
+	virtual const FSlateBrush* GetEditorIcon() override;
+	virtual const FText GetPaletteCategory() override;
+#endif
+
+protected:
+
+	// UPanelWidget
+	virtual void OnSlotAdded(UPanelSlot* Slot) override;
+	virtual void OnSlotRemoved(UPanelSlot* Slot) override;
+	// End UPanelWidget
+
+	// UWidget interface
+	virtual TSharedRef<SWidget> RebuildWidget() override;
+	// End of UWidget interface
+
+protected:
+
+protected:
+	TSharedPtr<class SRetainerWidget> MyRetainerWidget;
+};
diff --git a/Engine/Source/Runtime/UMG/Public/Slate/SRetainerWidget.h b/Engine/Source/Runtime/UMG/Public/Slate/SRetainerWidget.h
new file mode 100644
--- /dev/null
+++ b/Engine/Source/Runtime/UMG/Public/Slate/SRetainerWidget.h
@@ -0,0 +1,99 @@
+// Copyright 1998-2015 Epic Games, Inc. All Rights Reserved.
+
+#pragma once
+
+#include "Input/HittestGrid.h"
+#include "WidgetRenderer.h"
+
+/**
+ * 
+ */
+class UMG_API SRetainerWidget : public ISlateViewport, public SCompoundWidget, public FGCObject, public FTickableGameObject, public ICustomHitTestPath
+{
+public:
+	SLATE_BEGIN_ARGS(SRetainerWidget)
+	{
+		_Visibility = EVisibility::Visible;
+		_Phase = 0;
+		_PhaseCount = 1;
+	}
+	SLATE_DEFAULT_SLOT(FArguments, Content)
+		SLATE_ARGUMENT(int32, Phase)
+		SLATE_ARGUMENT(int32, PhaseCount)
+	SLATE_END_ARGS()
+
+	SRetainerWidget();
+	~SRetainerWidget();
+
+	void Construct(const FArguments& Args);
+
+	// ISlateViewport
+	using ISlateViewport::Tick;
+	virtual FIntPoint GetSize() const override;
+	virtual bool RequiresVsync() const override;
+	virtual FSlateShaderResource* GetViewportRenderTargetTexture() const override;
+	bool IsViewportTextureAlphaOnly() const override;
+	// ISlateViewport
+
+	void SetRetainedRendering(bool bRetainRendering);
+
+	void SetContent(const TSharedRef< SWidget >& InContent);
+
+	// FGCObject
+	virtual void AddReferencedObjects(FReferenceCollector& Collector) override;
+	// FGCObject
+
+	// SWidget
+	using SWidget::Tick;
+	virtual FChildren* GetChildren() override;
+	virtual bool ComputeVolatility() const override;
+	// SWidget
+
+	// FTickableGameObject
+	virtual void Tick(float DeltaTime) override;
+	virtual bool IsTickable() const override { return true; }
+	virtual bool IsTickableWhenPaused() const override { return true; }
+	virtual bool IsTickableInEditor() const override { return true; }
+	virtual TStatId GetStatId() const override { RETURN_QUICK_DECLARE_CYCLE_STAT(FRetainerWidget, STATGROUP_Tickables); }
+	// FTickableGameObject
+
+	// ICustomHitTestPath
+	virtual TArray<FWidgetAndPointer> GetBubblePathAndVirtualCursors(const FGeometry& InGeometry, FVector2D DesktopSpaceCoordinate, bool bIgnoreEnabledStatus) const override;
+	virtual void ArrangeChildren(FArrangedChildren& ArrangedChildren) const override;
+	virtual TSharedPtr<struct FVirtualPointerPosition> TranslateMouseCoordinateFor3DChild(const TSharedRef<SWidget>& ChildWidget, const FGeometry& ViewportGeometry, const FVector2D& ScreenSpaceMouseCoordinate, const FVector2D& LastScreenSpaceMouseCoordinate) const override;
+	// ICustomHitTestPath
+
+protected:
+	// BEGIN SLeafWidget interface
+	virtual int32 OnPaint(const FPaintArgs& Args, const FGeometry& AllottedGeometry, const FSlateRect& MyClippingRect, FSlateWindowElementList& OutDrawElements, int32 LayerId, const FWidgetStyle& InWidgetStyle, bool bParentEnabled) const override;
+	virtual FVector2D ComputeDesiredSize(float Scale) const override;
+	// END SLeafWidget interface
+
+	void RefreshRenderingMode();
+	bool ShouldBeRenderingOffscreen() const;
+
+private:
+	mutable FGeometry CachedAllottedGeometry;
+	mutable FSlateRect CachedClippingRect;
+
+	FSimpleSlot EmptyChildSlot;
+
+	FSlateBrush RenderTargetBrush;
+
+	mutable FWidgetRenderer WidgetRenderer;
+	mutable class UTextureRenderTarget2D* RenderTarget;
+	mutable class FSlateRenderTargetRHI* RenderTargetResource;
+	mutable TSharedPtr<SWidget> MyWidget;
+
+	bool bRenderingOffscreenDesire;
+	bool bRenderingOffscreen;
+
+	int32 Phase;
+	int32 PhaseCount;
+
+	double LastDrawTime;
+	int64 LastTickedFrame;
+
+	TSharedPtr<SVirtualWindow> Window;
+	TSharedPtr<FHittestGrid> HitTestGrid;
+};
\ No newline at end of file
diff --git a/Engine/Source/Runtime/UMG/Public/Slate/WidgetRenderer.h b/Engine/Source/Runtime/UMG/Public/Slate/WidgetRenderer.h
--- a/Engine/Source/Runtime/UMG/Public/Slate/WidgetRenderer.h
+++ b/Engine/Source/Runtime/UMG/Public/Slate/WidgetRenderer.h
@@ -23,19 +23,28 @@ public:
 /**
  * 
  */
-class FWidgetRenderer
+class UMG_API FWidgetRenderer
 {
 public:
-	FWidgetRenderer();
+	FWidgetRenderer(bool bUseGammaCorrection = false);
 	~FWidgetRenderer();
 
+	bool GetIsPrepassNeeded() const { return bPrepassNeeded; }
+	void SetIsPrepassNeeded(bool bInPrepassNeeded) { bPrepassNeeded = bInPrepassNeeded; }
+
 	ISlate3DRenderer* GetSlateRenderer();
 
-	UTextureRenderTarget2D* DrawWidget(TSharedRef<SWidget>& Widget, FVector2D DrawSize);
-	void DrawWidget(UTextureRenderTarget2D* RenderTarget, TSharedRef<SWidget>& Widget, FVector2D DrawSize, float DeltaTime);
-	void DrawWindow(UTextureRenderTarget2D* RenderTarget, TSharedRef<FHittestGrid> HitTestGrid, TSharedRef<SWindow> Window, FVector2D DrawSize, float DeltaTime);
+	static UTextureRenderTarget2D* CreateTargetFor(FVector2D DrawSize, TextureFilter InFilter, bool bUseGammaCorrection);
+
+	UTextureRenderTarget2D* DrawWidget(const TSharedRef<SWidget>& Widget, FVector2D DrawSize);
+	void DrawWidget(UTextureRenderTarget2D* RenderTarget, const TSharedRef<SWidget>& Widget, FVector2D DrawSize, float DeltaTime);
+	void DrawWindow(UTextureRenderTarget2D* RenderTarget, TSharedRef<FHittestGrid> HitTestGrid, TSharedRef<SWindow> Window, float Scale, FVector2D DrawSize, float DeltaTime);
 
 private:
 	/** The slate 3D renderer used to render the user slate widget */
 	TSharedPtr<ISlate3DRenderer> Renderer;
+	/** Prepass Needed when drawing the widget? */
+	bool bPrepassNeeded;
+	/** Is gamma space needed? */
+	bool bUseGammaSpace;
 };
diff --git a/Engine/Source/Runtime/Windows/D3D11RHI/Private/D3D11Commands.cpp b/Engine/Source/Runtime/Windows/D3D11RHI/Private/D3D11Commands.cpp
--- a/Engine/Source/Runtime/Windows/D3D11RHI/Private/D3D11Commands.cpp
+++ b/Engine/Source/Runtime/Windows/D3D11RHI/Private/D3D11Commands.cpp
@@ -252,6 +252,17 @@ void FD3D11DynamicRHI::RHISetBoundShaderState( FBoundShaderStateRHIParamRef Boun
 	DirtyUniformBuffers[SF_Hull] = 0xffff;
 	DirtyUniformBuffers[SF_Domain] = 0xffff;
 	DirtyUniformBuffers[SF_Geometry] = 0xffff;
+
+	// Shader changed.  All UB's must be reset by high level code to match other platforms anway.
+	// Clear to catch those bugs, and bugs with stale UB's causing layout mismatches.
+	// Release references to bound uniform buffers.
+	for (int32 Frequency = 0; Frequency < SF_NumFrequencies; ++Frequency)
+	{
+		for (int32 BindIndex = 0; BindIndex < MAX_UNIFORM_BUFFERS_PER_SHADER_STAGE; ++BindIndex)
+		{
+			BoundUniformBuffers[Frequency][BindIndex].SafeRelease();
+		}
+	}
 }
 
 void FD3D11DynamicRHI::RHISetShaderTexture(FVertexShaderRHIParamRef VertexShaderRHI,uint32 TextureIndex,FTextureRHIParamRef NewTextureRHI)
