Copying //UE4/Orion-Staging to Dev-Main (//UE4/Dev-Main)

diff --git a/Engine/Source/Runtime/UMG/Public/Components/Widget.h b/Engine/Source/Runtime/UMG/Public/Components/Widget.h
--- a/Engine/Source/Runtime/UMG/Public/Components/Widget.h
+++ b/Engine/Source/Runtime/UMG/Public/Components/Widget.h
@@ -424,6 +424,15 @@ public:
 		return ( DesignerFlags&FlagToCheck ) != 0;
 	}
 
+	/** @return The friendly name of the widget to display in the editor */
+	const FString& GetDisplayLabel() const
+	{
+		return DisplayLabel;
+	}
+
+	/** Sets the friendly name of the widget to display in the editor */
+	void SetDisplayLabel(const FString& DisplayLabel);
+
 #else
 	FORCEINLINE bool IsDesignTime() const { return false; }
 #endif
@@ -573,6 +582,10 @@ private:
 	/** Any flags used by the designer at edit time. */
 	UPROPERTY(Transient)
 	TEnumAsByte<EWidgetDesignFlags::Type> DesignerFlags;
+
+	/** The friendly name for this widget displayed in the designer and BP graph. */
+	UPROPERTY()
+	FString DisplayLabel;
 #endif
 
 #if !(UE_BUILD_SHIPPING || UE_BUILD_TEST)
diff --git a/Engine/Source/Runtime/UMG/Public/Slate/SMeshWidget.h b/Engine/Source/Runtime/UMG/Public/Slate/SMeshWidget.h
--- a/Engine/Source/Runtime/UMG/Public/Slate/SMeshWidget.h
+++ b/Engine/Source/Runtime/UMG/Public/Slate/SMeshWidget.h
@@ -10,7 +10,7 @@ class USlateVectorArtData;
  * The Mesh's material is used.
  * Hardware instancing is supported.
  */
-class UMG_API SMeshWidget : public SLeafWidget
+class UMG_API SMeshWidget : public SLeafWidget, public FGCObject
 {
 public:
 	SLATE_BEGIN_ARGS(SMeshWidget)
@@ -45,8 +45,12 @@ public:
 		RenderRuns.Add(FRenderRun(InMeshIndex, InInstanceOffset, InNumInstances));
 	}
 
-	/** Begin an update to the per instance buffer. Enables hardware instancing. */
+	/** Enable hardware instancing */
+	void EnableInstancing(uint32 MeshId, int32 InitialSize);
+	/** Begin an update to the per instance buffer. Automatically enables hardware instancing. */
 	TSharedPtr<FSlateInstanceBufferUpdate> BeginPerInstanceBufferUpdate(uint32 MeshId, int32 InitialSize);
+	/** Begin an update to the per instance buffer. Use */
+	TSharedPtr<FSlateInstanceBufferUpdate> BeginPerInstanceBufferUpdateConst(uint32 MeshId) const;
 
 protected:
 	// BEGIN SLeafWidget interface
@@ -54,6 +58,10 @@ protected:
 	virtual FVector2D ComputeDesiredSize(float) const override;
 	// END SLeafWidget interface
 
+	// ~ FGCObject
+	virtual void AddReferencedObjects(FReferenceCollector& Collector) override;
+	// ~ FGCObject
+
 private:
 
 	struct FRenderData
diff --git a/Engine/Source/Runtime/UMG/Public/Slate/SlateVectorArtData.h b/Engine/Source/Runtime/UMG/Public/Slate/SlateVectorArtData.h
--- a/Engine/Source/Runtime/UMG/Public/Slate/SlateVectorArtData.h
+++ b/Engine/Source/Runtime/UMG/Public/Slate/SlateVectorArtData.h
@@ -88,6 +88,12 @@ public:
 	/** Convert the static mesh data into slate vector art on demand. Does nothing in a cooked build. */
 	void EnsureValidData();
 
+	FVector2D GetDesiredSize() const;
+
+	FVector2D GetExtentMin() const;
+
+	FVector2D GetExtentMax() const;
+
 private:
 	// ~ UObject Interface
 	virtual void PreSave() override;
@@ -100,6 +106,10 @@ private:
 	/** The mesh data asset from which the vector art is sourced */
 	UPROPERTY(EditAnywhere, Category="Vector Art" )
 	UStaticMesh* MeshAsset;
+
+	/** The material which we are using, or the material from with the MIC was constructed. */
+	UPROPERTY(Transient)
+	UMaterialInterface* SourceMaterial;
 #endif
 
 	/** @see GetVertexData() */
@@ -113,4 +123,10 @@ private:
 	/** @see GetMaterial() */
 	UPROPERTY()
 	UMaterialInterface* Material;
+
+	UPROPERTY()
+	FVector2D ExtentMin;
+
+	UPROPERTY()
+	FVector2D ExtentMax;
 };
\ No newline at end of file
diff --git a/Engine/Source/Runtime/UMG/Public/Slate/SlateVectorArtInstanceData.h b/Engine/Source/Runtime/UMG/Public/Slate/SlateVectorArtInstanceData.h
new file mode 100644
--- /dev/null
+++ b/Engine/Source/Runtime/UMG/Public/Slate/SlateVectorArtInstanceData.h
@@ -0,0 +1,67 @@
+// Copyright 1998-2015 Epic Games, Inc. All Rights Reserved.
+
+#pragma once
+
+
+struct UMG_API FSlateVectorArtInstanceData
+{
+public:
+	const FVector4& GetData() const { return Data; }
+	FVector4& GetData() { return Data; }
+
+	/**
+	 * A Float4 can pack 12 bytes worth of data (we cannot safely use the exponent bits)
+	 */
+	template<int32 Component, int32 ByteIndex>
+	void PackFloatIntoByte(float InValue)
+	{
+		PackByteIntoByte<Component, ByteIndex>(static_cast<uint8>(FMath::RoundToInt(InValue * 0xFFu)));
+	}
+	
+	template<int32 Component, int32 ByteIndex>
+	void PackByteIntoByte(uint8 InValue)
+	{
+		// Each float has 24 usable bits of mantissa, but we cannot access the bits directly.
+		// We will not respect the IEEE "normalized mantissa" rules, so let 
+		// the compiler/FPU do conversions from byte to float and vice versa for us.
+
+		//Produces a value like 0xFFFF00FF; where the 00s are shifted by `Position` bytes.
+		static const uint32 Mask = ~(0xFF << ByteIndex * 8);
+
+		// Clear out a byte to OR with while maintaining the rest of the data intact.
+		uint32 CurrentData = static_cast<uint32>(Data[Component]) & Mask;
+		// OR in the value
+		Data[Component] = CurrentData | (InValue << ByteIndex*8);
+	}
+
+	static float AdjustScaleInput(float UnadjustedScaleInput)
+	{
+		// We only support scaling UP (initial "vector" art is at 1:1 resolution)
+		// We also don't want to waste any bits on supporting fractional scale values.
+		// So store a value between 0 and 3, with an assumed base scale of 1.0f
+		static const float MAXIMUM_SCALE = 4.0f;
+		return FMath::Clamp(UnadjustedScaleInput, 0.0f, MAXIMUM_SCALE) / MAXIMUM_SCALE;
+	}
+
+	template<int32 Component, int32 ByteIndex>
+	void PackByteIntoByte(float InValue)
+	{
+		checkf(false, TEXT("PackByteIntoByte() being used with a float value. This is almost definitely an error. Use PackFloatIntoByte()."));
+	}
+
+	template<int32 Component, int32 ByteIndex>
+	float UnpackValue() const
+	{
+		// See PackFloatIntoByte
+		static const float Denominator = static_cast<float>(0x1 << (ByteIndex+1)*8);
+		return Data[Component] / Denominator;
+	}
+
+	void SetPositionFixedPoint16(FVector2D Position);
+
+	void SetPosition(FVector2D Position);
+	void SetScale(float Scale);
+
+protected:
+	FVector4 Data;
+};
\ No newline at end of file
