Fix shadowed variable

diff --git a/Engine/Source/Runtime/UMG/Private/Slate/SlateDataSheet.cpp b/Engine/Source/Runtime/UMG/Private/Slate/SlateDataSheet.cpp
--- a/Engine/Source/Runtime/UMG/Private/Slate/SlateDataSheet.cpp
+++ b/Engine/Source/Runtime/UMG/Private/Slate/SlateDataSheet.cpp
@@ -47,9 +47,9 @@ void USlateDataSheet::EnqueueUpdateToGPU()
 		uint8* SrcData = (uint8*)FMemory::Malloc(DataSize);
 		FMemory::Memcpy(SrcData, Data, DataSize);
 
-		auto DataCleanup = [](uint8* SrcData, const FUpdateTextureRegion2D* Regions)
+		auto DataCleanup = [](uint8* InSrcData, const FUpdateTextureRegion2D* Regions)
 		{
-			FMemory::Free(SrcData);
+			FMemory::Free(InSrcData);
 		};
 		DataTexture->UpdateTextureRegions(0, 1, &DataSheetUpdateRegion, DataPitch, Data_PixelSize, SrcData, DataCleanup);
 	}
