Copying //UE4/Orion-Staging to Dev-Main (Originating from //Orion/Main at CL-2777663)

diff --git a/Engine/Source/Runtime/UMG/Private/Slate/SRetainerWidget.cpp b/Engine/Source/Runtime/UMG/Private/Slate/SRetainerWidget.cpp
--- a/Engine/Source/Runtime/UMG/Private/Slate/SRetainerWidget.cpp
+++ b/Engine/Source/Runtime/UMG/Private/Slate/SRetainerWidget.cpp
@@ -6,7 +6,10 @@
 #include "SRetainerWidget.h"
 #include "Runtime/SlateRHIRenderer/Public/Interfaces/ISlateRHIRendererModule.h"
 
-#if !(UE_BUILD_SHIPPING || UE_BUILD_TEST)
+DECLARE_CYCLE_STAT(TEXT("Retainer Widget Tick"), STAT_SlateRetainerWidgetTick, STATGROUP_Slate);
+DECLARE_CYCLE_STAT(TEXT("Retainer Widget Paint"), STAT_SlateRetainerWidgetPaint, STATGROUP_Slate);
+
+#if !UE_BUILD_SHIPPING
 
 /** True if we should allow widgets to be cached in the UI at all. */
 TAutoConsoleVariable<int32> EnableRetainedRendering(
@@ -19,6 +22,7 @@ static bool IsRetainedRenderingEnabled()
 	return EnableRetainedRendering.GetValueOnGameThread() == 1;
 }
 
+FOnRetainedModeChanged SRetainerWidget::OnRetainerModeChangedDelegate;
 #else
 
 static bool IsRetainedRenderingEnabled()
@@ -77,10 +81,38 @@ SRetainerWidget::~SRetainerWidget()
 		{
 			delete InRenderTargetRHI;
 		});
+
+	if( FSlateApplication::IsInitialized() )
+	{
+		FSlateApplication::Get().OnUpdateRetainerWidgets().RemoveAll( this );
+
+
+#if !UE_BUILD_SHIPPING
+		OnRetainerModeChangedDelegate.RemoveAll( this );
+#endif
+	}
+	
 }
 
 void SRetainerWidget::Construct(const FArguments& InArgs)
 {
+	if( FSlateApplication::IsInitialized() )
+	{
+		FSlateApplication::Get().OnUpdateRetainerWidgets().AddRaw( this, &SRetainerWidget::OnTickRetainers );
+
+#if !UE_BUILD_SHIPPING
+		OnRetainerModeChangedDelegate.AddRaw( this, &SRetainerWidget::OnRetainerModeChanged );
+
+		static bool bStaticInit = false;
+
+		if( !bStaticInit )
+		{
+			bStaticInit = true;
+			EnableRetainedRendering.AsVariable()->SetOnChangedCallback( FConsoleVariableDelegate::CreateStatic( &SRetainerWidget::OnRetainerModeCVarChanged ) );
+		}
+#endif
+	}
+
 	RenderTarget = NewObject<UTextureRenderTarget2D>();
 	RenderTarget->SRGB = true;
 	RenderTarget->TargetGamma = 1;
@@ -142,6 +174,21 @@ bool SRetainerWidget::ShouldBeRenderingOffscreen() const
 	return bRenderingOffscreenDesire && IsRetainedRenderingEnabled();
 }
 
+void SRetainerWidget::OnRetainerModeChanged()
+{
+	RefreshRenderingMode();
+	Invalidate(EInvalidateWidget::Layout);
+}
+
+#if !UE_BUILD_SHIPPING
+
+void SRetainerWidget::OnRetainerModeCVarChanged( IConsoleVariable* CVar )
+{
+	OnRetainerModeChangedDelegate.Broadcast();
+}
+
+#endif
+
 void SRetainerWidget::SetRetainedRendering(bool bRetainRendering)
 {
 	bRenderingOffscreenDesire = bRetainRendering;
@@ -207,13 +254,14 @@ FChildren* SRetainerWidget::GetChildren()
 
 bool SRetainerWidget::ComputeVolatility() const
 {
-	return true;
+	return true;// return SCompoundWidget::ComputeVolatility();
 }
 
-void SRetainerWidget::Tick(float DeltaTime)
+void SRetainerWidget::OnTickRetainers(float DeltaTime)
 {
-	if ( bRenderingOffscreen )
+	if ( bRenderingOffscreen && GetVisibility().IsVisible() && ChildSlot.GetWidget()->GetVisibility().IsVisible() )
 	{
+		SCOPE_CYCLE_COUNTER( STAT_SlateRetainerWidgetTick );
 		if ( LastTickedFrame != GFrameCounter && ( GFrameCounter % PhaseCount ) == Phase )
 		{
 			LastTickedFrame = GFrameCounter;
@@ -276,8 +324,9 @@ int32 SRetainerWidget::OnPaint(const FPaintArgs& Args, const FGeometry& Allotted
 
 	MutableThis->RefreshRenderingMode();
 
-	if ( bRenderingOffscreen )
+	if ( bRenderingOffscreen && GetVisibility().IsVisible() && ChildSlot.GetWidget()->GetVisibility().IsVisible() )
 	{
+		SCOPE_CYCLE_COUNTER( STAT_SlateRetainerWidgetPaint );
 		CachedAllottedGeometry = AllottedGeometry;
 		CachedClippingRect = MyClippingRect;
 
diff --git a/Engine/Source/Runtime/UMG/Private/Slate/WidgetRenderer.cpp b/Engine/Source/Runtime/UMG/Private/Slate/WidgetRenderer.cpp
--- a/Engine/Source/Runtime/UMG/Private/Slate/WidgetRenderer.cpp
+++ b/Engine/Source/Runtime/UMG/Private/Slate/WidgetRenderer.cpp
@@ -78,6 +78,11 @@ void FWidgetRenderer::DrawWindow(
 #if !UE_SERVER
 	FGeometry WindowGeometry = FGeometry::MakeRoot(DrawSize * (1 / Scale), FSlateLayoutTransform(Scale));
 
+	if ( GUsingNullRHI )
+	{
+		return;
+	}
+	
 	if ( !bFoldTick )
 	{
 		Window->TickWidgetsRecursively(WindowGeometry, FApp::GetCurrentTime(), DeltaTime);
diff --git a/Engine/Source/Runtime/UMG/Private/UserWidget.cpp b/Engine/Source/Runtime/UMG/Private/UserWidget.cpp
--- a/Engine/Source/Runtime/UMG/Private/UserWidget.cpp
+++ b/Engine/Source/Runtime/UMG/Private/UserWidget.cpp
@@ -11,6 +11,23 @@
 
 #define LOCTEXT_NAMESPACE "UMG"
 
+
+static FGeometry NullGeometry;
+static FSlateRect NullRect;
+static FSlateWindowElementList NullElementList;
+static FWidgetStyle NullStyle;
+
+FPaintContext::FPaintContext()
+		: AllottedGeometry(NullGeometry)
+		, MyClippingRect(NullRect)
+		, OutDrawElements(NullElementList)
+		, LayerId(0)
+		, WidgetStyle(NullStyle)
+		, bParentEnabled(true)
+		, MaxLayer(0)
+{
+}
+
 /////////////////////////////////////////////////////
 // UUserWidget
 
diff --git a/Engine/Source/Runtime/UMG/Public/Animation/MovieScene2DTransformSection.h b/Engine/Source/Runtime/UMG/Public/Animation/MovieScene2DTransformSection.h
--- a/Engine/Source/Runtime/UMG/Public/Animation/MovieScene2DTransformSection.h
+++ b/Engine/Source/Runtime/UMG/Public/Animation/MovieScene2DTransformSection.h
@@ -23,18 +23,25 @@ enum class EKey2DTransformAxis
 	None
 };
 
+enum class EKey2DTransformValueType
+{
+	Key,
+	Default
+};
 
 struct F2DTransformKey
 {
-	F2DTransformKey( EKey2DTransformChannel InChannel, EKey2DTransformAxis InAxis, float InValue )
+	F2DTransformKey( EKey2DTransformChannel InChannel, EKey2DTransformAxis InAxis, float InValue, EKey2DTransformValueType InValueType )
 	{
 		Channel = InChannel;
 		Axis = InAxis;
 		Value = InValue;
+		ValueType = InValueType;
 	}
 	EKey2DTransformChannel Channel;
 	EKey2DTransformAxis Axis;
 	float Value;
+	EKey2DTransformValueType ValueType;	
 };
 
 
diff --git a/Engine/Source/Runtime/UMG/Public/Animation/MovieSceneMarginSection.h b/Engine/Source/Runtime/UMG/Public/Animation/MovieSceneMarginSection.h
--- a/Engine/Source/Runtime/UMG/Public/Animation/MovieSceneMarginSection.h
+++ b/Engine/Source/Runtime/UMG/Public/Animation/MovieSceneMarginSection.h
@@ -15,16 +15,24 @@ enum class EKeyMarginChannel
 	Bottom
 };
 
+enum class EKeyMarginValueType
+{
+	Key,
+	Default
+};
+
 
 struct FMarginKey
 {
-	FMarginKey( EKeyMarginChannel InChannel, float InValue )
+	FMarginKey( EKeyMarginChannel InChannel, float InValue, EKeyMarginValueType InValueType )
 	{
 		Channel = InChannel;
 		Value = InValue;
+		ValueType = InValueType;
 	}
 	EKeyMarginChannel Channel;
 	float Value;
+	EKeyMarginValueType ValueType;
 };
 
 
diff --git a/Engine/Source/Runtime/UMG/Public/Animation/WidgetAnimation.h b/Engine/Source/Runtime/UMG/Public/Animation/WidgetAnimation.h
--- a/Engine/Source/Runtime/UMG/Public/Animation/WidgetAnimation.h
+++ b/Engine/Source/Runtime/UMG/Public/Animation/WidgetAnimation.h
@@ -66,18 +66,13 @@ public:
 
 	// UMovieSceneAnimation overrides
 
-	virtual void BindPossessableObject(const FGuid& ObjectId, UObject& PossessedObject) override;
+	virtual void BindPossessableObject(const FGuid& ObjectId, UObject& PossessedObject, UObject* Context) override;
 	virtual bool CanPossessObject(UObject& Object) const override;
-	virtual UObject* FindObject(const FGuid& ObjectId) const override;
-	virtual FGuid FindObjectId(UObject& Object) const override;
+	virtual UObject* FindPossessableObject(const FGuid& ObjectId, UObject* Context) const override;
 	virtual UMovieScene* GetMovieScene() const override;
 	virtual UObject* GetParentObject(UObject* Object) const override;
 	virtual void UnbindPossessableObjects(const FGuid& ObjectId) override;
 
-#if WITH_EDITOR
-	virtual bool TryGetObjectDisplayName(const FGuid& ObjectId, FText& OutDisplayName) const override;
-#endif
-
 	const TArray<FWidgetAnimationBinding>& GetBindings() const { return AnimationBindings; }
 public:
 
diff --git a/Engine/Source/Runtime/UMG/Public/Animation/WidgetAnimationBinding.h b/Engine/Source/Runtime/UMG/Public/Animation/WidgetAnimationBinding.h
--- a/Engine/Source/Runtime/UMG/Public/Animation/WidgetAnimationBinding.h
+++ b/Engine/Source/Runtime/UMG/Public/Animation/WidgetAnimationBinding.h
@@ -2,9 +2,9 @@
 
 #pragma once
 
+#include "WidgetTree.h"
 #include "WidgetAnimationBinding.generated.h"
 
-
 /**
  * A single object bound to a UMG sequence.
  */
diff --git a/Engine/Source/Runtime/UMG/Public/Blueprint/UserWidget.h b/Engine/Source/Runtime/UMG/Public/Blueprint/UserWidget.h
--- a/Engine/Source/Runtime/UMG/Public/Blueprint/UserWidget.h
+++ b/Engine/Source/Runtime/UMG/Public/Blueprint/UserWidget.h
@@ -13,11 +13,6 @@
 
 #include "UserWidget.generated.h"
 
-static FGeometry NullGeometry;
-static FSlateRect NullRect;
-static FSlateWindowElementList NullElementList;
-static FWidgetStyle NullStyle;
-
 class UWidget;
 class UWidgetAnimation;
 struct FLocalPlayerContext;
@@ -34,15 +29,7 @@ struct UMG_API FPaintContext
 public:
 
 	/** Don't ever use this constructor.  Needed for code generation. */
-	FPaintContext()
-		: AllottedGeometry(NullGeometry)
-		, MyClippingRect(NullRect)
-		, OutDrawElements(NullElementList)
-		, LayerId(0)
-		, WidgetStyle(NullStyle)
-		, bParentEnabled(true)
-		, MaxLayer(0)
-	{ }
+	FPaintContext();
 
 	FPaintContext(const FGeometry& InAllottedGeometry, const FSlateRect& InMyClippingRect, FSlateWindowElementList& InOutDrawElements, const int32 InLayerId, const FWidgetStyle& InWidgetStyle, const bool bInParentEnabled)
 		: AllottedGeometry(InAllottedGeometry)
diff --git a/Engine/Source/Runtime/UMG/Public/Slate/SMeshWidget.h b/Engine/Source/Runtime/UMG/Public/Slate/SMeshWidget.h
--- a/Engine/Source/Runtime/UMG/Public/Slate/SMeshWidget.h
+++ b/Engine/Source/Runtime/UMG/Public/Slate/SMeshWidget.h
@@ -62,6 +62,9 @@ protected:
 	virtual void AddReferencedObjects(FReferenceCollector& Collector) override;
 	// ~ FGCObject
 
+protected:
+	static void PushUpdate(uint32 VectorArtId, const SMeshWidget& Widget, const FVector2D& Position, float Scale, uint32 BaseAddress);
+
 private:
 
 	struct FRenderData
diff --git a/Engine/Source/Runtime/UMG/Public/Slate/SRetainerWidget.h b/Engine/Source/Runtime/UMG/Public/Slate/SRetainerWidget.h
--- a/Engine/Source/Runtime/UMG/Public/Slate/SRetainerWidget.h
+++ b/Engine/Source/Runtime/UMG/Public/Slate/SRetainerWidget.h
@@ -5,10 +5,13 @@
 #include "Input/HittestGrid.h"
 #include "WidgetRenderer.h"
 
+
+DECLARE_MULTICAST_DELEGATE( FOnRetainedModeChanged );
+
 /**
  * 
  */
-class UMG_API SRetainerWidget : public ISlateViewport, public SCompoundWidget, public FGCObject, public FTickableGameObject, public ICustomHitTestPath
+class UMG_API SRetainerWidget : public ISlateViewport, public SCompoundWidget, public FGCObject, public ICustomHitTestPath
 {
 public:
 	SLATE_BEGIN_ARGS(SRetainerWidget)
@@ -49,13 +52,7 @@ public:
 	virtual bool ComputeVolatility() const override;
 	// SWidget
 
-	// FTickableGameObject
-	virtual void Tick(float DeltaTime) override;
-	virtual bool IsTickable() const override { return true; }
-	virtual bool IsTickableWhenPaused() const override { return true; }
-	virtual bool IsTickableInEditor() const override { return true; }
-	virtual TStatId GetStatId() const override { RETURN_QUICK_DECLARE_CYCLE_STAT(FRetainerWidget, STATGROUP_Tickables); }
-	// FTickableGameObject
+	virtual void OnTickRetainers(float DeltaTime);
 
 	// ICustomHitTestPath
 	virtual TArray<FWidgetAndPointer> GetBubblePathAndVirtualCursors(const FGeometry& InGeometry, FVector2D DesktopSpaceCoordinate, bool bIgnoreEnabledStatus) const override;
@@ -71,8 +68,13 @@ protected:
 
 	void RefreshRenderingMode();
 	bool ShouldBeRenderingOffscreen() const;
-
+	void OnRetainerModeChanged();
 private:
+#if !UE_BUILD_SHIPPING
+	static void OnRetainerModeCVarChanged( IConsoleVariable* CVar );
+	static FOnRetainedModeChanged OnRetainerModeChangedDelegate;
+#endif
+
 	mutable FGeometry CachedAllottedGeometry;
 	mutable FSlateRect CachedClippingRect;
 
