Copying //UE4/Dev-Sequencer to Dev-Main (//UE4/Dev-Main)

diff --git a/Engine/Source/Runtime/UMG/Private/UMGStyle.cpp b/Engine/Source/Runtime/UMG/Private/UMGStyle.cpp
--- a/Engine/Source/Runtime/UMG/Private/UMGStyle.cpp
+++ b/Engine/Source/Runtime/UMG/Private/UMGStyle.cpp
@@ -46,43 +46,48 @@ TSharedRef< FSlateStyleSet > FUMGStyle::Create()
 	
 	Style->Set("MarchingAnts", new BORDER_BRUSH( TEXT("NonMarchingAnts"), FMargin(0.25f), FLinearColor(1,1,1,0.5) ));
 
-	Style->Set("Widget", new IMAGE_BRUSH(TEXT("Widget"), Icon16x16));
-	Style->Set("Widget.CheckBox", new IMAGE_BRUSH(TEXT("CheckBox"), Icon16x16));
-	Style->Set("Widget.Button", new IMAGE_BRUSH(TEXT("Button"), Icon16x16));
-	Style->Set("Widget.EditableTextBox", new IMAGE_BRUSH(TEXT("EditableTextBox"), Icon16x16));
-	Style->Set("Widget.EditableText", new IMAGE_BRUSH(TEXT("EditableText"), Icon16x16));
-	Style->Set("Widget.HorizontalBox", new IMAGE_BRUSH(TEXT("HorizontalBox"), Icon16x16));
-	Style->Set("Widget.VerticalBox", new IMAGE_BRUSH(TEXT("VerticalBox"), Icon16x16));
-	Style->Set("Widget.Image", new IMAGE_BRUSH(TEXT("Image"), Icon16x16));
-	Style->Set("Widget.Canvas", new IMAGE_BRUSH(TEXT("Canvas"), Icon16x16));
-	Style->Set("Widget.TextBlock", new IMAGE_BRUSH(TEXT("TextBlock"), Icon16x16));
-	Style->Set("Widget.Border", new IMAGE_BRUSH(TEXT("Border"), Icon16x16));
-	Style->Set("Widget.Slider", new IMAGE_BRUSH(TEXT("Slider"), Icon16x16));
-	Style->Set("Widget.Spacer", new IMAGE_BRUSH(TEXT("Spacer"), Icon16x16));
-	Style->Set("Widget.ScrollBox", new IMAGE_BRUSH(TEXT("ScrollBox"), Icon16x16));
-	Style->Set("Widget.ProgressBar", new IMAGE_BRUSH(TEXT("ProgressBar"), Icon16x16));
-	Style->Set("Widget.MenuAnchor", new IMAGE_BRUSH(TEXT("MenuAnchor"), Icon16x16));
-	Style->Set("Widget.ScrollBar", new IMAGE_BRUSH(TEXT("ScrollBar"), Icon16x16));
-	Style->Set("Widget.UniformGrid", new IMAGE_BRUSH(TEXT("UniformGrid"), Icon16x16));
-	Style->Set("Widget.WidgetSwitcher", new IMAGE_BRUSH(TEXT("WidgetSwitcher"), Icon16x16));
-	Style->Set("Widget.MultiLineEditableText", new IMAGE_BRUSH(TEXT("MultiLineEditableText"), Icon16x16));
-	Style->Set("Widget.MultiLineEditableTextBox", new IMAGE_BRUSH(TEXT("MultiLineEditableTextBox"), Icon16x16));
-	Style->Set("Widget.Viewport", new IMAGE_BRUSH(TEXT("Viewport"), Icon16x16));
-	Style->Set("Widget.ComboBox", new IMAGE_BRUSH(TEXT("icon_umg_ComboBox_16x"), Icon16x16));
-	Style->Set("Widget.ListView", new IMAGE_BRUSH(TEXT("icon_umg_ListView_16x"), Icon16x16));
-	Style->Set("Widget.TileView", new IMAGE_BRUSH(TEXT("icon_umg_TileView_16x"), Icon16x16));
-	Style->Set("Widget.Overlay", new IMAGE_BRUSH(TEXT("icon_umg_Overlay_16x"), Icon16x16));
-	Style->Set("Widget.Throbber", new IMAGE_BRUSH(TEXT("icon_umg_ThrobberHorizontal_16x"), Icon16x16));
-	Style->Set("Widget.CircularThrobber", new IMAGE_BRUSH(TEXT("icon_umg_ThrobberA_16x"), Icon16x16));
-	Style->Set("Widget.NativeWidgetHost", new IMAGE_BRUSH(TEXT("NativeWidgetHost"), Icon16x16));
-	Style->Set("Widget.ScaleBox", new IMAGE_BRUSH(TEXT("ScaleBox"), Icon16x16));
-	Style->Set("Widget.SizeBox", new IMAGE_BRUSH(TEXT("SizeBox"), Icon16x16));
-	Style->Set("Widget.SpinBox", new IMAGE_BRUSH(TEXT("SpinBox"), Icon16x16));
-	Style->Set("Widget.Grid", new IMAGE_BRUSH(TEXT("Grid"), Icon16x16));
-	Style->Set("Widget.WrapBox", new IMAGE_BRUSH(TEXT("WrapBox"), Icon16x16));
-	Style->Set("Widget.NamedSlot", new IMAGE_BRUSH(TEXT("NamedSlot"), Icon16x16));
-
-	Style->Set("Widget.UserWidget", new IMAGE_BRUSH(TEXT("UserWidget"), Icon16x16));
+	Style->Set("ClassIcon.Widget", new IMAGE_BRUSH(TEXT("Widget"), Icon16x16));
+	Style->Set("ClassIcon.CheckBox", new IMAGE_BRUSH(TEXT("CheckBox"), Icon16x16));
+	Style->Set("ClassIcon.Button", new IMAGE_BRUSH(TEXT("Button"), Icon16x16));
+	Style->Set("ClassIcon.EditableTextBox", new IMAGE_BRUSH(TEXT("EditableTextBox"), Icon16x16));
+	Style->Set("ClassIcon.EditableText", new IMAGE_BRUSH(TEXT("EditableText"), Icon16x16));
+	Style->Set("ClassIcon.HorizontalBox", new IMAGE_BRUSH(TEXT("HorizontalBox"), Icon16x16));
+	Style->Set("ClassIcon.VerticalBox", new IMAGE_BRUSH(TEXT("VerticalBox"), Icon16x16));
+	Style->Set("ClassIcon.Image", new IMAGE_BRUSH(TEXT("Image"), Icon16x16));
+	Style->Set("ClassIcon.CanvasPanel", new IMAGE_BRUSH(TEXT("Canvas"), Icon16x16));
+	Style->Set("ClassIcon.TextBlock", new IMAGE_BRUSH(TEXT("TextBlock"), Icon16x16));
+	Style->Set("ClassIcon.Border", new IMAGE_BRUSH(TEXT("Border"), Icon16x16));
+	Style->Set("ClassIcon.Slider", new IMAGE_BRUSH(TEXT("Slider"), Icon16x16));
+	Style->Set("ClassIcon.Spacer", new IMAGE_BRUSH(TEXT("Spacer"), Icon16x16));
+	Style->Set("ClassIcon.ExpandableArea", new IMAGE_BRUSH(TEXT("Spacer"), Icon16x16));
+	Style->Set("ClassIcon.ScrollBox", new IMAGE_BRUSH(TEXT("ScrollBox"), Icon16x16));
+	Style->Set("ClassIcon.ProgressBar", new IMAGE_BRUSH(TEXT("ProgressBar"), Icon16x16));
+	Style->Set("ClassIcon.MenuAnchor", new IMAGE_BRUSH(TEXT("MenuAnchor"), Icon16x16));
+	Style->Set("ClassIcon.InvalidationBox", new IMAGE_BRUSH(TEXT("MenuAnchor"), Icon16x16));
+	Style->Set("ClassIcon.RetainerBox", new IMAGE_BRUSH(TEXT("MenuAnchor"), Icon16x16));
+	Style->Set("ClassIcon.ScrollBar", new IMAGE_BRUSH(TEXT("ScrollBar"), Icon16x16));
+	Style->Set("ClassIcon.UniformGridPanel", new IMAGE_BRUSH(TEXT("UniformGrid"), Icon16x16));
+	Style->Set("ClassIcon.WidgetSwitcher", new IMAGE_BRUSH(TEXT("WidgetSwitcher"), Icon16x16));
+	Style->Set("ClassIcon.MultiLineEditableText", new IMAGE_BRUSH(TEXT("MultiLineEditableText"), Icon16x16));
+	Style->Set("ClassIcon.MultiLineEditableTextBox", new IMAGE_BRUSH(TEXT("MultiLineEditableTextBox"), Icon16x16));
+	Style->Set("ClassIcon.Viewport", new IMAGE_BRUSH(TEXT("Viewport"), Icon16x16));
+	Style->Set("ClassIcon.ComboBox", new IMAGE_BRUSH(TEXT("icon_umg_ComboBox_16x"), Icon16x16));
+	Style->Set("ClassIcon.ComboBoxString", new IMAGE_BRUSH(TEXT("icon_umg_ComboBox_16x"), Icon16x16));
+	Style->Set("ClassIcon.ListView", new IMAGE_BRUSH(TEXT("icon_umg_ListView_16x"), Icon16x16));
+	Style->Set("ClassIcon.TileView", new IMAGE_BRUSH(TEXT("icon_umg_TileView_16x"), Icon16x16));
+	Style->Set("ClassIcon.Overlay", new IMAGE_BRUSH(TEXT("icon_umg_Overlay_16x"), Icon16x16));
+	Style->Set("ClassIcon.Throbber", new IMAGE_BRUSH(TEXT("icon_umg_ThrobberHorizontal_16x"), Icon16x16));
+	Style->Set("ClassIcon.CircularThrobber", new IMAGE_BRUSH(TEXT("icon_umg_ThrobberA_16x"), Icon16x16));
+	Style->Set("ClassIcon.NativeWidgetHost", new IMAGE_BRUSH(TEXT("NativeWidgetHost"), Icon16x16));
+	Style->Set("ClassIcon.ScaleBox", new IMAGE_BRUSH(TEXT("ScaleBox"), Icon16x16));
+	Style->Set("ClassIcon.SizeBox", new IMAGE_BRUSH(TEXT("SizeBox"), Icon16x16));
+	Style->Set("ClassIcon.SafeZone", new IMAGE_BRUSH(TEXT("SafeZone"), Icon16x16));
+	Style->Set("ClassIcon.SpinBox", new IMAGE_BRUSH(TEXT("SpinBox"), Icon16x16));
+	Style->Set("ClassIcon.GridPanel", new IMAGE_BRUSH(TEXT("Grid"), Icon16x16));
+	Style->Set("ClassIcon.WrapBox", new IMAGE_BRUSH(TEXT("WrapBox"), Icon16x16));
+	Style->Set("ClassIcon.NamedSlot", new IMAGE_BRUSH(TEXT("NamedSlot"), Icon16x16));
+
+	Style->Set("ClassIcon.UserWidget", new IMAGE_BRUSH(TEXT("UserWidget"), Icon16x16));
 
 	Style->Set("Animations.TabIcon", new IMAGE_BRUSH(TEXT("Animations_16x"), Icon16x16));
 	Style->Set("Designer.TabIcon", new IMAGE_BRUSH(TEXT("Designer_16x"), Icon16x16));
diff --git a/Engine/Source/Runtime/UMG/Private/UserWidget.cpp b/Engine/Source/Runtime/UMG/Private/UserWidget.cpp
--- a/Engine/Source/Runtime/UMG/Private/UserWidget.cpp
+++ b/Engine/Source/Runtime/UMG/Private/UserWidget.cpp
@@ -799,11 +799,6 @@ void UUserWidget::PreSave()
 
 #if WITH_EDITOR
 
-const FSlateBrush* UUserWidget::GetEditorIcon()
-{
-	return FUMGStyle::Get().GetBrush("Widget.UserWidget");
-}
-
 const FText UUserWidget::GetPaletteCategory()
 {
 	return PaletteCategory;
diff --git a/Engine/Source/Runtime/UMG/Public/Blueprint/UserWidget.h b/Engine/Source/Runtime/UMG/Public/Blueprint/UserWidget.h
--- a/Engine/Source/Runtime/UMG/Public/Blueprint/UserWidget.h
+++ b/Engine/Source/Runtime/UMG/Public/Blueprint/UserWidget.h
@@ -782,7 +782,6 @@ public:
 
 #if WITH_EDITOR
 	//~ Begin UWidget Interface
-	virtual const FSlateBrush* GetEditorIcon() override;
 	virtual const FText GetPaletteCategory() override;
 	//~ End UWidget Interface
 
diff --git a/Engine/Source/Runtime/UMG/Public/Components/Border.h b/Engine/Source/Runtime/UMG/Public/Components/Border.h
--- a/Engine/Source/Runtime/UMG/Public/Components/Border.h
+++ b/Engine/Source/Runtime/UMG/Public/Components/Border.h
@@ -132,7 +132,6 @@ public:
 	virtual void PostEditChangeProperty(struct FPropertyChangedEvent& PropertyChangedEvent) override;
 	//~ End UObject Interface
 
-	virtual const FSlateBrush* GetEditorIcon() override;
 	virtual const FText GetPaletteCategory() override;
 #endif
 
diff --git a/Engine/Source/Runtime/UMG/Public/Components/Widget.h b/Engine/Source/Runtime/UMG/Public/Components/Widget.h
--- a/Engine/Source/Runtime/UMG/Public/Components/Widget.h
+++ b/Engine/Source/Runtime/UMG/Public/Components/Widget.h
@@ -488,8 +488,9 @@ public:
 	virtual void OnCreationFromPalette() { }
 
 	/** Gets the editor icon */
+	DEPRECATED(4.12, "GetEditorIcon is deprecated. Please define widget icons in your style set in the form ClassIcon.MyWidget, and register your style through FClassIconFinder::(Un)RegisterIconSource")
 	virtual const FSlateBrush* GetEditorIcon();
-	
+
 	/** Allows general fixups and connections only used at editor time. */
 	virtual void ConnectEditorData() { }
 
