Copying //UE4/Dev-Sequencer to //UE4/Main

diff --git a/Engine/Source/Runtime/UMG/Private/Animation/MovieScene2DTransformTrack.cpp b/Engine/Source/Runtime/UMG/Private/Animation/MovieScene2DTransformTrack.cpp
--- a/Engine/Source/Runtime/UMG/Private/Animation/MovieScene2DTransformTrack.cpp
+++ b/Engine/Source/Runtime/UMG/Private/Animation/MovieScene2DTransformTrack.cpp
@@ -7,6 +7,13 @@
 #include "MovieScene2DTransformTrackInstance.h"
 #include "MovieSceneCommonHelpers.h"
 
+UMovieScene2DTransformTrack::UMovieScene2DTransformTrack(const FObjectInitializer& ObjectInitializer)
+	: Super(ObjectInitializer)
+{
+#if WITH_EDITORONLY_DATA
+	TrackTint = FColor(48, 227, 255);
+#endif
+}
 
 UMovieSceneSection* UMovieScene2DTransformTrack::CreateNewSection()
 {
diff --git a/Engine/Source/Runtime/UMG/Private/Animation/UMGSequencePlayer.cpp b/Engine/Source/Runtime/UMG/Private/Animation/UMGSequencePlayer.cpp
--- a/Engine/Source/Runtime/UMG/Private/Animation/UMGSequencePlayer.cpp
+++ b/Engine/Source/Runtime/UMG/Private/Animation/UMGSequencePlayer.cpp
@@ -194,4 +194,9 @@ void UUMGSequencePlayer::GetRuntimeObjects(TSharedRef<FMovieSceneSequenceInstanc
 EMovieScenePlayerStatus::Type UUMGSequencePlayer::GetPlaybackStatus() const
 {
 	return PlayerStatus;
+}
+
+void UUMGSequencePlayer::SetPlaybackStatus(EMovieScenePlayerStatus::Type InPlaybackStatus)
+{
+	PlayerStatus = InPlaybackStatus;
 }
\ No newline at end of file
diff --git a/Engine/Source/Runtime/UMG/Public/Animation/MovieScene2DTransformTrack.h b/Engine/Source/Runtime/UMG/Public/Animation/MovieScene2DTransformTrack.h
--- a/Engine/Source/Runtime/UMG/Public/Animation/MovieScene2DTransformTrack.h
+++ b/Engine/Source/Runtime/UMG/Public/Animation/MovieScene2DTransformTrack.h
@@ -18,6 +18,8 @@ class UMovieScene2DTransformTrack : public UMovieScenePropertyTrack
 
 public:
 
+	UMovieScene2DTransformTrack(const FObjectInitializer& ObjectInitializer);
+	
 	// UMovieSceneTrack interface
 
 	virtual UMovieSceneSection* CreateNewSection() override;
diff --git a/Engine/Source/Runtime/UMG/Public/Animation/UMGSequencePlayer.h b/Engine/Source/Runtime/UMG/Public/Animation/UMGSequencePlayer.h
--- a/Engine/Source/Runtime/UMG/Public/Animation/UMGSequencePlayer.h
+++ b/Engine/Source/Runtime/UMG/Public/Animation/UMGSequencePlayer.h
@@ -49,6 +49,7 @@ public:
 	virtual void RemoveMovieSceneInstance( class UMovieSceneSection& MovieSceneSection, TSharedRef<FMovieSceneSequenceInstance> InstanceToRemove ) override {}
 	virtual TSharedRef<FMovieSceneSequenceInstance> GetRootMovieSceneSequenceInstance() const override { return RootMovieSceneInstance.ToSharedRef(); }
 	virtual EMovieScenePlayerStatus::Type GetPlaybackStatus() const override;
+	virtual void SetPlaybackStatus(EMovieScenePlayerStatus::Type InPlaybackStatus) override;
 
 	DECLARE_EVENT_OneParam(UUMGSequencePlayer, FOnSequenceFinishedPlaying, UUMGSequencePlayer&);
 	FOnSequenceFinishedPlaying& OnSequenceFinishedPlaying() { return OnSequenceFinishedPlayingEvent; }
