Slate - popups in nested child windows on 3d rendered widgets will now work.

diff --git a/Engine/Source/Runtime/UMG/Private/Slate/WidgetRenderer.cpp b/Engine/Source/Runtime/UMG/Private/Slate/WidgetRenderer.cpp
--- a/Engine/Source/Runtime/UMG/Private/Slate/WidgetRenderer.cpp
+++ b/Engine/Source/Runtime/UMG/Private/Slate/WidgetRenderer.cpp
@@ -47,6 +47,23 @@ bool SVirtualWindow::OnVisualizeTooltip(const TSharedPtr<SWidget>& TooltipConten
 	return true;
 }
 
+void SVirtualWindow::OnArrangeChildren(const FGeometry& AllottedGeometry, FArrangedChildren& ArrangedChildren) const
+{
+	SWindow::OnArrangeChildren(AllottedGeometry, ArrangedChildren);
+
+	// @HACK VREDITOR - otherwise popup layers don't work in nested child windows, in tab managers and such.
+	if ( ArrangedChildren.Allows3DWidgets() )
+	{
+		const TArray< TSharedRef<SWindow> >& WindowChildren = GetChildWindows();
+		for ( int32 ChildIndex=0; ChildIndex < WindowChildren.Num(); ++ChildIndex )
+		{
+			const TSharedRef<SWindow>& ChildWindow = WindowChildren[ChildIndex];
+			FGeometry ChildWindowGeometry = ChildWindow->GetWindowGeometryInWindow();
+			ChildWindow->ArrangeChildren(ChildWindowGeometry, ArrangedChildren);
+		}
+	}
+}
+
 FWidgetRenderer::FWidgetRenderer(bool bUseGammaCorrection)
 	: bPrepassNeeded(true)
 	, bUseGammaSpace(bUseGammaCorrection)
diff --git a/Engine/Source/Runtime/UMG/Public/Slate/WidgetRenderer.h b/Engine/Source/Runtime/UMG/Public/Slate/WidgetRenderer.h
--- a/Engine/Source/Runtime/UMG/Public/Slate/WidgetRenderer.h
+++ b/Engine/Source/Runtime/UMG/Public/Slate/WidgetRenderer.h
@@ -18,9 +18,10 @@ class SVirtualWindow : public SWindow
 
 public:
 	void Construct(const FArguments& InArgs);
-
+	
 	virtual FPopupMethodReply OnQueryPopupMethod() const override;
 	virtual bool OnVisualizeTooltip(const TSharedPtr<SWidget>& TooltipContent) override;
+	virtual void OnArrangeChildren(const FGeometry& AllottedGeometry, FArrangedChildren& ArrangedChildren) const override;
 
 private:
 	TSharedPtr<class STooltipPresenter> TooltipPresenter;
