UMG - Fixing a regression in the WidgetComponent, renamed a parameter in SetWidget from InWidget back to Widget, so that it doesn't break blueprints.

diff --git a/Engine/Source/Runtime/UMG/Public/Components/WidgetComponent.h b/Engine/Source/Runtime/UMG/Public/Components/WidgetComponent.h
--- a/Engine/Source/Runtime/UMG/Public/Components/WidgetComponent.h
+++ b/Engine/Source/Runtime/UMG/Public/Components/WidgetComponent.h
@@ -121,7 +121,7 @@ public:
 	 *	with either a newer widget or a nullptr
 	 */ 
 	UFUNCTION(BlueprintCallable, Category=UserInterface)
-	void SetWidget(UUserWidget* InWidget);
+	void SetWidget(UUserWidget* Widget);
 
 	/**  
 	 *  Sets a Slate widget to be rendered.  You can use this to draw native Slate widgets using a WidgetComponent, instead
