UMG - Adding back the logic to use the normal WidgetTree when the DesignerWidgetTree is not defined.  This was also done in 4.12.1.

diff --git a/Engine/Source/Runtime/UMG/Private/WidgetBlueprintGeneratedClass.cpp b/Engine/Source/Runtime/UMG/Private/WidgetBlueprintGeneratedClass.cpp
--- a/Engine/Source/Runtime/UMG/Private/WidgetBlueprintGeneratedClass.cpp
+++ b/Engine/Source/Runtime/UMG/Private/WidgetBlueprintGeneratedClass.cpp
@@ -30,7 +30,7 @@ void UWidgetBlueprintGeneratedClass::InitializeWidgetStatic(UUserWidget* UserWid
 	UserWidget->bCanEverPaint = InCanEverPaint;
 
 #if WITH_EDITORONLY_DATA
-	UWidgetTree* ClonedTree = (UWidgetTree*)StaticDuplicateObject(InDesignerWidgetTree, UserWidget, NAME_None, RF_Transactional);
+	UWidgetTree* ClonedTree = (UWidgetTree*)StaticDuplicateObject(InDesignerWidgetTree ? InDesignerWidgetTree : InWidgetTree, UserWidget, NAME_None, RF_Transactional);
 #else
 	UWidgetTree* ClonedTree = DuplicateObject<UWidgetTree>(InWidgetTree, UserWidget);
 #endif
