UMG/Sequencer - Fixing the arugment order, a new parameter was injected into the play animation function which broke serialization for animation nodes and was causing them to play in reverse.  Moving the parameter to the end of the function parameter list instead of adding it between existing parameters.

diff --git a/Engine/Source/Runtime/UMG/Private/Animation/UMGSequencePlayer.cpp b/Engine/Source/Runtime/UMG/Private/Animation/UMGSequencePlayer.cpp
--- a/Engine/Source/Runtime/UMG/Private/Animation/UMGSequencePlayer.cpp
+++ b/Engine/Source/Runtime/UMG/Private/Animation/UMGSequencePlayer.cpp
@@ -14,6 +14,7 @@ UUMGSequencePlayer::UUMGSequencePlayer(const FObjectInitializer& ObjectInitializ
 	PlayerStatus = EMovieScenePlayerStatus::Stopped;
 	TimeCursorPosition = 0.0f;
 	AnimationStartOffset = 0;
+	PlaybackSpeed = 1;
 	Animation = nullptr;
 }
 
@@ -134,7 +135,7 @@ void UUMGSequencePlayer::Tick(float DeltaTime)
 	}
 }
 
-void UUMGSequencePlayer::PlayInternal(double StartAtTime, double EndAtTime, double SubAnimStartTime, double SubAnimEndTime, int32 InNumLoopsToPlay, float InPlaybackSpeed, EUMGSequencePlayMode::Type InPlayMode)
+void UUMGSequencePlayer::PlayInternal(double StartAtTime, double EndAtTime, double SubAnimStartTime, double SubAnimEndTime, int32 InNumLoopsToPlay, EUMGSequencePlayMode::Type InPlayMode, float InPlaybackSpeed)
 {
 	RootMovieSceneInstance = MakeShareable( new FMovieSceneSequenceInstance( *Animation ) );
 	RootMovieSceneInstance->RefreshInstance( *this );
@@ -176,20 +177,20 @@ void UUMGSequencePlayer::PlayInternal(double StartAtTime, double EndAtTime, doub
 	Animation->OnAnimationStarted.Broadcast();
 }
 
-void UUMGSequencePlayer::Play(float StartAtTime, int32 InNumLoopsToPlay, float InPlaybackSpeed, EUMGSequencePlayMode::Type InPlayMode)
+void UUMGSequencePlayer::Play(float StartAtTime, int32 InNumLoopsToPlay, EUMGSequencePlayMode::Type InPlayMode, float InPlaybackSpeed)
 {
 	double SubAnimStartTime = 0.0;
 	double SubAnimEndTime = TimeRange.Size<float>();
 
-	PlayInternal(StartAtTime, 0.0, SubAnimStartTime, SubAnimEndTime, InNumLoopsToPlay, InPlaybackSpeed, InPlayMode);
+	PlayInternal(StartAtTime, 0.0, SubAnimStartTime, SubAnimEndTime, InNumLoopsToPlay, InPlayMode, InPlaybackSpeed);
 }
 
-void UUMGSequencePlayer::PlayTo(float StartAtTime, float EndAtTime, int32 InNumLoopsToPlay, float InPlaybackSpeed, EUMGSequencePlayMode::Type InPlayMode)
+void UUMGSequencePlayer::PlayTo(float StartAtTime, float EndAtTime, int32 InNumLoopsToPlay, EUMGSequencePlayMode::Type InPlayMode, float InPlaybackSpeed)
 {
 	double SubAnimStartTime = 0.0;
 	double SubAnimEndTime = TimeRange.Size<float>();
 
-	PlayInternal(StartAtTime, EndAtTime, SubAnimStartTime, SubAnimEndTime, InNumLoopsToPlay, InPlaybackSpeed, InPlayMode);
+	PlayInternal(StartAtTime, EndAtTime, SubAnimStartTime, SubAnimEndTime, InNumLoopsToPlay, InPlayMode, InPlaybackSpeed);
 }
 
 void UUMGSequencePlayer::Pause()
diff --git a/Engine/Source/Runtime/UMG/Private/UserWidget.cpp b/Engine/Source/Runtime/UMG/Private/UserWidget.cpp
--- a/Engine/Source/Runtime/UMG/Private/UserWidget.cpp
+++ b/Engine/Source/Runtime/UMG/Private/UserWidget.cpp
@@ -291,14 +291,14 @@ void UUserWidget::Invalidate()
 	}
 }
 
-void UUserWidget::PlayAnimation( const UWidgetAnimation* InAnimation, float StartAtTime, int32 NumberOfLoops, float PlaybackSpeed, EUMGSequencePlayMode::Type PlayMode)
+void UUserWidget::PlayAnimation( const UWidgetAnimation* InAnimation, float StartAtTime, int32 NumberOfLoops, EUMGSequencePlayMode::Type PlayMode, float PlaybackSpeed)
 {
 	FScopedNamedEvent NamedEvent(FColor::Emerald, "Widget::PlayAnimation");
 
 	UUMGSequencePlayer* Player = GetOrAddPlayer(InAnimation);
 	if (Player)
 	{
-		Player->Play(StartAtTime, NumberOfLoops, PlaybackSpeed, PlayMode);
+		Player->Play(StartAtTime, NumberOfLoops, PlayMode, PlaybackSpeed);
 
 		Invalidate();
 
@@ -306,14 +306,14 @@ void UUserWidget::PlayAnimation( const UWidgetAnimation* InAnimation, float Star
 	}
 }
 
-void UUserWidget::PlayAnimationTo(const UWidgetAnimation* InAnimation, float StartAtTime, float EndAtTime, int32 NumberOfLoops, float PlaybackSpeed, EUMGSequencePlayMode::Type PlayMode)
+void UUserWidget::PlayAnimationTo(const UWidgetAnimation* InAnimation, float StartAtTime, float EndAtTime, int32 NumberOfLoops, EUMGSequencePlayMode::Type PlayMode, float PlaybackSpeed)
 {
 	FScopedNamedEvent NamedEvent(FColor::Emerald, "Widget::PlayAnimationTo");
 
 	UUMGSequencePlayer* Player = GetOrAddPlayer(InAnimation);
 	if (Player)
 	{
-		Player->PlayTo(StartAtTime, EndAtTime, NumberOfLoops, PlaybackSpeed, PlayMode);
+		Player->PlayTo(StartAtTime, EndAtTime, NumberOfLoops, PlayMode, PlaybackSpeed);
 
 		Invalidate();
 
diff --git a/Engine/Source/Runtime/UMG/Public/Animation/UMGSequencePlayer.h b/Engine/Source/Runtime/UMG/Public/Animation/UMGSequencePlayer.h
--- a/Engine/Source/Runtime/UMG/Public/Animation/UMGSequencePlayer.h
+++ b/Engine/Source/Runtime/UMG/Public/Animation/UMGSequencePlayer.h
@@ -20,10 +20,10 @@ public:
 	void Tick( float DeltaTime );
 
 	/** Begins playing or restarts an animation */
-	void Play(float StartAtTime, int32 InNumLoopsToPlay, float InPlaybackSpeed, EUMGSequencePlayMode::Type InPlayMode);
+	void Play(float StartAtTime, int32 InNumLoopsToPlay, EUMGSequencePlayMode::Type InPlayMode, float InPlaybackSpeed);
 
 	/** Begins playing or restarts an animation  and plays to the specified end time */
-	void PlayTo(float StartAtTime, float EndAtTime, int32 InNumLoopsToPlay, float InPlaybackSpeed, EUMGSequencePlayMode::Type InPlayMode);
+	void PlayTo(float StartAtTime, float EndAtTime, int32 InNumLoopsToPlay, EUMGSequencePlayMode::Type InPlayMode, float InPlaybackSpeed);
 
 	/** Stops a running animation and resets time */
 	void Stop();
@@ -64,7 +64,7 @@ public:
 
 private:
 	/** Internal play function with a verbose parameter set */
-	void PlayInternal(double StartAtTime, double EndAtTime, double SubAnimStartTime, double SubAnimEndTime, int32 InNumLoopsToPlay, float InPlaybackSpeed, EUMGSequencePlayMode::Type InPlayMode);
+	void PlayInternal(double StartAtTime, double EndAtTime, double SubAnimStartTime, double SubAnimEndTime, int32 InNumLoopsToPlay, EUMGSequencePlayMode::Type InPlayMode, float InPlaybackSpeed);
 
 	/** Animation being played */
 	UPROPERTY()
diff --git a/Engine/Source/Runtime/UMG/Public/Blueprint/UserWidget.h b/Engine/Source/Runtime/UMG/Public/Blueprint/UserWidget.h
--- a/Engine/Source/Runtime/UMG/Public/Blueprint/UserWidget.h
+++ b/Engine/Source/Runtime/UMG/Public/Blueprint/UserWidget.h
@@ -706,7 +706,7 @@ public:
 	 * @param PlayMode Specifies the playback mode
 	 */
 	UFUNCTION(BlueprintCallable, BlueprintCosmetic, Category = "User Interface|Animation")
-		void PlayAnimation(const UWidgetAnimation* InAnimation, float StartAtTime = 0.0f, int32 NumLoopsToPlay = 1, float PlaybackSpeed = 1.0f, EUMGSequencePlayMode::Type PlayMode = EUMGSequencePlayMode::Forward);
+	void PlayAnimation(const UWidgetAnimation* InAnimation, float StartAtTime = 0.0f, int32 NumLoopsToPlay = 1, EUMGSequencePlayMode::Type PlayMode = EUMGSequencePlayMode::Forward, float PlaybackSpeed = 1.0f);
 
 	/**
 	 * Plays an animation in this widget a specified number of times stoping at a specified time
@@ -719,7 +719,7 @@ public:
 	 * @param PlayMode Specifies the playback mode
 	 */
 	UFUNCTION(BlueprintCallable, BlueprintCosmetic, Category="User Interface|Animation")
-	void PlayAnimationTo(const UWidgetAnimation* InAnimation, float StartAtTime = 0.0f, float EndAtTime = 0.0f, int32 NumLoopsToPlay = 1, float PlaybackSpeed = 1.0f, EUMGSequencePlayMode::Type PlayMode = EUMGSequencePlayMode::Forward);
+	void PlayAnimationTo(const UWidgetAnimation* InAnimation, float StartAtTime = 0.0f, float EndAtTime = 0.0f, int32 NumLoopsToPlay = 1, EUMGSequencePlayMode::Type PlayMode = EUMGSequencePlayMode::Forward, float PlaybackSpeed = 1.0f);
 
 	/**
 	 * Stops an already running animation in this widget
