Copying //UE4/Dev-Editor to //UE4/Dev-Main (Source: //UE4/Dev-Editor @ 3152045)

diff --git a/Engine/Source/Runtime/Slate/Private/Widgets/Text/SlateEditableTextLayout.cpp b/Engine/Source/Runtime/Slate/Private/Widgets/Text/SlateEditableTextLayout.cpp
--- a/Engine/Source/Runtime/Slate/Private/Widgets/Text/SlateEditableTextLayout.cpp
+++ b/Engine/Source/Runtime/Slate/Private/Widgets/Text/SlateEditableTextLayout.cpp
@@ -3151,14 +3151,16 @@ FVector2D FSlateEditableTextLayout::ComputeDesiredSize(float LayoutScaleMultipli
 	}
 	else
 	{
-		// If a wrapping width has been provided, then we need to report the wrapped size as the desired width
-		// (rather than the actual text layout size as that can have non-breaking lines that extend beyond the wrap width)
-		// Note: We don't do this when auto-wrapping as it would cause a feedback loop in the Slate sizing logic
-		const FVector2D TextLayoutSize = WrappingWidth > 0 ? TextLayout->GetWrappedSize() : TextLayout->GetSize();
+		// If an explicit wrapping width has been provided, then we need to report the wrapped size as the desired width if it has lines that extend beyond the fixed wrapping width
+		// Note: We don't do this when auto-wrapping with a non-explicit width as it would cause a feedback loop in the Slate sizing logic
+		FVector2D TextLayoutSize = TextLayout->GetSize();
+		if (WrappingWidth > 0 && TextLayoutSize.X > WrappingWidth)
+		{
+			TextLayoutSize = TextLayout->GetWrappedSize();
+		}
 
 		DesiredWidth = TextLayoutSize.X;
 		DesiredHeight = TextLayoutSize.Y;
-		
 	}
 
 	// The layouts current margin size. We should not report a size smaller then the margins.
diff --git a/Engine/Source/Runtime/SlateCore/Private/Fonts/LegacySlateFontInfoCache.cpp b/Engine/Source/Runtime/SlateCore/Private/Fonts/LegacySlateFontInfoCache.cpp
--- a/Engine/Source/Runtime/SlateCore/Private/Fonts/LegacySlateFontInfoCache.cpp
+++ b/Engine/Source/Runtime/SlateCore/Private/Fonts/LegacySlateFontInfoCache.cpp
@@ -156,7 +156,7 @@ const FFontData& FLegacySlateFontInfoCache::GetLastResortFontData()
 		// Don't allow GC while we perform this allocation
 		FGCScopeGuard GCGuard;
 
-		const FString LastResortFontPath = FPaths::EngineContentDir() / TEXT("Slate/Fonts/LastResort.ttf");
+		const FString LastResortFontPath = FPaths::EngineContentDir() / TEXT("SlateDebug/Fonts/LastResort.ttf");
 		UFontBulkData* FontBulkData = NewObject<UFontBulkData>();
 		FontBulkData->Initialize(LastResortFontPath);
 		LastResortFontData = MakeShareable(new FFontData(LastResortFontPath, FontBulkData, EFontHinting::Default));
diff --git a/Engine/Source/Runtime/SlateCore/Private/Rendering/DrawElements.cpp b/Engine/Source/Runtime/SlateCore/Private/Rendering/DrawElements.cpp
--- a/Engine/Source/Runtime/SlateCore/Private/Rendering/DrawElements.cpp
+++ b/Engine/Source/Runtime/SlateCore/Private/Rendering/DrawElements.cpp
@@ -799,22 +799,13 @@ void FSlateWindowElementList::ResetBuffers()
 	VolatilePaintList.Reset();
 	BatchData.Reset();
 
-#if SLATE_POOL_DRAW_ELEMENTS
-	DrawElementFreePool.Append(RootDrawLayer.DrawElements);
-#endif
-
 	// Reset the draw elements on the root draw layer
 	RootDrawLayer.DrawElements.Reset();
 
 	// Return child draw layers to the pool, and reset their draw elements.
 	for ( auto& Entry : DrawLayers )
 	{
-#if SLATE_POOL_DRAW_ELEMENTS
-		TArray<FSlateDrawElement*>& DrawElements = Entry.Value->DrawElements;
-		DrawElementFreePool.Append(DrawElements);
-#else
 		TArray<FSlateDrawElement>& DrawElements = Entry.Value->DrawElements;
-#endif
 
 		DrawElements.Reset();
 		DrawLayerPool.Add(Entry.Value);
diff --git a/Engine/Source/Runtime/SlateCore/Private/Rendering/ElementBatcher.cpp b/Engine/Source/Runtime/SlateCore/Private/Rendering/ElementBatcher.cpp
--- a/Engine/Source/Runtime/SlateCore/Private/Rendering/ElementBatcher.cpp
+++ b/Engine/Source/Runtime/SlateCore/Private/Rendering/ElementBatcher.cpp
@@ -106,26 +106,11 @@ void FSlateElementBatcher::AddElements(FSlateDrawLayer& InDrawLayer)
 
 	DrawLayer = &InDrawLayer;
 
-#if SLATE_POOL_DRAW_ELEMENTS
-	const TArray<FSlateDrawElement*>& DrawElements = InDrawLayer.DrawElements;
-#else
 	const TArray<FSlateDrawElement>& DrawElements = InDrawLayer.DrawElements;
-#endif
 
 	for( int32 DrawElementIndex = 0; DrawElementIndex < DrawElements.Num(); ++DrawElementIndex )
 	{
-#if SLATE_POOL_DRAW_ELEMENTS
-		if ( DrawElementIndex < ( DrawElements.Num() - 1 ) )
-		{
-			FPlatformMisc::Prefetch(DrawElements[DrawElementIndex + 1]);
-		}
-#endif
-
-#if SLATE_POOL_DRAW_ELEMENTS
-		const FSlateDrawElement& DrawElement = *DrawElements[DrawElementIndex];
-#else
 		const FSlateDrawElement& DrawElement = DrawElements[DrawElementIndex];
-#endif
 
 		const FSlateRect& InClippingRect = DrawElement.GetClippingRect();
 	
diff --git a/Engine/Source/Runtime/SlateCore/Public/Rendering/DrawElements.h b/Engine/Source/Runtime/SlateCore/Public/Rendering/DrawElements.h
--- a/Engine/Source/Runtime/SlateCore/Public/Rendering/DrawElements.h
+++ b/Engine/Source/Runtime/SlateCore/Public/Rendering/DrawElements.h
@@ -1067,13 +1067,8 @@ public:
 	// Element batch maps sorted by layer.
 	FElementBatchMap LayerToElementBatches;
 
-#if SLATE_POOL_DRAW_ELEMENTS
-	/** The elements drawn on this layer */
-	TArray<FSlateDrawElement*> DrawElements;
-#else
 	/** The elements drawn on this layer */
 	TArray<FSlateDrawElement> DrawElements;
-#endif
 };
 
 /**
@@ -1118,16 +1113,7 @@ public:
 	{
 		return TopLevelWindow.Pin();
 	}
-	
-#if SLATE_POOL_DRAW_ELEMENTS
-	
-	/** @return Get the draw elements that we want to render into this window */
-	FORCEINLINE const TArray<FSlateDrawElement*>& GetDrawElements() const
-	{
-		return RootDrawLayer.DrawElements;
-	}
 
-#else
 
 	/** @return Get the draw elements that we want to render into this window */
 	FORCEINLINE const TArray<FSlateDrawElement>& GetDrawElements() const
@@ -1145,44 +1131,25 @@ public:
 		TArray<FSlateDrawElement>& ActiveDrawElements = DrawStack.Last()->DrawElements;
 		ActiveDrawElements.Add(InDrawElement);
 	}
-#endif
 
 	/**
 	 * Creates an uninitialized draw element
 	 */
 	FORCEINLINE FSlateDrawElement& AddUninitialized()
 	{
-#if SLATE_POOL_DRAW_ELEMENTS
-		FSlateDrawElement* DrawElement = ( DrawElementFreePool.Num() > 0 ) ? DrawElementFreePool.Pop() : new FSlateDrawElement();
-		DrawStack.Last()->DrawElements.Push(DrawElement);
-
-		return *DrawElement;
-#else
 		TArray<FSlateDrawElement>& ActiveDrawElements = DrawStack.Last()->DrawElements;
 		const int32 InsertIdx = ActiveDrawElements.AddDefaulted();
 		return ActiveDrawElements[InsertIdx];
-#endif
 	}
 
-#if SLATE_POOL_DRAW_ELEMENTS
 	/**
 	 * Append draw elements to the list of draw elements
 	 */
-	FORCEINLINE void AppendDrawElements(const TArray<FSlateDrawElement*>& InDrawElements)
-	{
-		TArray<FSlateDrawElement*>& ActiveDrawElements = DrawStack.Last()->DrawElements;
-		ActiveDrawElements.Append(InDrawElements);
-	}
-#else
-	/**
-	* Append draw elements to the list of draw elements
-	*/
 	FORCEINLINE void AppendDrawElements(const TArray<FSlateDrawElement>& InDrawElements)
 	{
 		TArray<FSlateDrawElement>& ActiveDrawElements = DrawStack.Last()->DrawElements;
 		ActiveDrawElements.Append(InDrawElements);
 	}
-#endif
 
 	/**
 	 * Some widgets may want to paint their children after after another, loosely-related widget finished painting.
@@ -1327,11 +1294,6 @@ private:
 	/** */
 	TArray< FSlateDrawLayer* > DrawStack;
 
-#if SLATE_POOL_DRAW_ELEMENTS
-	/** List of draw elements for the window */
-	TArray<FSlateDrawElement*> DrawElementFreePool;
-#endif
-
 	TArray< TSharedPtr<FSlateRenderDataHandle, ESPMode::ThreadSafe> > CachedRenderHandlesInUse;
 
 	/**
diff --git a/Engine/Source/Runtime/SlateCore/Public/Rendering/RenderingCommon.h b/Engine/Source/Runtime/SlateCore/Public/Rendering/RenderingCommon.h
--- a/Engine/Source/Runtime/SlateCore/Public/Rendering/RenderingCommon.h
+++ b/Engine/Source/Runtime/SlateCore/Public/Rendering/RenderingCommon.h
@@ -8,7 +8,6 @@ struct FVector2D;
 class FSlateRect;
 enum class EPopupMethod : uint8;
 
-#define SLATE_PRE_MULTIPLY 1
 #define SLATE_USE_32BIT_INDICES !PLATFORM_USES_ES2
 
 #if SLATE_USE_32BIT_INDICES
diff --git a/Engine/Source/Runtime/SlateCore/Public/Rendering/SlateRenderer.h b/Engine/Source/Runtime/SlateCore/Public/Rendering/SlateRenderer.h
--- a/Engine/Source/Runtime/SlateCore/Public/Rendering/SlateRenderer.h
+++ b/Engine/Source/Runtime/SlateCore/Public/Rendering/SlateRenderer.h
@@ -200,6 +200,20 @@ public:
 	DECLARE_MULTICAST_DELEGATE_TwoParams( FOnSlateWindowRendered, SWindow&, void* );
 	FOnSlateWindowRendered& OnSlateWindowRendered() { return SlateWindowRendered; }
 
+	/**
+	 * Called on the game thread right before a window backbuffer is about to be resized
+	 * @param Pointer to the API specific backbuffer type
+	 */
+	DECLARE_MULTICAST_DELEGATE_OneParam(FOnPreResizeWindowBackbuffer, void*);
+	FOnPreResizeWindowBackbuffer& OnPreResizeWindowBackBuffer() { return PreResizeBackBufferDelegate; }
+
+	/**
+	 * Called on the game thread right after a window backbuffer has been resized
+	 * @param Pointer to the API specific backbuffer type
+	 */
+	DECLARE_MULTICAST_DELEGATE_OneParam(FOnPostResizeWindowBackbuffer, void*);
+	FOnPostResizeWindowBackbuffer& OnPostResizeWindowBackBuffer() { return PostResizeBackBufferDelegate; }
+
 	/** 
 	 * Sets which color vision filter to use
 	 */
@@ -428,6 +442,9 @@ protected:
 	/** Callback that fires after Slate has rendered each window, each frame */
 	FOnSlateWindowRendered SlateWindowRendered;
 
+	FOnPreResizeWindowBackbuffer PreResizeBackBufferDelegate;
+	FOnPostResizeWindowBackbuffer PostResizeBackBufferDelegate;
+
 	/**
 	 * Necessary to grab before flushing the resource pool, as it may be being 
 	 * accessed by multiple threads when loading.
diff --git a/Engine/Source/Runtime/SlateRHIRenderer/Private/SlateMaterialResource.cpp b/Engine/Source/Runtime/SlateRHIRenderer/Private/SlateMaterialResource.cpp
--- a/Engine/Source/Runtime/SlateRHIRenderer/Private/SlateMaterialResource.cpp
+++ b/Engine/Source/Runtime/SlateRHIRenderer/Private/SlateMaterialResource.cpp
@@ -13,6 +13,11 @@ FSlateMaterialResource::FSlateMaterialResource(const UMaterialInterface& InMater
 {
 	SlateProxy->ActualSize = InImageSize.IntPoint();
 	SlateProxy->Resource = this;
+
+#if !UE_BUILD_SHIPPING
+	MaterialObjectWeakPtr = MaterialObject; 
+	UpdateMaterialName();
+#endif
 }
 
 FSlateMaterialResource::~FSlateMaterialResource()
@@ -26,6 +31,12 @@ FSlateMaterialResource::~FSlateMaterialResource()
 void FSlateMaterialResource::UpdateMaterial(const UMaterialInterface& InMaterialResource, const FVector2D& InImageSize, FSlateShaderResource* InTextureMask )
 {
 	MaterialObject = &InMaterialResource;
+
+#if !UE_BUILD_SHIPPING
+	MaterialObjectWeakPtr = MaterialObject;
+	UpdateMaterialName();
+#endif
+
 	if( !SlateProxy )
 	{
 		SlateProxy = new FSlateShaderResourceProxy;
@@ -43,6 +54,12 @@ void FSlateMaterialResource::UpdateMaterial(const UMaterialInterface& InMaterial
 void FSlateMaterialResource::ResetMaterial()
 {
 	MaterialObject = nullptr;
+
+#if !UE_BUILD_SHIPPING
+	MaterialObjectWeakPtr = nullptr;
+	UpdateMaterialName();
+#endif
+
 	TextureMaskResource = nullptr;
 	if (SlateProxy)
 	{
@@ -51,4 +68,23 @@ void FSlateMaterialResource::ResetMaterial()
 	SlateProxy = nullptr;
 	Width = 0;
 	Height = 0;
-}
\ No newline at end of file
+}
+
+#if !UE_BUILD_SHIPPING
+void FSlateMaterialResource::UpdateMaterialName()
+{
+	if(const UMaterialInstanceDynamic* MID = Cast<UMaterialInstanceDynamic>(MaterialObject))
+	{
+		// MID's don't have nice names. Get the name of the parent instead for tracking
+		MaterialName = MID->Parent->GetFName();
+	}
+	else if(MaterialObject)
+	{
+		MaterialName = MaterialObject->GetFName();
+	}
+	else
+	{
+		MaterialName = NAME_None;
+	}
+}
+#endif
diff --git a/Engine/Source/Runtime/SlateRHIRenderer/Private/SlateMaterialResource.h b/Engine/Source/Runtime/SlateRHIRenderer/Private/SlateMaterialResource.h
--- a/Engine/Source/Runtime/SlateRHIRenderer/Private/SlateMaterialResource.h
+++ b/Engine/Source/Runtime/SlateRHIRenderer/Private/SlateMaterialResource.h
@@ -27,11 +27,21 @@ public:
 	FSlateShaderResource* GetTextureMaskResource() const { return TextureMaskResource; }
 public:
 	const class UMaterialInterface* MaterialObject;
+#if !UE_BUILD_SHIPPING
+	// Used to guard against crashes when the material object is deleted.  This is expensive so we do not do it in shipping
+	TWeakObjectPtr<UMaterialInterface> MaterialObjectWeakPtr;
+	FName MaterialName;
+#endif
 	/** Slate proxy used for batching the material */
 	FSlateShaderResourceProxy* SlateProxy;
 
 	FSlateShaderResource* TextureMaskResource;
 	uint32 Width;
 	uint32 Height;
+
+private:
+#if !UE_BUILD_SHIPPING
+	void UpdateMaterialName();
+#endif
 };
 
diff --git a/Engine/Source/Runtime/SlateRHIRenderer/Private/SlateRHIRenderer.cpp b/Engine/Source/Runtime/SlateRHIRenderer/Private/SlateRHIRenderer.cpp
--- a/Engine/Source/Runtime/SlateRHIRenderer/Private/SlateRHIRenderer.cpp
+++ b/Engine/Source/Runtime/SlateRHIRenderer/Private/SlateRHIRenderer.cpp
@@ -321,6 +321,7 @@ void FSlateRHIRenderer::ConditionalResizeViewport( FViewportInfo* ViewInfo, uint
 		ViewInfo->ProjectionMatrix = CreateProjectionMatrix( NewWidth, NewHeight );
 		ViewInfo->bFullscreen = bFullscreen;
 
+		PreResizeBackBufferDelegate.Broadcast(&ViewInfo->ViewportRHI);
 		if( IsValidRef( ViewInfo->ViewportRHI ) )
 		{
 			RHIResizeViewport(ViewInfo->ViewportRHI, NewWidth, NewHeight, bFullscreen);
@@ -329,6 +330,8 @@ void FSlateRHIRenderer::ConditionalResizeViewport( FViewportInfo* ViewInfo, uint
 		{
 			ViewInfo->ViewportRHI = RHICreateViewport(ViewInfo->OSWindow, NewWidth, NewHeight, bFullscreen, ViewInfo->PixelFormat);
 		}
+
+		PostResizeBackBufferDelegate.Broadcast(&ViewInfo->ViewportRHI);
 	}
 }
 
diff --git a/Engine/Source/Runtime/SlateRHIRenderer/Private/SlateRHIRenderingPolicy.cpp b/Engine/Source/Runtime/SlateRHIRenderer/Private/SlateRHIRenderingPolicy.cpp
--- a/Engine/Source/Runtime/SlateRHIRenderer/Private/SlateRHIRenderingPolicy.cpp
+++ b/Engine/Source/Runtime/SlateRHIRenderer/Private/SlateRHIRenderingPolicy.cpp
@@ -330,13 +330,9 @@ void FSlateRHIRenderingPolicy::DrawElements(FRHICommandListImmediate& RHICmdList
 				RHICmdList.SetBlendState(
 					( RenderBatch.DrawFlags & ESlateBatchDrawFlag::NoBlending )
 					? TStaticBlendState<>::GetRHI()
-#if SLATE_PRE_MULTIPLY
 					: ( ( RenderBatch.DrawFlags & ESlateBatchDrawFlag::PreMultipliedAlpha )
 						? TStaticBlendState<CW_RGBA, BO_Add, BF_One, BF_InverseSourceAlpha, BO_Add, BF_One, BF_InverseSourceAlpha>::GetRHI()
 						: TStaticBlendState<CW_RGBA, BO_Add, BF_SourceAlpha, BF_InverseSourceAlpha, BO_Add, BF_One, BF_InverseSourceAlpha>::GetRHI() )
-#else
-					: TStaticBlendState<CW_RGBA, BO_Add, BF_SourceAlpha, BF_InverseSourceAlpha, BO_Add, BF_One, BF_InverseSourceAlpha>::GetRHI()
-#endif
 					);
 #else
 				RHICmdList.SetBlendState(TStaticBlendState<CW_RGB, BO_Add, BF_One, BF_One, BO_Add, BF_Zero, BF_InverseSourceAlpha>::GetRHI());
@@ -463,6 +459,13 @@ void FSlateRHIRenderingPolicy::DrawElements(FRHICommandListImmediate& RHICmdList
 				FSlateMaterialResource* MaterialShaderResource = (FSlateMaterialResource*)ShaderResource;
 				if( MaterialShaderResource->GetMaterialObject() != nullptr )
 				{
+#if !UE_BUILD_SHIPPING
+					// pending kill objects may still be rendered for a frame so it is valid for the check to pass
+					const bool bEvenIfPendingKill = true;
+					// This test needs to be thread safe.  It doesnt give us as many chances to trap bugs here but it is still useful
+					const bool bThreadSafe = true;
+					checkf(MaterialShaderResource->MaterialObjectWeakPtr.IsValid(bEvenIfPendingKill, bThreadSafe), TEXT("Material %s has become invalid.  This means the resource was garbage collected while slate was using it"), *MaterialShaderResource->MaterialName.ToString());
+#endif
 					FMaterialRenderProxy* MaterialRenderProxy = MaterialShaderResource->GetRenderProxy();
 
 					const FMaterial* Material = MaterialRenderProxy->GetMaterial(SceneView->GetFeatureLevel());
diff --git a/Engine/Source/Runtime/SlateRHIRenderer/Private/SlateRHIResourceManager.cpp b/Engine/Source/Runtime/SlateRHIRenderer/Private/SlateRHIResourceManager.cpp
--- a/Engine/Source/Runtime/SlateRHIRenderer/Private/SlateRHIResourceManager.cpp
+++ b/Engine/Source/Runtime/SlateRHIRenderer/Private/SlateRHIResourceManager.cpp
@@ -64,7 +64,7 @@ void FDynamicResourceMap::AddUTextureResource( UTexture* TextureObject, TSharedR
 		check(TextureObject == InResource->TextureObject);
 		TextureMap.Add(TextureObject, InResource);
 
-		TextureMemorySincePurge += TextureObject->GetResourceSize(EResourceSizeMode::Inclusive);
+		TextureMemorySincePurge += TextureObject->GetResourceSizeBytes(EResourceSizeMode::Inclusive);
 	}
 }
 
@@ -84,7 +84,7 @@ void FDynamicResourceMap::RemoveUTextureResource( UTexture* TextureObject )
 	if(TextureObject)
 	{
 		TextureMap.Remove(TextureObject);
-		TextureMemorySincePurge -= TextureObject->GetResourceSize(EResourceSizeMode::Inclusive);
+		TextureMemorySincePurge -= TextureObject->GetResourceSizeBytes(EResourceSizeMode::Inclusive);
 	}
 }
 
diff --git a/Engine/Source/Runtime/StreamingPauseRendering/Private/StreamingPauseRendering.cpp b/Engine/Source/Runtime/StreamingPauseRendering/Private/StreamingPauseRendering.cpp
--- a/Engine/Source/Runtime/StreamingPauseRendering/Private/StreamingPauseRendering.cpp
+++ b/Engine/Source/Runtime/StreamingPauseRendering/Private/StreamingPauseRendering.cpp
@@ -9,6 +9,42 @@
 
 IMPLEMENT_MODULE(FStreamingPauseRenderingModule, StreamingPauseRendering);
 
+class FBackgroundView : public ISlateViewport
+{
+public:
+	FBackgroundView(FTexture2DRHIRef RenderTargetTexture, FIntPoint InSize)
+		: RenderTarget( new FSlateRenderTargetRHI(RenderTargetTexture, InSize.X, InSize.Y))
+		, Size(InSize)
+	{
+		BeginInitResource(RenderTarget);
+	}
+
+	~FBackgroundView()
+	{
+		ReleaseResourceAndFlush(RenderTarget);
+		delete RenderTarget;
+	}
+
+	/* ISlateViewport interface. */
+	virtual FIntPoint GetSize() const override
+	{
+		return Size;
+	}
+
+	virtual FSlateShaderResource* GetViewportRenderTargetTexture() const override
+	{
+		return RenderTarget;
+	}
+
+	virtual bool RequiresVsync() const override
+	{
+		return false;
+	}
+
+private:
+	FSlateRenderTargetRHI* RenderTarget;
+	FIntPoint Size;
+};
 
 static TAutoConsoleVariable<int32> CVarRenderLastFrameInStreamingPause(
 	TEXT("r.RenderLastFrameInStreamingPause"),
@@ -18,9 +54,9 @@ static TAutoConsoleVariable<int32> CVarRenderLastFrameInStreamingPause(
 
 
 FStreamingPauseRenderingModule::FStreamingPauseRenderingModule()
-	: Viewport(nullptr)
+	: SceneViewport(nullptr)
 	, ViewportWidget(nullptr)
-	, ViewportClient(nullptr)
+	, BackgroundView(nullptr)
 { }
 
 
@@ -55,7 +91,7 @@ void FStreamingPauseRenderingModule::BeginStreamingPause( FViewport* GameViewpor
 		.EnableGammaCorrection(false);
 
 	ViewportWidget->SetContent
-		(
+	(
 		SNew(SVerticalBox)
 		+SVerticalBox::Slot()
 		.VAlign(VAlign_Bottom)
@@ -73,32 +109,31 @@ void FStreamingPauseRenderingModule::BeginStreamingPause( FViewport* GameViewpor
  
  	if( SizeXY.X > 0 && SizeXY.Y > 0 && bRenderLastFrame )
 	{
-		ViewportClient = GameViewport->GetClient();
+		FViewportClient* ViewportClient = GameViewport->GetClient();
 
-		Viewport = MakeShareable(new FSceneViewport(ViewportClient, ViewportWidget));
+		SceneViewport = MakeShareable(new FSceneViewport(ViewportClient, ViewportWidget));
 
-		ViewportWidget->SetViewportInterface(Viewport.ToSharedRef());
-	
-		Viewport->UpdateViewportRHI(false,SizeXY.X,SizeXY.Y, EWindowMode::Fullscreen);
+		SceneViewport->UpdateViewportRHI(false,SizeXY.X,SizeXY.Y, EWindowMode::Fullscreen);
 
-		Viewport->EnqueueBeginRenderFrame();
+		SceneViewport->EnqueueBeginRenderFrame();
 
-		FCanvas Canvas(Viewport.Get(), nullptr, ViewportClient->GetWorld(), ViewportClient->GetWorld()->FeatureLevel);
+		FCanvas Canvas(SceneViewport.Get(), nullptr, ViewportClient->GetWorld(), ViewportClient->GetWorld()->FeatureLevel);
 		{
-			ViewportClient->Draw(Viewport.Get(), &Canvas);
+			ViewportClient->Draw(SceneViewport.Get(), &Canvas);
 		}
 		Canvas.Flush_GameThread();
 
-		//Don't need debug canvas I presume?
-		//Viewport->GetDebugCanvas()->Flush(true);
 
 		ENQUEUE_UNIQUE_RENDER_COMMAND_ONEPARAMETER(
 			EndDrawingCommand,
-			FViewport*,Viewport,Viewport.Get(),
-			//FIntPoint,InRestoreSize,RestoreSize,
+			FViewport*,Viewport, SceneViewport.Get(),
 		{
 			Viewport->EndRenderFrame(RHICmdList, false, false);
 		});
+
+		BackgroundView = MakeShareable(new FBackgroundView(SceneViewport->GetRenderTargetTexture(), SizeXY));
+		ViewportWidget->SetViewportInterface(BackgroundView.ToSharedRef());
+
 	}
 
 	//Create the loading screen and begin rendering the widget.
@@ -116,8 +151,8 @@ void FStreamingPauseRenderingModule::EndStreamingPause()
 	GetMoviePlayer()->WaitForMovieToFinish();
 
 	ViewportWidget = nullptr;
-	Viewport = nullptr;
-	ViewportClient = nullptr;
+	SceneViewport = nullptr;
+	BackgroundView = nullptr;
 	
 	FlushRenderingCommands();
 }
\ No newline at end of file
diff --git a/Engine/Source/Runtime/StreamingPauseRendering/Public/StreamingPauseRendering.h b/Engine/Source/Runtime/StreamingPauseRendering/Public/StreamingPauseRendering.h
--- a/Engine/Source/Runtime/StreamingPauseRendering/Public/StreamingPauseRendering.h
+++ b/Engine/Source/Runtime/StreamingPauseRendering/Public/StreamingPauseRendering.h
@@ -37,9 +37,12 @@ public:
 
 public:
 
-	TSharedPtr<FSceneViewport>		Viewport;
-	TSharedPtr<SViewport>			ViewportWidget;
-	FViewportClient*			ViewportClient;
+	/** Viewport being used to render the scene once to a target while paused */
+	TSharedPtr<FSceneViewport> SceneViewport;
+	/** Slate viewport widget used to draw the target */
+	TSharedPtr<SViewport> ViewportWidget;
+	/** Helper class to translate the RHI render target to Slate */
+	TSharedPtr<class FBackgroundView> BackgroundView;
 
 	/** Delegate providing default functionality for beginning streaming pause. */
 	FBeginStreamingPauseDelegate BeginDelegate;
diff --git a/Engine/Source/Runtime/UMG/Private/Components/ComboBoxString.cpp b/Engine/Source/Runtime/UMG/Private/Components/ComboBoxString.cpp
--- a/Engine/Source/Runtime/UMG/Private/Components/ComboBoxString.cpp
+++ b/Engine/Source/Runtime/UMG/Private/Components/ComboBoxString.cpp
@@ -256,7 +256,7 @@ void UComboBoxString::HandleSelectionChanged(TSharedPtr<FString> Item, ESelectIn
 	}
 
 	// When the selection changes we always generate another widget to represent the content area of the comobox.
-	ComoboBoxContent->SetContent( HandleGenerateWidget(Item) );
+	ComoboBoxContent->SetContent( HandleGenerateWidget(CurrentOptionPtr) );
 }
 
 void UComboBoxString::HandleOpening()
diff --git a/Engine/Source/Runtime/UMG/Private/Components/ProgressBar.cpp b/Engine/Source/Runtime/UMG/Private/Components/ProgressBar.cpp
--- a/Engine/Source/Runtime/UMG/Private/Components/ProgressBar.cpp
+++ b/Engine/Source/Runtime/UMG/Private/Components/ProgressBar.cpp
@@ -1,6 +1,7 @@
 // Copyright 1998-2016 Epic Games, Inc. All Rights Reserved.
 
 #include "UMGPrivatePCH.h"
+#include "Slate/SlateBrushAsset.h"
 
 #define LOCTEXT_NAMESPACE "UMG"
 
diff --git a/Engine/Source/Runtime/UMG/Private/Components/WidgetInteractionComponent.cpp b/Engine/Source/Runtime/UMG/Private/Components/WidgetInteractionComponent.cpp
--- a/Engine/Source/Runtime/UMG/Private/Components/WidgetInteractionComponent.cpp
+++ b/Engine/Source/Runtime/UMG/Private/Components/WidgetInteractionComponent.cpp
@@ -497,4 +497,9 @@ const FHitResult& UWidgetInteractionComponent::GetLastHitResult() const
 	return LastHitResult;
 }
 
+FVector2D UWidgetInteractionComponent::Get2DHitLocation() const
+{
+	return LocalHitLocation;
+}
+
 #undef LOCTEXT_NAMESPACE
diff --git a/Engine/Source/Runtime/UMG/Private/Slate/SRetainerWidget.cpp b/Engine/Source/Runtime/UMG/Private/Slate/SRetainerWidget.cpp
--- a/Engine/Source/Runtime/UMG/Private/Slate/SRetainerWidget.cpp
+++ b/Engine/Source/Runtime/UMG/Private/Slate/SRetainerWidget.cpp
@@ -69,26 +69,6 @@ void SRetainerWidget::Construct(const FArguments& InArgs)
 {
 	STAT(MyStatId = FDynamicStats::CreateStatId<FStatGroup_STATGROUP_Slate>(InArgs._StatId);)
 
-	if( FSlateApplication::IsInitialized() )
-	{
-		if (FApp::CanEverRender())
-		{
-			FSlateApplication::Get().OnPreTick().AddRaw(this, &SRetainerWidget::OnTickRetainers);
-		}
-
-#if !UE_BUILD_SHIPPING
-		OnRetainerModeChangedDelegate.AddRaw( this, &SRetainerWidget::OnRetainerModeChanged );
-
-		static bool bStaticInit = false;
-
-		if( !bStaticInit )
-		{
-			bStaticInit = true;
-			EnableRetainedRendering.AsVariable()->SetOnChangedCallback( FConsoleVariableDelegate::CreateStatic( &SRetainerWidget::OnRetainerModeCVarChanged ) );
-		}
-#endif
-	}
-
 	RenderTarget = NewObject<UTextureRenderTarget2D>();
 	RenderTarget->SRGB = true;
 	RenderTarget->TargetGamma = 1;
@@ -112,12 +92,32 @@ void SRetainerWidget::Construct(const FArguments& InArgs)
 	bRenderingOffscreenDesire = true;
 	bRenderingOffscreen = false;
 
+	Window->SetContent(MyWidget.ToSharedRef());
+
 	ChildSlot
 	[
-		MyWidget.ToSharedRef()
+		Window.ToSharedRef()
 	];
 
-	Window->SetContent(MyWidget.ToSharedRef());
+	if ( FSlateApplication::IsInitialized() )
+	{
+		if (FApp::CanEverRender())
+		{
+			FSlateApplication::Get().OnPreTick().AddRaw(this, &SRetainerWidget::OnTickRetainers);
+		}
+
+#if !UE_BUILD_SHIPPING
+		OnRetainerModeChangedDelegate.AddRaw(this, &SRetainerWidget::OnRetainerModeChanged);
+
+		static bool bStaticInit = false;
+
+		if ( !bStaticInit )
+		{
+			bStaticInit = true;
+			EnableRetainedRendering.AsVariable()->SetOnChangedCallback(FConsoleVariableDelegate::CreateStatic(&SRetainerWidget::OnRetainerModeCVarChanged));
+		}
+#endif
+	}
 }
 
 bool SRetainerWidget::ShouldBeRenderingOffscreen() const
@@ -125,6 +125,11 @@ bool SRetainerWidget::ShouldBeRenderingOffscreen() const
 	return bRenderingOffscreenDesire && IsRetainedRenderingEnabled();
 }
 
+bool SRetainerWidget::IsAnythingVisibleToRender() const
+{
+	return GetVisibility().IsVisible() && MyWidget.IsValid() && MyWidget->GetVisibility().IsVisible();
+}
+
 void SRetainerWidget::OnRetainerModeChanged()
 {
 	RefreshRenderingMode();
@@ -153,37 +158,14 @@ void SRetainerWidget::RefreshRenderingMode()
 	{
 		bRenderingOffscreen = bShouldBeRenderingOffscreen;
 
-		ChildSlot
-		[
-			MyWidget.ToSharedRef()
-		];
-
-		// TODO When parent pointers come into play, we need to make it so the virtual window,
-		// doesn't actually ever own the content, it just sorta kind of owns it.
-		if ( bRenderingOffscreen )
-		{
-			Window->SetContent(MyWidget.ToSharedRef());
-		}
-		else
-		{
-			Window->SetContent(SNullWidget::NullWidget);
-		}
+		Window->SetContent(MyWidget.ToSharedRef());
 	}
 }
 
 void SRetainerWidget::SetContent(const TSharedRef< SWidget >& InContent)
 {
 	MyWidget = InContent;
-
-	ChildSlot
-	[
-		InContent
-	];
-
-	if ( bRenderingOffscreen )
-	{
-		Window->SetContent(InContent);
-	}
+	Window->SetContent(InContent);
 }
 
 UMaterialInstanceDynamic* SRetainerWidget::GetEffectMaterial() const
@@ -235,7 +217,7 @@ FChildren* SRetainerWidget::GetChildren()
 
 bool SRetainerWidget::ComputeVolatility() const
 {
-	return true;// return SCompoundWidget::ComputeVolatility();
+	return true;
 }
 
 void SRetainerWidget::OnTickRetainers(float DeltaTime)
@@ -245,8 +227,8 @@ void SRetainerWidget::OnTickRetainers(float DeltaTime)
 	// we should not be added to tick if we're not rendering
 	checkSlow(FApp::CanEverRender());
 
-	bool bShouldRenderAtAnything = GetVisibility().IsVisible() && ChildSlot.GetWidget()->GetVisibility().IsVisible();
-	if ( bRenderingOffscreen && bShouldRenderAtAnything )
+	const bool bShouldRenderAnything = IsAnythingVisibleToRender();
+	if ( bRenderingOffscreen && bShouldRenderAnything )
 	{
 		SCOPE_CYCLE_COUNTER( STAT_SlateRetainerWidgetTick );
 		if ( LastTickedFrame != GFrameCounter && ( GFrameCounter % PhaseCount ) == Phase )
@@ -335,8 +317,8 @@ int32 SRetainerWidget::OnPaint(const FPaintArgs& Args, const FGeometry& Allotted
 
 	MutableThis->RefreshRenderingMode();
 
-	bool bShouldRenderAtAnything = GetVisibility().IsVisible() && ChildSlot.GetWidget()->GetVisibility().IsVisible(); 
-	if ( bRenderingOffscreen && bShouldRenderAtAnything )
+	const bool bShouldRenderAnything = IsAnythingVisibleToRender();
+	if ( bRenderingOffscreen && bShouldRenderAnything )
 	{
 		SCOPE_CYCLE_COUNTER( STAT_SlateRetainerWidgetPaint );
 		CachedAllottedGeometry = AllottedGeometry;
@@ -375,7 +357,7 @@ int32 SRetainerWidget::OnPaint(const FPaintArgs& Args, const FGeometry& Allotted
 		
 		return LayerId;
 	}
-	else if( bShouldRenderAtAnything )
+	else if( bShouldRenderAnything )
 	{
 		return SCompoundWidget::OnPaint(Args, AllottedGeometry, MyClippingRect, OutDrawElements, LayerId, InWidgetStyle, bParentEnabled);
 	}
diff --git a/Engine/Source/Runtime/UMG/Private/WidgetBlueprintLibrary.cpp b/Engine/Source/Runtime/UMG/Private/WidgetBlueprintLibrary.cpp
--- a/Engine/Source/Runtime/UMG/Private/WidgetBlueprintLibrary.cpp
+++ b/Engine/Source/Runtime/UMG/Private/WidgetBlueprintLibrary.cpp
@@ -8,6 +8,9 @@
 #include "Blueprint/WidgetLayoutLibrary.h"
 #include "Engine/Engine.h"
 
+//For PIE error messages
+#include "MessageLog.h"
+
 #define LOCTEXT_NAMESPACE "UMG"
 
 /////////////////////////////////////////////////////
@@ -80,6 +83,12 @@ void UWidgetBlueprintLibrary::SetInputMode_UIOnlyEx(APlayerController* Target, U
 		}
 		Target->SetInputMode(InputMode);
 	}
+	#if WITH_EDITOR 
+	else
+	{ 
+		FMessageLog("PIE").Error(LOCTEXT("UMG WidgetBlueprint Library: SetInputMode_UIOnly", "SetInputMode_UIOnly expects a valid player controller as target"));
+	}
+	#endif // WITH_EDITOR
 }
 
 void UWidgetBlueprintLibrary::SetInputMode_GameAndUI(APlayerController* Target, UWidget* InWidgetToFocus, bool bLockMouseToViewport, bool bHideCursorDuringCapture)
@@ -101,6 +110,13 @@ void UWidgetBlueprintLibrary::SetInputMode_GameAndUIEx(APlayerController* Target
 		}
 		Target->SetInputMode(InputMode);
 	}
+	#if WITH_EDITOR 
+	else
+	{
+		FMessageLog("PIE").Error(LOCTEXT("UMG WidgetBlueprint Library: SetInputMode_GameAndUI", "SetInputMode_GameAndUI expects a valid player controller as target"));	
+	}
+	#endif // WITH_EDITOR
+	
 }
 
 void UWidgetBlueprintLibrary::SetInputMode_GameOnly(APlayerController* Target)
@@ -110,6 +126,12 @@ void UWidgetBlueprintLibrary::SetInputMode_GameOnly(APlayerController* Target)
 		FInputModeGameOnly InputMode;
 		Target->SetInputMode(InputMode);
 	}
+	#if WITH_EDITOR 
+	else
+	{
+		FMessageLog("PIE").Error(LOCTEXT("UMG WidgetBlueprint Library: SetInputMode_GameOnly", "SetInputMode_GameOnly expects a valid player controller as target"));
+	}
+	#endif // WITH_EDITOR
 }
 
 void UWidgetBlueprintLibrary::SetFocusToGameViewport()
diff --git a/Engine/Source/Runtime/UMG/Public/Components/ProgressBar.h b/Engine/Source/Runtime/UMG/Public/Components/ProgressBar.h
--- a/Engine/Source/Runtime/UMG/Public/Components/ProgressBar.h
+++ b/Engine/Source/Runtime/UMG/Public/Components/ProgressBar.h
@@ -2,10 +2,12 @@
 
 #pragma once
 
+#include "Widget.h"
 #include "ProgressWidgetStyle.h"
-
 #include "ProgressBar.generated.h"
 
+class USlateBrushAsset;
+
 /**
  * The progress bar widget is a simple bar that fills up that can be restyled to fit any number of uses.
  *
diff --git a/Engine/Source/Runtime/UMG/Public/Components/TextBlock.h b/Engine/Source/Runtime/UMG/Public/Components/TextBlock.h
--- a/Engine/Source/Runtime/UMG/Public/Components/TextBlock.h
+++ b/Engine/Source/Runtime/UMG/Public/Components/TextBlock.h
@@ -2,6 +2,7 @@
 
 #pragma once
 
+#include "TextWidgetTypes.h"
 #include "TextBlock.generated.h"
 
 /**
diff --git a/Engine/Source/Runtime/UMG/Public/Components/WidgetComponent.h b/Engine/Source/Runtime/UMG/Public/Components/WidgetComponent.h
--- a/Engine/Source/Runtime/UMG/Public/Components/WidgetComponent.h
+++ b/Engine/Source/Runtime/UMG/Public/Components/WidgetComponent.h
@@ -90,7 +90,7 @@ public:
 	 * @param HitResult The hit on this widget in the world
 	 * @param (Out) The transformed 2D hit location on the widget
 	 */
-	void GetLocalHitLocation(FVector WorldHitLocation, FVector2D& OutLocalHitLocation) const;
+	virtual void GetLocalHitLocation(FVector WorldHitLocation, FVector2D& OutLocalHitLocation) const;
 
 	/** @return Gets the last local location that was hit */
 	FVector2D GetLastLocalHitLocation() const
diff --git a/Engine/Source/Runtime/UMG/Public/Components/WidgetInteractionComponent.h b/Engine/Source/Runtime/UMG/Public/Components/WidgetInteractionComponent.h
--- a/Engine/Source/Runtime/UMG/Public/Components/WidgetInteractionComponent.h
+++ b/Engine/Source/Runtime/UMG/Public/Components/WidgetInteractionComponent.h
@@ -154,6 +154,12 @@ public:
 	UFUNCTION(BlueprintCallable, Category="Interaction")
 	const FHitResult& GetLastHitResult() const;
 
+	/**
+	 * Gets the last hit location on the widget in 2D, local pixel units of the render target.
+	 */
+	UFUNCTION(BlueprintCallable, Category="Interaction")
+	FVector2D Get2DHitLocation() const;
+
 	/**
 	 * Set custom hit result.  This is only taken into account if InteractionSource is set to EWidgetInteractionSource::Custom.
 	 */
diff --git a/Engine/Source/Runtime/UMG/Public/Slate/SRetainerWidget.h b/Engine/Source/Runtime/UMG/Public/Slate/SRetainerWidget.h
--- a/Engine/Source/Runtime/UMG/Public/Slate/SRetainerWidget.h
+++ b/Engine/Source/Runtime/UMG/Public/Slate/SRetainerWidget.h
@@ -76,6 +76,7 @@ protected:
 
 	void RefreshRenderingMode();
 	bool ShouldBeRenderingOffscreen() const;
+	bool IsAnythingVisibleToRender() const;
 	void OnRetainerModeChanged();
 private:
 #if !UE_BUILD_SHIPPING
