Copying //UE4/Orion-Staging to //UE4/Main (Source: //Orion/Dev-General @ 3196521)

diff --git a/Engine/Source/Runtime/UMG/Private/Components/BackgroundBlurWidget.cpp b/Engine/Source/Runtime/UMG/Private/Components/BackgroundBlurWidget.cpp
new file mode 100644
--- /dev/null
+++ b/Engine/Source/Runtime/UMG/Private/Components/BackgroundBlurWidget.cpp
@@ -0,0 +1,275 @@
+// Copyright 1998-2016 Epic Games, Inc. All Rights Reserved.
+
+#include "UMGPrivatePCH.h"
+
+#include "BackgroundBlurWidget.h"
+
+#define LOCTEXT_NAMESPACE "UMG"
+
+// This feature has not been tested on es2
+#if PLATFORM_USES_ES2
+static int32 bAllowBackgroundBlur = 0;
+#else
+static int32 bAllowBackgroundBlur = 1;
+#endif
+
+static FAutoConsoleVariableRef CVarSlateAllowBackgroundBlurWidgets(TEXT("Slate.AllowBackgroundBlurWidgets"), bAllowBackgroundBlur, TEXT("If 0, no background blur widgets will be rendered") );
+
+static int32 MaxKernelSize = 255;
+static FAutoConsoleVariableRef CVarSlateMaxKernelSize(TEXT("Slate.BackgroundBlurMaxKernelSize"), MaxKernelSize, TEXT("The maximum allowed kernel size.  Note: Very large numbers can cause a huge decrease in performance"));
+
+static int32 bDownsampleForBlur = 1;
+static FAutoConsoleVariableRef CVarDownsampleForBlur(TEXT("Slate.BackgroundBlurDownsample"), bDownsampleForBlur, TEXT(""), ECVF_Cheat);
+
+static int32 bForceLowQualityBrushFallback = 0;
+static FAutoConsoleVariableRef CVarForceLowQualityBackgroundBlurOverride(TEXT("Slate.ForceBackgroundBlurLowQualityOverride"), bForceLowQualityBrushFallback, TEXT("Whether or not to force a slate brush to be used instead of actually blurring the background"), ECVF_Scalability);
+
+class SBackgroundBlurWidget : public SCompoundWidget
+{
+public:
+	SLATE_BEGIN_ARGS(SBackgroundBlurWidget)
+		: _BlurRadius(0)
+		, _LowQualityFallbackBrush(nullptr)
+	{}
+		SLATE_DEFAULT_SLOT(FArguments, Content)
+		SLATE_ATTRIBUTE(int32, BlurRadius)
+		SLATE_ARGUMENT(const FSlateBrush*, LowQualityFallbackBrush)
+	SLATE_END_ARGS()
+
+	void Construct(const FArguments& InArgs)
+	{
+		LowQualityFallbackBrush = InArgs._LowQualityFallbackBrush;
+		BlurRadiusAttrib = InArgs._BlurRadius;
+		Strength = 0;
+
+		ChildSlot
+		[
+			InArgs._Content.Widget
+		];
+	}
+
+	void SetContent(const TSharedRef< SWidget >& InContent)
+	{
+		ChildSlot
+		[
+			InContent
+		];
+	}
+
+	void SetBlurRadius(int32 InBlurRadius)
+	{
+		BlurRadiusAttrib.Set(InBlurRadius);
+		Invalidate(EInvalidateWidget::Layout);
+	}
+
+	void SetBlurStrength(float InStrength)
+	{
+		Strength = InStrength;
+
+		// Strength cant be negative
+		Strength = FMath::Max(0.0f, InStrength);
+		Invalidate(EInvalidateWidget::Layout);
+	}
+
+	void SetLowQualityBackgroundBrush(const FSlateBrush* InBrush) 
+	{
+		LowQualityFallbackBrush = InBrush;
+		Invalidate(EInvalidateWidget::Layout);
+	}
+
+	virtual int32 OnPaint( const FPaintArgs& Args, const FGeometry& AllottedGeometry, const FSlateRect& MyClippingRect, FSlateWindowElementList& OutDrawElements, int32 LayerId, const FWidgetStyle& InWidgetStyle, bool bParentEnabled ) const override 
+	{
+		int32 PostFXLayerId = LayerId;
+		if(bAllowBackgroundBlur && !bForceLowQualityBrushFallback)
+		{
+			const int32 BlurRadius = BlurRadiusAttrib.Get();
+
+			if(Strength > 0)
+			{
+				FPaintGeometry PaintGeometry = AllottedGeometry.ToPaintGeometry();
+
+				// extract the layout transform from the draw element
+				FSlateLayoutTransform InverseLayoutTransform(Inverse(FSlateLayoutTransform(PaintGeometry.DrawScale, PaintGeometry.DrawPosition)));
+				// The clip rect is NOT subject to the rotations specified by MakeRotatedBox.
+				FSlateRotatedClipRectType RenderClipRect = FSlateRotatedClipRectType::MakeSnappedRotatedRect(MyClippingRect, InverseLayoutTransform, AllottedGeometry.GetAccumulatedRenderTransform());
+
+				float OffsetX = PaintGeometry.DrawPosition.X - FMath::TruncToFloat(PaintGeometry.DrawPosition.X);
+				float OffsetY = PaintGeometry.DrawPosition.Y - FMath::TruncToFloat(PaintGeometry.DrawPosition.Y);
+
+				int32 RenderTargetWidth = FMath::RoundToInt(RenderClipRect.ExtentX.X);
+				int32 RenderTargetHeight = FMath::RoundToInt(RenderClipRect.ExtentY.Y);
+
+				int32 KernelSize = 0;
+				float ComputedStrength = 0;
+				int32 DownsampleAmount = 0;
+				ComputeEffectiveKernelSize(BlurRadius, KernelSize, DownsampleAmount);
+				
+				ComputedStrength = FMath::Max(.5f, Strength);
+
+				if (DownsampleAmount > 0)
+				{
+					RenderTargetWidth = FMath::DivideAndRoundUp(RenderTargetWidth,DownsampleAmount);
+					RenderTargetHeight = FMath::DivideAndRoundUp(RenderTargetHeight,DownsampleAmount);
+					ComputedStrength /= DownsampleAmount;
+				}
+
+				if(RenderTargetWidth > 0 && RenderTargetHeight > 0)
+				{
+					FSlateDrawElement::MakePostProcessPass(OutDrawElements, LayerId, PaintGeometry, MyClippingRect, FVector4(KernelSize, ComputedStrength, RenderTargetWidth, RenderTargetHeight), DownsampleAmount);
+				}
+		
+
+				++PostFXLayerId;
+			}
+		}
+		else if(bAllowBackgroundBlur && bForceLowQualityBrushFallback && LowQualityFallbackBrush && LowQualityFallbackBrush->DrawAs != ESlateBrushDrawType::NoDrawType)
+		{
+			const bool bIsEnabled = ShouldBeEnabled(bParentEnabled);
+			const uint32 DrawEffects = bIsEnabled ? ESlateDrawEffect::None : ESlateDrawEffect::DisabledEffect;
+
+			const FLinearColor FinalColorAndOpacity(InWidgetStyle.GetColorAndOpacityTint() * LowQualityFallbackBrush->GetTint(InWidgetStyle));
+
+			FSlateDrawElement::MakeBox(OutDrawElements, PostFXLayerId, AllottedGeometry.ToPaintGeometry(), LowQualityFallbackBrush, MyClippingRect, DrawEffects, FinalColorAndOpacity);
+			++PostFXLayerId;
+		}
+
+		return SCompoundWidget::OnPaint(Args, AllottedGeometry, MyClippingRect, OutDrawElements, PostFXLayerId, InWidgetStyle, bParentEnabled);
+	}
+
+private:
+	FORCENOINLINE void ComputeEffectiveKernelSize(int32 Radius, int32& OutKernelSize, int32& OutDownsampleAmount) const
+	{
+		if(Radius==0)
+		{
+			// Auto compute
+			OutKernelSize = FMath::RoundToInt(Strength*3);
+		}
+		else
+		{
+			OutKernelSize = Radius;
+		}
+	
+		const bool bRequiresDownsample = OutKernelSize > 9 && bDownsampleForBlur;
+
+		if(bRequiresDownsample)
+		{
+			OutDownsampleAmount = OutKernelSize >= 64 ? 4 : 2;
+			OutKernelSize /= OutDownsampleAmount;
+		}
+
+		if (OutKernelSize % 2 == 0)
+		{
+			// Kernel sizes must be odd
+			++OutKernelSize;
+		}
+		OutKernelSize = FMath::Clamp(OutKernelSize, 3, MaxKernelSize);
+
+	}
+private:
+	TAttribute<int32> BlurRadiusAttrib;
+	const FSlateBrush* LowQualityFallbackBrush;
+	float Strength;
+};
+
+
+UBackgroundBlurWidget::UBackgroundBlurWidget(const FObjectInitializer& ObjectInitializer)
+	: Super(ObjectInitializer)
+	, LowQualityFallbackBrush(FSlateNoResource())
+{
+}
+
+void UBackgroundBlurWidget::ReleaseSlateResources(bool bReleaseChildren)
+{
+	Super::ReleaseSlateResources(bReleaseChildren);
+
+	MyWidget.Reset();
+}
+
+TSharedRef<SWidget> UBackgroundBlurWidget::RebuildWidget()
+{
+	MyWidget = SNew(SBackgroundBlurWidget);
+
+	if ( GetChildrenCount() > 0 )
+	{
+		MyWidget->SetContent(GetContentSlot()->Content ? GetContentSlot()->Content->TakeWidget() : SNullWidget::NullWidget);
+	}
+	
+	return BuildDesignTimeWidget(MyWidget.ToSharedRef());
+}
+
+void UBackgroundBlurWidget::SynchronizeProperties()
+{
+	Super::SynchronizeProperties();
+
+	if(MyWidget.IsValid())
+	{
+		if(bOverrideAutoRadiusCalculation)
+		{
+			MyWidget->SetBlurRadius(BlurRadius);
+		}
+		else
+		{
+			// Force it to zero to auto-comute a radius
+			MyWidget->SetBlurRadius(0);
+		}
+		MyWidget->SetBlurStrength(BlurStrength);
+		MyWidget->SetLowQualityBackgroundBrush(&LowQualityFallbackBrush);
+	}
+}
+
+void UBackgroundBlurWidget::OnSlotAdded(UPanelSlot* InSlot)
+{
+	// Add the child to the live slot if it already exists
+	if (MyWidget.IsValid() )
+	{
+		MyWidget->SetContent(InSlot->Content ? InSlot->Content->TakeWidget() : SNullWidget::NullWidget);
+	}
+}
+
+void UBackgroundBlurWidget::OnSlotRemoved(UPanelSlot* InSlot)
+{
+	// Remove the widget from the live slot if it exists.
+	if (MyWidget.IsValid() )
+	{
+		MyWidget->SetContent(SNullWidget::NullWidget);
+	}
+}
+
+void UBackgroundBlurWidget::SetBlurRadius(int32 InBlurRadius)
+{
+	BlurRadius = InBlurRadius;
+	if (MyWidget.IsValid())
+	{
+		MyWidget->SetBlurRadius(BlurRadius);
+	}
+}
+
+void UBackgroundBlurWidget::SetBlurStrength(float InStrength)
+{
+	BlurStrength = InStrength;
+	if (MyWidget.IsValid())
+	{
+		MyWidget->SetBlurStrength(BlurStrength);
+	}
+}
+
+void UBackgroundBlurWidget::SetLowQualityFallbackBrush(const FSlateBrush& InBrush)
+{
+	LowQualityFallbackBrush = InBrush;
+	if(MyWidget.IsValid())
+	{
+		MyWidget->SetLowQualityBackgroundBrush(&LowQualityFallbackBrush);
+	}
+}
+
+#if WITH_EDITOR
+
+const FText UBackgroundBlurWidget::GetPaletteCategory()
+{
+	return LOCTEXT("SpecialFX", "Special Effects");
+}
+
+#endif
+
+
+#undef LOCTEXT_NAMESPACE
diff --git a/Engine/Source/Runtime/UMG/Private/Components/InputKeySelector.cpp b/Engine/Source/Runtime/UMG/Private/Components/InputKeySelector.cpp
--- a/Engine/Source/Runtime/UMG/Private/Components/InputKeySelector.cpp
+++ b/Engine/Source/Runtime/UMG/Private/Components/InputKeySelector.cpp
@@ -67,7 +67,7 @@ void UInputKeySelector::SynchronizeProperties()
 TSharedRef<SWidget> UInputKeySelector::RebuildWidget()
 {
 	MyInputKeySelector = SNew(SInputKeySelector)
-		.SelectedKey( SelectedKey)
+		.SelectedKey(SelectedKey)
 		.Font(Font)
 		.Margin(Margin)
 		.ColorAndOpacity(ColorAndOpacity)
@@ -79,7 +79,7 @@ TSharedRef<SWidget> UInputKeySelector::RebuildWidget()
 	return MyInputKeySelector.ToSharedRef();
 }
 
-void UInputKeySelector::HandleKeySelected( FInputChord InSelectedKey )
+void UInputKeySelector::HandleKeySelected(const FInputChord& InSelectedKey)
 {
 	SelectedKey = InSelectedKey;
 	OnKeySelected.Broadcast(SelectedKey);
diff --git a/Engine/Source/Runtime/UMG/Private/UserWidget.cpp b/Engine/Source/Runtime/UMG/Private/UserWidget.cpp
--- a/Engine/Source/Runtime/UMG/Private/UserWidget.cpp
+++ b/Engine/Source/Runtime/UMG/Private/UserWidget.cpp
@@ -770,6 +770,14 @@ APlayerController* UUserWidget::GetOwningPlayer() const
 	return PlayerContext.IsValid() ? PlayerContext.GetPlayerController() : nullptr;
 }
 
+void UUserWidget::SetOwningPlayer(APlayerController* LocalPlayerController)
+{
+	if (LocalPlayerController && LocalPlayerController->IsLocalController())
+	{
+		PlayerContext = FLocalPlayerContext(LocalPlayerController);
+	}
+}
+
 class APawn* UUserWidget::GetOwningPlayerPawn() const
 {
 	if ( APlayerController* PC = GetOwningPlayer() )
diff --git a/Engine/Source/Runtime/UMG/Public/Blueprint/UserWidget.h b/Engine/Source/Runtime/UMG/Public/Blueprint/UserWidget.h
--- a/Engine/Source/Runtime/UMG/Public/Blueprint/UserWidget.h
+++ b/Engine/Source/Runtime/UMG/Public/Blueprint/UserWidget.h
@@ -239,7 +239,7 @@ public:
 	ULocalPlayer* GetOwningLocalPlayer() const;
 
 	/**
-	 * Sets the local player associated with this UI.
+	 * Sets the player associated with this UI via LocalPlayer reference.
 	 * @param LocalPlayer The local player you want to be the conceptual owner of this UI.
 	 */
 	UFUNCTION(BlueprintCallable, BlueprintCosmetic, Category="Player")
@@ -252,6 +252,13 @@ public:
 	UFUNCTION(BlueprintCallable, BlueprintCosmetic, Category="Player")
 	class APlayerController* GetOwningPlayer() const;
 
+	/**
+	 * Sets the local player associated with this UI via PlayerController reference.
+	 * @param LocalPlayerController The PlayerController of the local player you want to be the conceptual owner of this UI.
+	 */
+	UFUNCTION(BlueprintCallable, BlueprintCosmetic, Category="Player")
+	void SetOwningPlayer(APlayerController* LocalPlayerController);
+
 	/**
 	 * Gets the player pawn associated with this UI.
 	 * @return Gets the owning player pawn that's owned by the player controller assigned to this widget.
diff --git a/Engine/Source/Runtime/UMG/Public/Components/BackgroundBlurWidget.h b/Engine/Source/Runtime/UMG/Public/Components/BackgroundBlurWidget.h
new file mode 100644
--- /dev/null
+++ b/Engine/Source/Runtime/UMG/Public/Components/BackgroundBlurWidget.h
@@ -0,0 +1,68 @@
+// Copyright 1998-2016 Epic Games, Inc. All Rights Reserved.
+
+#pragma once
+
+#include "ContentWidget.h"
+#include "BackgroundBlurWidget.generated.h"
+
+UCLASS()
+class UMG_API UBackgroundBlurWidget : public UContentWidget
+{
+	GENERATED_UCLASS_BODY()
+
+public:
+	virtual void ReleaseSlateResources(bool bReleaseChildren) override;
+
+#if WITH_EDITOR
+	virtual const FText GetPaletteCategory() override;
+#endif
+	UFUNCTION(BlueprintCallable, Category = "Appearance")
+	void SetBlurRadius(int32 InBlurRadius);
+
+	UFUNCTION(BlueprintCallable, Category = "Appearance")
+	void SetBlurStrength(float InStrength);
+
+	UFUNCTION(BlueprintCallable, Category = "Appearance")
+	void SetLowQualityFallbackBrush(const FSlateBrush& InBrush);
+
+protected:
+	/** UWidget interface */
+	virtual TSharedRef<SWidget> RebuildWidget() override;
+	virtual void SynchronizeProperties() override;
+
+	/** UPanelWidget interface */
+	virtual void OnSlotAdded(UPanelSlot* Slot) override;
+	virtual void OnSlotRemoved(UPanelSlot* Slot) override;
+
+public:
+	/**
+	* This is the strength of the blur.  It is heavily dependent on the BlurRadius value.  The ultimate strength of the blur is determined by this value and the BlurRadius value
+	* A small blur radius and a large strength will quickly darken the entire blurred area and result in less overall blurriness that can be achieved.
+	*/
+	UPROPERTY(EditAnywhere, BlueprintReadOnly, Category=Appearance, meta=(ClampMin=0,ClampMax=100))
+	float BlurStrength;
+
+	/**
+	 * Whether or not the radius should be computed automatically or if it should use the radius  
+	 */
+	UPROPERTY()
+	bool bOverrideAutoRadiusCalculation;
+
+	/**
+	 * This is the number of pixels which will be weighted in each direction from any given pixel when computing the blur
+	 * A larger value is more costly but allows for stronger blurs.  
+	 */
+	UPROPERTY(EditAnywhere, AdvancedDisplay, BlueprintReadOnly, Category=Appearance, meta=(ClampMin=0, ClampMax=255, EditCondition="bOverrideAutoRadiusCalculation"))
+	int32 BlurRadius;
+
+	/**
+	 * An image to draw instead of applying a blur when low quality override mode is enabled. 
+	 * You can enable low quality mode for background blurs by setting the cvar Slate.ForceBackgroundBlurLowQualityOverride to 1. 
+	 * This is usually done in the project's scalability settings
+	 */
+	UPROPERTY(EditAnywhere, AdvancedDisplay, BlueprintReadOnly, Category=Appearance)
+	FSlateBrush LowQualityFallbackBrush;
+
+protected:
+	TSharedPtr<class SBackgroundBlurWidget> MyWidget;
+};
diff --git a/Engine/Source/Runtime/UMG/Public/Components/InputKeySelector.h b/Engine/Source/Runtime/UMG/Public/Components/InputKeySelector.h
--- a/Engine/Source/Runtime/UMG/Public/Components/InputKeySelector.h
+++ b/Engine/Source/Runtime/UMG/Public/Components/InputKeySelector.h
@@ -79,11 +79,10 @@ protected:
 	//~ End UWidget Interface
 
 private:
-	void HandleKeySelected( FInputChord SelectedKey );
+	virtual void HandleKeySelected(const FInputChord& InSelectedKey);
 	void HandleIsSelectingKeyChanged();
 
 private:
-
 	/** The style for the button used to start key selection mode. */
 	const FButtonStyle* ButtonStyle;
 
diff --git a/Engine/Source/Runtime/UMG/Public/Components/TextBlock.h b/Engine/Source/Runtime/UMG/Public/Components/TextBlock.h
--- a/Engine/Source/Runtime/UMG/Public/Components/TextBlock.h
+++ b/Engine/Source/Runtime/UMG/Public/Components/TextBlock.h
@@ -123,7 +123,7 @@ public:
 	 * @param InText The text to assign to the widget
 	 */
 	UFUNCTION(BlueprintCallable, Category="Widget", meta=(DisplayName="SetText (Text)"))
-	void SetText(FText InText);
+	virtual void SetText(FText InText);
 
 	//~ Begin UWidget Interface
 	virtual void SynchronizeProperties() override;
