Copying //UE4/Dev-Editor to //UE4/Dev-Main (Source: //UE4/Dev-Editor @ 3228984)

diff --git a/Engine/Source/Runtime/SlateRHIRenderer/Private/SlatePostProcessor.h b/Engine/Source/Runtime/SlateRHIRenderer/Private/SlatePostProcessor.h
--- a/Engine/Source/Runtime/SlateRHIRenderer/Private/SlatePostProcessor.h
+++ b/Engine/Source/Runtime/SlateRHIRenderer/Private/SlatePostProcessor.h
@@ -32,7 +32,7 @@ public:
 	~FSlatePostProcessor();
 
 	void BlurRect(FRHICommandListImmediate& RHICmdList, IRendererModule& RendererModule, const FBlurRectParams& Params, const FPostProcessRectParams& RectParams);
-
+	void ReleaseRenderTargets();
 private:
 	void DownsampleRect(FRHICommandListImmediate& RHICmdList, IRendererModule& RendererModule, const FPostProcessRectParams& Params, const FIntPoint& DownsampleSize);
 	void UpsampleRect(FRHICommandListImmediate& RHICmdList, IRendererModule& RendererModule, const FPostProcessRectParams& Params, const FIntPoint& DownsampleSize);
diff --git a/Engine/Source/Runtime/SlateRHIRenderer/Private/SlateRHIRenderer.cpp b/Engine/Source/Runtime/SlateRHIRenderer/Private/SlateRHIRenderer.cpp
--- a/Engine/Source/Runtime/SlateRHIRenderer/Private/SlateRHIRenderer.cpp
+++ b/Engine/Source/Runtime/SlateRHIRenderer/Private/SlateRHIRenderer.cpp
@@ -148,6 +148,7 @@ FSlateRHIRenderer::FSlateRHIRenderer( TSharedRef<FSlateFontServices> InSlateFont
 	, EnqueuedWindowDrawBuffer(NULL)
 	, FreeBufferIndex(1)
 #endif
+	, CurrentSceneIndex(-1)
 {
 	ResourceManager = InResourceManager;
 
@@ -192,6 +193,8 @@ bool FSlateRHIRenderer::Initialize()
 
 	ElementBatcher = MakeShareable( new FSlateElementBatcher( RenderingPolicy.ToSharedRef() ) );
 
+	CurrentSceneIndex = -1;
+	ActiveScenes.Empty();
 	return true;
 }
 
@@ -234,6 +237,8 @@ void FSlateRHIRenderer::Destroy()
 	}
 
 	WindowToViewportInfo.Empty();
+	CurrentSceneIndex = -1;
+	ActiveScenes.Empty();
 }
 
 /** Returns a draw buffer that can be used by Slate windows to draw window elements */
@@ -383,8 +388,10 @@ void FSlateRHIRenderer::ConditionalResizeViewport( FViewportInfo* ViewInfo, uint
 		ViewInfo->HDRColorGamut = HDRColorGamut;
 		ViewInfo->HDROutputDevice = HDROutputDevice;
 
+		PreResizeBackBufferDelegate.Broadcast(&ViewInfo->ViewportRHI);
 		if( IsValidRef( ViewInfo->ViewportRHI ) )
 		{
+			ensureMsgf(ViewInfo->ViewportRHI->GetRefCount() == 1, TEXT("Viewport backbuffer was not properly released"));
 			RHIResizeViewport(ViewInfo->ViewportRHI, NewWidth, NewHeight, bFullscreen, ViewInfo->PixelFormat);
 		}
 		else
@@ -1556,6 +1563,56 @@ ISlateAtlasProvider* FSlateRHIRenderer::GetTextureAtlasProvider()
 	return nullptr;
 }
 
+
+
+int32 FSlateRHIRenderer::RegisterCurrentScene(FSceneInterface* Scene)
+{
+	check(IsInGameThread());
+	if (Scene)
+	{
+		CurrentSceneIndex = ActiveScenes.AddUnique(Scene);
+	}
+	else
+	{
+		CurrentSceneIndex = -1;
+	}
+
+	// We need to keep the ActiveScenes array synchronized with the Policy's ActiveScenes array on
+	// the render thread.
+	ENQUEUE_UNIQUE_RENDER_COMMAND_THREEPARAMETER(
+		RegisterCurrentSceneOnPolicy,
+		FSlateRHIRenderingPolicy*, InRenderPolicy, RenderingPolicy.Get(),
+		FSceneInterface*, InScene, Scene,
+		int32, InSceneIndex, CurrentSceneIndex,
+	{
+		InRenderPolicy->AddSceneAt(InScene, InSceneIndex);
+	});
+	return CurrentSceneIndex;
+}
+
+int32 FSlateRHIRenderer::GetCurrentSceneIndex() const
+{
+	return CurrentSceneIndex;
+}
+
+void FSlateRHIRenderer::ClearScenes()
+{
+	if(!IsInSlateThread())
+	{
+		CurrentSceneIndex = -1;
+		ActiveScenes.Empty();
+
+		// We need to keep the ActiveScenes array synchronized with the Policy's ActiveScenes array on
+		// the render thread.
+		ENQUEUE_UNIQUE_RENDER_COMMAND_ONEPARAMETER(
+			ClearScenesOnPolicy,
+			FSlateRenderingPolicy*, InRenderPolicy, RenderingPolicy.Get(),
+			{
+				InRenderPolicy->ClearScenes();
+			});
+	}
+}
+
 bool FSlateRHIRenderer::AreShadersInitialized() const
 {
 #if WITH_EDITORONLY_DATA
@@ -1580,8 +1637,14 @@ void FSlateRHIRenderer::ReleaseAccessedResources(bool bImmediatelyFlush)
 	// Clear accessed UTexture and Material objects from the previous frame
 	ResourceManager->BeginReleasingAccessedResources(bImmediatelyFlush);
 
+	// We keep track of the Scene objects from SceneViewports on the SlateRenderer. Make sure that this gets refreshed every frame.
+	ClearScenes();
+
 	if ( bImmediatelyFlush )
 	{
+		// Release resources generated specifically by the rendering policy if we are flushing.  This should NOT be done unless flushing
+		RenderingPolicy->FlushGeneratedResources();
+
 		FlushCommands();
 	}
 }
diff --git a/Engine/Source/Runtime/SlateRHIRenderer/Private/SlateRHIRenderer.h b/Engine/Source/Runtime/SlateRHIRenderer/Private/SlateRHIRenderer.h
--- a/Engine/Source/Runtime/SlateRHIRenderer/Private/SlateRHIRenderer.h
+++ b/Engine/Source/Runtime/SlateRHIRenderer/Private/SlateRHIRenderer.h
@@ -186,8 +186,12 @@ private:
 		{
 			if (RTProvider)
 			{
-				FSlateRenderTargetRHI* SlateTarget = (FSlateRenderTargetRHI*)(RTProvider->GetViewportRenderTargetTexture());
-				return SlateTarget->GetTypedResource();
+				FSlateShaderResource* RenderTargetTexture = RTProvider->GetViewportRenderTargetTexture();
+				if( RenderTargetTexture )
+				{
+					FSlateRenderTargetRHI* RHITarget = (FSlateRenderTargetRHI*)RenderTargetTexture;
+					return RHITarget->GetTypedResource();
+				}
 			}
 			return nullptr;
 		}
@@ -234,6 +238,9 @@ public:
 	virtual void ReleaseAccessedResources(bool bImmediatelyFlush) override;
 	virtual TSharedRef<FSlateRenderDataHandle, ESPMode::ThreadSafe> CacheElementRenderData(const ILayoutCache* Cacher, FSlateWindowElementList& ElementList) override;
 	virtual void ReleaseCachingResourcesFor(const ILayoutCache* Cacher) override;
+	virtual int32 RegisterCurrentScene(FSceneInterface* Scene) override;
+	virtual int32 GetCurrentSceneIndex() const override;
+	virtual void ClearScenes() override;
 
 	/** Draws windows from a FSlateDrawBuffer on the render thread */
 	void DrawWindow_RenderThread(FRHICommandListImmediate& RHICmdList, FSlateRHIRenderer::FViewportInfo& ViewportInfo, FSlateWindowElementList& WindowElementList, bool bLockToVsync, bool bClear);
@@ -340,6 +347,10 @@ private:
 	bool bTakingAScreenShot;
 	FIntRect ScreenshotRect;
 	TArray<FColor>* OutScreenshotData;
+
+	/** These are state management variables for Scenes on the game thread. A similar copy exists on the RHI Rendering Policy for the rendering thread.*/
+	TArray<FSceneInterface*> ActiveScenes;
+	int32 CurrentSceneIndex;
 };
 
 struct FSlateEndDrawingWindowsCommand : public FRHICommand < FSlateEndDrawingWindowsCommand >
diff --git a/Engine/Source/Runtime/SlateRHIRenderer/Private/SlateRHIRenderingPolicy.cpp b/Engine/Source/Runtime/SlateRHIRenderer/Private/SlateRHIRenderingPolicy.cpp
--- a/Engine/Source/Runtime/SlateRHIRenderer/Private/SlateRHIRenderingPolicy.cpp
+++ b/Engine/Source/Runtime/SlateRHIRenderer/Private/SlateRHIRenderingPolicy.cpp
@@ -179,13 +179,19 @@ void FSlateRHIRenderingPolicy::UpdateVertexAndIndexBuffers(FRHICommandListImmedi
 	SET_DWORD_STAT( STAT_SlateVertexCount, InBatchData.GetNumBatchedVertices() );
 }
 
-static FSceneView& CreateSceneView( FSceneViewFamilyContext& ViewFamilyContext, FSlateBackBuffer& BackBuffer, const FMatrix& ViewProjectionMatrix )
+static FSceneView* CreateSceneView( FSceneViewFamilyContext* ViewFamilyContext, FSlateBackBuffer& BackBuffer, const FMatrix& ViewProjectionMatrix )
 {
+	// In loading screens, the engine is NULL, so we skip out.
+	if (GEngine == nullptr)
+	{
+		return nullptr;
+	}
+
 	FIntRect ViewRect(FIntPoint(0, 0), BackBuffer.GetSizeXY());
 
 	// make a temporary view
 	FSceneViewInitOptions ViewInitOptions;
-	ViewInitOptions.ViewFamily = &ViewFamilyContext;
+	ViewInitOptions.ViewFamily = ViewFamilyContext;
 	ViewInitOptions.SetViewRectangle(ViewRect);
 	ViewInitOptions.ViewOrigin = FVector::ZeroVector;
 	ViewInitOptions.ViewRotationMatrix = FMatrix::Identity;
@@ -194,7 +200,7 @@ static FSceneView& CreateSceneView( FSceneViewFamilyContext& ViewFamilyContext,
 	ViewInitOptions.OverlayColor = FLinearColor::White;
 
 	FSceneView* View = new FSceneView( ViewInitOptions );
-	ViewFamilyContext.Views.Add( View );
+	ViewFamilyContext->Views.Add( View );
 
 	const FIntPoint BufferSize = BackBuffer.GetSizeXY();
 
@@ -215,7 +221,7 @@ static FSceneView& CreateSceneView( FSceneViewFamilyContext& ViewFamilyContext,
 
 	View->ViewUniformBuffer = TUniformBufferRef<FViewUniformShaderParameters>::CreateUniformBufferImmediate(ViewUniformShaderParameters, UniformBuffer_SingleFrame);
 
-	return *View;
+	return View;
 }
 
 void FSlateRHIRenderingPolicy::DrawElements(FRHICommandListImmediate& RHICmdList, FSlateBackBuffer& BackBuffer, const FMatrix& ViewProjectionMatrix, const TArray<FSlateRenderBatch>& RenderBatches, bool bAllowSwitchVerticalAxis)
@@ -234,19 +240,43 @@ void FSlateRHIRenderingPolicy::DrawElements(FRHICommandListImmediate& RHICmdList
 
 	const float DisplayGamma = bGammaCorrect ? GEngine ? GEngine->GetDisplayGamma() : 2.2f : 1.0f;
 
-	FSceneViewFamilyContext SceneViewContext
-	(
-		FSceneViewFamily::ConstructionValues
-		(
-			&BackBuffer,
-			nullptr,
-			DefaultShowFlags
-		)
-		.SetWorldTimes( TimeSeconds, DeltaTimeSeconds, RealTimeSeconds )
-		.SetGammaCorrection( DisplayGamma )
-	);
+	// In order to support MaterialParameterCollections, we need to create multiple FSceneViews for 
+	// each possible Scene that we encounter. The following code creates these as separate arrays, where the first 
+	// N entries map directly to entries from ActiveScenes. The final entry is added to represent the absence of a
+	// valid scene, i.e. a -1 in the SceneIndex parameter of the batch.
+	int32 NumScenes = ResourceManager->GetSceneCount() + 1;
+	TArray<FSceneView*, TInlineAllocator<3> > SceneViews;
+	SceneViews.SetNum(NumScenes);
+	TArray<FSceneViewFamilyContext*, TInlineAllocator<3> > SceneViewFamilyContexts;
+	SceneViewFamilyContexts.SetNum(NumScenes);
+	for (int32 i = 0; i < ResourceManager->GetSceneCount(); i++)
+	{
+		SceneViewFamilyContexts[i] = new  FSceneViewFamilyContext
+			(
+				FSceneViewFamily::ConstructionValues
+				(
+					&BackBuffer,
+					ResourceManager->GetSceneAt(i),
+					DefaultShowFlags
+				)
+				.SetWorldTimes(TimeSeconds, DeltaTimeSeconds, RealTimeSeconds)
+				.SetGammaCorrection(DisplayGamma)
+			);
+		SceneViews[i] = CreateSceneView(SceneViewFamilyContexts[i], BackBuffer, ViewProjectionMatrix);
+	}
 
-	FSceneView* SceneView = NULL;
+	SceneViewFamilyContexts[NumScenes-1] = new  FSceneViewFamilyContext
+		(
+			FSceneViewFamily::ConstructionValues
+			(
+				&BackBuffer,
+				nullptr,
+				DefaultShowFlags
+				)
+			.SetWorldTimes(TimeSeconds, DeltaTimeSeconds, RealTimeSeconds)
+			.SetGammaCorrection(DisplayGamma)
+			);
+	SceneViews[NumScenes-1] = CreateSceneView(SceneViewFamilyContexts[NumScenes-1], BackBuffer, ViewProjectionMatrix);
 
 	TShaderMapRef<FSlateElementVS> GlobalVertexShader(GetGlobalShaderMap(GMaxRHIFeatureLevel));
 
@@ -462,11 +492,27 @@ void FSlateRHIRenderingPolicy::DrawElements(FRHICommandListImmediate& RHICmdList
 			{
 				check(RenderBatch.NumIndices > 0);
 				// Note: This code is only executed if the engine is loaded (in early loading screens attempting to use a material is unsupported
-				if (!SceneView)
+				int ActiveSceneIndex = RenderBatch.SceneIndex;
+
+				// We are assuming at this point that the SceneIndex from the batch is either -1, meaning no scene or a valid scene.
+				// We set up the "no scene" option as the last SceneView in the array above.
+				if (RenderBatch.SceneIndex == -1)
+				{
+					ActiveSceneIndex = NumScenes - 1;
+				}
+				else
+				{
+					ensure(RenderBatch.SceneIndex < ResourceManager->GetSceneCount());
+				}
+
+				// Handle the case where we skipped out early above
+				if (SceneViews[ActiveSceneIndex] == nullptr)
 				{
-					SceneView = &CreateSceneView(SceneViewContext, BackBuffer, ViewProjectionMatrix);
+					continue;
 				}
 
+				const FSceneView& ActiveSceneView = *SceneViews[ActiveSceneIndex];
+
 				FSlateMaterialResource* MaterialShaderResource = (FSlateMaterialResource*)ShaderResource;
 				if( MaterialShaderResource->GetMaterialObject() != nullptr )
 				{
@@ -479,7 +525,7 @@ void FSlateRHIRenderingPolicy::DrawElements(FRHICommandListImmediate& RHICmdList
 #endif
 					FMaterialRenderProxy* MaterialRenderProxy = MaterialShaderResource->GetRenderProxy();
 
-					const FMaterial* Material = MaterialRenderProxy->GetMaterial(SceneView->GetFeatureLevel());
+					const FMaterial* Material = MaterialRenderProxy->GetMaterial(ActiveSceneView.GetFeatureLevel());
 
 					FSlateMaterialShaderPS* PixelShader = GetMaterialPixelShader(Material, ShaderType, DrawEffects);
 
@@ -498,9 +544,9 @@ void FSlateRHIRenderingPolicy::DrawElements(FRHICommandListImmediate& RHICmdList
 
 						VertexShader->SetViewProjection(RHICmdList, ViewProjection);
 						VertexShader->SetVerticalAxisMultiplier(RHICmdList, bAllowSwitchVerticalAxis && RHINeedsToSwitchVerticalAxis(GShaderPlatformForFeatureLevel[GMaxRHIFeatureLevel]) ? -1.0f : 1.0f);
-						VertexShader->SetMaterialShaderParameters(RHICmdList, *SceneView, MaterialRenderProxy, Material);
+						VertexShader->SetMaterialShaderParameters(RHICmdList, ActiveSceneView, MaterialRenderProxy, Material);
 
-						PixelShader->SetParameters(RHICmdList, *SceneView, MaterialRenderProxy, Material, DisplayGamma, ShaderParams.PixelParams);
+						PixelShader->SetParameters(RHICmdList, ActiveSceneView, MaterialRenderProxy, Material, DisplayGamma, ShaderParams.PixelParams);
 						PixelShader->SetDisplayGamma(RHICmdList, (DrawFlags & ESlateBatchDrawFlag::NoGamma) ? 1.0f : DisplayGamma);
 
 						FSlateShaderResource* MaskResource = MaterialShaderResource->GetTextureMaskResource();
@@ -612,6 +658,14 @@ void FSlateRHIRenderingPolicy::DrawElements(FRHICommandListImmediate& RHICmdList
 		}
 
 	}
+
+	for (int i = 0; i < NumScenes; i++)
+	{
+		// Don't need to delete SceneViews b/c the SceneViewFamily will delete it when it goes away.
+		delete SceneViewFamilyContexts[i];
+	}
+	SceneViews.Empty();
+	SceneViewFamilyContexts.Empty();
 }
 
 
@@ -784,3 +838,19 @@ EPrimitiveType FSlateRHIRenderingPolicy::GetRHIPrimitiveType(ESlateDrawPrimitive
 
 };
 
+
+void FSlateRHIRenderingPolicy::AddSceneAt(FSceneInterface* Scene, int32 Index)
+{
+	ResourceManager->AddSceneAt(Scene, Index);
+}
+
+void FSlateRHIRenderingPolicy::ClearScenes()
+{
+	ResourceManager->ClearScenes();
+}
+
+void FSlateRHIRenderingPolicy::FlushGeneratedResources()
+{
+	PostProcessor->ReleaseRenderTargets();
+}
+	
diff --git a/Engine/Source/Runtime/SlateRHIRenderer/Private/SlateRHIRenderingPolicy.h b/Engine/Source/Runtime/SlateRHIRenderer/Private/SlateRHIRenderingPolicy.h
--- a/Engine/Source/Runtime/SlateRHIRenderer/Private/SlateRHIRenderingPolicy.h
+++ b/Engine/Source/Runtime/SlateRHIRenderer/Private/SlateRHIRenderingPolicy.h
@@ -40,6 +40,11 @@ public:
 
 	void SetUseGammaCorrection( bool bInUseGammaCorrection ) { bGammaCorrect = bInUseGammaCorrection; }
 
+	virtual void AddSceneAt(FSceneInterface* Scene, int32 Index) override;
+	virtual void ClearScenes() override;
+
+	virtual void FlushGeneratedResources();
+
 protected:
 	void UpdateVertexAndIndexBuffers(FRHICommandListImmediate& RHICmdList, FSlateBatchData& BatchData, TSlateElementVertexBuffer<FSlateVertex>& VertexBuffer, FSlateElementIndexBuffer& IndexBuffer);
 
diff --git a/Engine/Source/Runtime/SlateRHIRenderer/Private/SlateRHIResourceManager.cpp b/Engine/Source/Runtime/SlateRHIRenderer/Private/SlateRHIResourceManager.cpp
--- a/Engine/Source/Runtime/SlateRHIRenderer/Private/SlateRHIResourceManager.cpp
+++ b/Engine/Source/Runtime/SlateRHIRenderer/Private/SlateRHIResourceManager.cpp
@@ -789,13 +789,9 @@ FSlateMaterialResource* FSlateRHIResourceManager::GetMaterialResource(const UObj
 		
 		DynamicResourceMap.AddMaterialResource(Key, MaterialResource.ToSharedRef());
 	}
-	else if( MaterialResource->GetTextureMaskResource() != TextureMask )
-	{
-		MaterialResource->UpdateMaterial( *Material, ImageSize, TextureMask );
-	}
 	else
 	{
-		MaterialResource->SlateProxy->ActualSize = ImageSize.IntPoint();
+		MaterialResource->UpdateMaterial( *Material, ImageSize, TextureMask );
 	}
 
 	GetAccessedUObjects().Add(const_cast<UMaterialInterface*>( Material ));
@@ -1202,3 +1198,32 @@ UTexture* FSlateRHIResourceManager::GetBadResourceTexture()
 
 	return BadResourceTexture;
 }
+
+
+int32 FSlateRHIResourceManager::GetSceneCount()
+{
+	checkSlow(IsInRenderingThread());
+	return ActiveScenes.Num();
+}
+
+FSceneInterface* FSlateRHIResourceManager::GetSceneAt(int32 Index)
+{
+	checkSlow(IsInRenderingThread());
+	return ActiveScenes[Index];
+}
+
+void FSlateRHIResourceManager::AddSceneAt(FSceneInterface* Scene, int32 Index)
+{
+	checkSlow(IsInRenderingThread());
+	if (ActiveScenes.Num() <= Index)
+	{
+		ActiveScenes.SetNumZeroed(Index + 1);
+	}
+	ActiveScenes[Index] = Scene;
+}
+
+void FSlateRHIResourceManager::ClearScenes()
+{
+	checkSlow(IsInRenderingThread());
+	ActiveScenes.Empty();
+}
diff --git a/Engine/Source/Runtime/SlateRHIRenderer/Private/SlateRHIResourceManager.h b/Engine/Source/Runtime/SlateRHIRenderer/Private/SlateRHIResourceManager.h
--- a/Engine/Source/Runtime/SlateRHIRenderer/Private/SlateRHIResourceManager.h
+++ b/Engine/Source/Runtime/SlateRHIRenderer/Private/SlateRHIResourceManager.h
@@ -22,6 +22,7 @@ class FSlateUTextureResource;
 class ILayoutCache;
 class ISlateStyle;
 class UTexture;
+class FSceneInterface;
 
 /** 
  * Lookup key for materials.  Sometimes the same material is used with different masks so there must be
@@ -242,6 +243,11 @@ public:
 	 */
 	void ReloadTextures();
 
+	int32 GetSceneCount();
+	FSceneInterface* GetSceneAt(int32 Index);
+	void AddSceneAt(FSceneInterface* Scene, int32 Index);
+	void ClearScenes();
+
 private:
 	void ReleaseCachedBuffer(FRHICommandListImmediate& RHICmdList, FCachedRenderBuffers* PooledBuffer);
 
@@ -352,5 +358,8 @@ private:
 	 * to be deleted after the RHI thread is done with them.
 	 */
 	TArray<FCachedRenderBuffers*> PooledBuffersPendingRelease;
+
+
+	TArray<FSceneInterface*> ActiveScenes;
 };
 
diff --git a/Engine/Source/Runtime/StreamingPauseRendering/Private/StreamingPauseRendering.cpp b/Engine/Source/Runtime/StreamingPauseRendering/Private/StreamingPauseRendering.cpp
--- a/Engine/Source/Runtime/StreamingPauseRendering/Private/StreamingPauseRendering.cpp
+++ b/Engine/Source/Runtime/StreamingPauseRendering/Private/StreamingPauseRendering.cpp
@@ -66,6 +66,7 @@ FStreamingPauseRenderingModule::FStreamingPauseRenderingModule()
 	: SceneViewport(nullptr)
 	, ViewportWidget(nullptr)
 	, BackgroundView(nullptr)
+	, bMovieWasStarted(false)
 { }
 
 
@@ -93,75 +94,89 @@ void FStreamingPauseRenderingModule::ShutdownModule()
 
 void FStreamingPauseRenderingModule::BeginStreamingPause( FViewport* GameViewport )
 {
-	check(GameViewport);
-
-	//Create the viewport widget and add a throbber.
-	ViewportWidget = SNew( SViewport )
-		.EnableGammaCorrection(false);
-
-	ViewportWidget->SetContent
-	(
-		SNew(SVerticalBox)
-		+SVerticalBox::Slot()
-		.VAlign(VAlign_Bottom)
-		.HAlign(HAlign_Right)
-		.Padding(FMargin(10.0f))
-		[
-			SNew(SThrobber)
-		]
-	);
-
-	//Render the current scene to a new viewport.
- 	FIntPoint SizeXY = GameViewport->GetSizeXY();
- 	bool bRenderLastFrame = CVarRenderLastFrameInStreamingPause.GetValueOnGameThread() != 0;
- 	bRenderLastFrame = bRenderLastFrame && !( GEngine->StereoRenderingDevice.IsValid() && GameViewport->IsStereoRenderingAllowed() );
- 
- 	if( SizeXY.X > 0 && SizeXY.Y > 0 && bRenderLastFrame )
+	// If a movie is already playing don't bother starting one
+	if(!GetMoviePlayer()->IsMovieCurrentlyPlaying())
 	{
-		FViewportClient* ViewportClient = GameViewport->GetClient();
+		check(GameViewport);
+
+		//Create the viewport widget and add a throbber.
+		ViewportWidget = SNew(SViewport)
+			.EnableGammaCorrection(false);
+
+		ViewportWidget->SetContent
+		(
+			SNew(SVerticalBox)
+			+SVerticalBox::Slot()
+			.VAlign(VAlign_Bottom)
+			.HAlign(HAlign_Right)
+			.Padding(FMargin(10.0f))
+			[
+				SNew(SThrobber)
+			]
+		);
+
+		//Render the current scene to a new viewport.
+		FIntPoint SizeXY = GameViewport->GetSizeXY();
+		bool bRenderLastFrame = CVarRenderLastFrameInStreamingPause.GetValueOnGameThread() != 0;
+		bRenderLastFrame = bRenderLastFrame && !(GEngine->StereoRenderingDevice.IsValid() && GameViewport->IsStereoRenderingAllowed());
+
+		if(SizeXY.X > 0 && SizeXY.Y > 0 && bRenderLastFrame)
+		{
+			FViewportClient* ViewportClient = GameViewport->GetClient();
 
-		SceneViewport = MakeShareable(new FSceneViewport(ViewportClient, ViewportWidget));
+			SceneViewport = MakeShareable(new FSceneViewport(ViewportClient, ViewportWidget));
 
-		SceneViewport->UpdateViewportRHI(false,SizeXY.X,SizeXY.Y, EWindowMode::Fullscreen, PF_Unknown);
+			SceneViewport->UpdateViewportRHI(false, SizeXY.X, SizeXY.Y, EWindowMode::Fullscreen, PF_Unknown);
 
-		SceneViewport->EnqueueBeginRenderFrame();
+			SceneViewport->EnqueueBeginRenderFrame();
 
-		FCanvas Canvas(SceneViewport.Get(), nullptr, ViewportClient->GetWorld(), ViewportClient->GetWorld()->FeatureLevel);
-		{
-			ViewportClient->Draw(SceneViewport.Get(), &Canvas);
-		}
-		Canvas.Flush_GameThread();
+			FCanvas Canvas(SceneViewport.Get(), nullptr, ViewportClient->GetWorld(), ViewportClient->GetWorld()->FeatureLevel);
+			{
+				ViewportClient->Draw(SceneViewport.Get(), &Canvas);
+			}
+			Canvas.Flush_GameThread();
 
 
-		ENQUEUE_UNIQUE_RENDER_COMMAND_ONEPARAMETER(
-			EndDrawingCommand,
-			FViewport*,Viewport, SceneViewport.Get(),
-		{
-			Viewport->EndRenderFrame(RHICmdList, false, false);
-		});
+			ENQUEUE_UNIQUE_RENDER_COMMAND_ONEPARAMETER(
+				EndDrawingCommand,
+				FViewport*, Viewport, SceneViewport.Get(),
+				{
+					Viewport->EndRenderFrame(RHICmdList, false, false);
+				});
 
-		BackgroundView = MakeShareable(new FBackgroundView(SceneViewport->GetRenderTargetTexture(), SizeXY));
-		ViewportWidget->SetViewportInterface(BackgroundView.ToSharedRef());
+			BackgroundView = MakeShareable(new FBackgroundView(SceneViewport->GetRenderTargetTexture(), SizeXY));
+			ViewportWidget->SetViewportInterface(BackgroundView.ToSharedRef());
+
+		}
 
-	}
 
-	//Create the loading screen and begin rendering the widget.
-	FLoadingScreenAttributes LoadingScreen;
-	LoadingScreen.bAutoCompleteWhenLoadingCompletes = true; 
-	LoadingScreen.WidgetLoadingScreen = ViewportWidget; // SViewport from above
-	GetMoviePlayer()->SetupLoadingScreen(LoadingScreen);			
-	GetMoviePlayer()->PlayMovie();
+		//Create the loading screen and begin rendering the widget.
+		FLoadingScreenAttributes LoadingScreen;
+		LoadingScreen.bAutoCompleteWhenLoadingCompletes = true;
+		LoadingScreen.WidgetLoadingScreen = ViewportWidget; // SViewport from above
+		GetMoviePlayer()->SetupLoadingScreen(LoadingScreen);
+		GetMoviePlayer()->PlayMovie();
+
+		// We started playing a movie
+		bMovieWasStarted = true;
+	}
 }
 
 
 void FStreamingPauseRenderingModule::EndStreamingPause()
 {
-	//Stop rendering the loading screen and resume
-	GetMoviePlayer()->WaitForMovieToFinish();
-
-	ViewportWidget = nullptr;
-	SceneViewport = nullptr;
-	BackgroundView = nullptr;
-	
-	FlushRenderingCommands();
+	// Only wait for the movie to finish if we were the one to start it.
+	if(bMovieWasStarted)
+	{
+		//Stop rendering the loading screen and resume
+		GetMoviePlayer()->WaitForMovieToFinish();
+
+		ViewportWidget = nullptr;
+		SceneViewport = nullptr;
+		BackgroundView = nullptr;
+
+		FlushRenderingCommands();
+
+		bMovieWasStarted = false;
+	}
 }
diff --git a/Engine/Source/Runtime/StreamingPauseRendering/Public/StreamingPauseRendering.h b/Engine/Source/Runtime/StreamingPauseRendering/Public/StreamingPauseRendering.h
--- a/Engine/Source/Runtime/StreamingPauseRendering/Public/StreamingPauseRendering.h
+++ b/Engine/Source/Runtime/StreamingPauseRendering/Public/StreamingPauseRendering.h
@@ -53,4 +53,10 @@ public:
 
 	/** Delegate providing default functionality for ending streaming pause. */
 	FEndStreamingPauseDelegate EndDelegate;
+
+	/**
+	 * If a movie was started by BeginStreamingPause.  
+	 * This could be false if a movie was already playing when BeginStreamingPause was called
+	 */
+	bool bMovieWasStarted;
 };
diff --git a/Engine/Source/Runtime/UMG/Private/Animation/UMGSequencePlayer.cpp b/Engine/Source/Runtime/UMG/Private/Animation/UMGSequencePlayer.cpp
--- a/Engine/Source/Runtime/UMG/Private/Animation/UMGSequencePlayer.cpp
+++ b/Engine/Source/Runtime/UMG/Private/Animation/UMGSequencePlayer.cpp
@@ -35,7 +35,7 @@ void UUMGSequencePlayer::Tick(float DeltaTime)
 	{
 		const double AnimationLength = CurrentPlayRange.Size<double>();
 
-		const double LastTimePosition = TimeCursorPosition;
+		double LastTimePosition = TimeCursorPosition;
 		TimeCursorPosition += bIsPlayingForward ? DeltaTime * PlaybackSpeed : -DeltaTime * PlaybackSpeed;
 
 		// Check if we crossed over bounds
@@ -71,6 +71,7 @@ void UUMGSequencePlayer::Tick(float DeltaTime)
 				else
 				{
 					TimeCursorPosition += AnimationLength;
+					LastTimePosition = TimeCursorPosition;
 				}
 			}
 		}
@@ -90,6 +91,7 @@ void UUMGSequencePlayer::Tick(float DeltaTime)
 				else
 				{
 					TimeCursorPosition -= AnimationLength;
+					LastTimePosition = TimeCursorPosition;
 				}
 			}
 		}
@@ -111,7 +113,7 @@ void UUMGSequencePlayer::Tick(float DeltaTime)
 		if (RootTemplateInstance.IsValid())
 		{			
 			const FMovieSceneContext Context(
-				FMovieSceneEvaluationRange(TimeCursorPosition + AnimationStartOffset, (TimeCursorPosition - LastTimePosition) + AnimationStartOffset),
+				FMovieSceneEvaluationRange(TimeCursorPosition + AnimationStartOffset, LastTimePosition + AnimationStartOffset),
 				PlayerStatus);
 			RootTemplateInstance.Evaluate(Context, *this);
 		}
diff --git a/Engine/Source/Runtime/UMG/Private/Components/ComboBox.cpp b/Engine/Source/Runtime/UMG/Private/Components/ComboBox.cpp
--- a/Engine/Source/Runtime/UMG/Private/Components/ComboBox.cpp
+++ b/Engine/Source/Runtime/UMG/Private/Components/ComboBox.cpp
@@ -10,6 +10,7 @@
 UComboBox::UComboBox(const FObjectInitializer& ObjectInitializer)
 	: Super(ObjectInitializer)
 {
+	bIsFocusable = true;
 }
 
 void UComboBox::ReleaseSlateResources(bool bReleaseChildren)
@@ -27,7 +28,8 @@ TSharedRef<SWidget> UComboBox::RebuildWidget()
 	MyComboBox =
 		SNew(SComboBox<UObject*>)
 		.OptionsSource(&Items)
-		.OnGenerateWidget(BIND_UOBJECT_DELEGATE(SComboBox<UObject*>::FOnGenerateWidget, HandleGenerateWidget));
+		.OnGenerateWidget(BIND_UOBJECT_DELEGATE(SComboBox<UObject*>::FOnGenerateWidget, HandleGenerateWidget))
+		.IsFocusable(bIsFocusable);
 
 	return MyComboBox.ToSharedRef();
 }
diff --git a/Engine/Source/Runtime/UMG/Private/Components/ComboBoxString.cpp b/Engine/Source/Runtime/UMG/Private/Components/ComboBoxString.cpp
--- a/Engine/Source/Runtime/UMG/Private/Components/ComboBoxString.cpp
+++ b/Engine/Source/Runtime/UMG/Private/Components/ComboBoxString.cpp
@@ -32,6 +32,7 @@ UComboBoxString::UComboBoxString(const FObjectInitializer& ObjectInitializer)
 	ItemStyle.DropIndicator_Below.TintColor = ItemStyle.DropIndicator_Below.TintColor.GetSpecifiedColor();
 
 	ForegroundColor = FLinearColor::Black;
+	bIsFocusable = true;
 
 	ContentPadding = FMargin(4.0, 2.0);
 	MaxListHeight = 450.0f;
@@ -93,6 +94,7 @@ TSharedRef<SWidget> UComboBoxString::RebuildWidget()
 		.OnGenerateWidget(BIND_UOBJECT_DELEGATE(SComboBox< TSharedPtr<FString> >::FOnGenerateWidget, HandleGenerateWidget))
 		.OnSelectionChanged(BIND_UOBJECT_DELEGATE(SComboBox< TSharedPtr<FString> >::FOnSelectionChanged, HandleSelectionChanged))
 		.OnComboBoxOpening(BIND_UOBJECT_DELEGATE(FOnComboBoxOpening, HandleOpening))
+		.IsFocusable(bIsFocusable)
 		[
 			SAssignNew(ComoboBoxContent, SBox)
 		];
diff --git a/Engine/Source/Runtime/UMG/Private/Components/PanelWidget.cpp b/Engine/Source/Runtime/UMG/Private/Components/PanelWidget.cpp
--- a/Engine/Source/Runtime/UMG/Private/Components/PanelWidget.cpp
+++ b/Engine/Source/Runtime/UMG/Private/Components/PanelWidget.cpp
@@ -84,8 +84,12 @@ bool UPanelWidget::RemoveChildAt(int32 Index)
 
 	OnSlotRemoved(PanelSlot);
 
-	const bool bReleaseChildren = true;
-	PanelSlot->ReleaseSlateResources(bReleaseChildren);
+	// If the child is a UserWidget, we should let it manage it's own slate resources instead of forcing a clear here. This fixes issues such as UE-39106
+	if (PanelSlot->Content && !PanelSlot->Content->IsA<UUserWidget>())
+	{
+		const bool bReleaseChildren = true;
+		PanelSlot->ReleaseSlateResources(bReleaseChildren);
+	}
 
 	PanelSlot->Parent = nullptr;
 	PanelSlot->Content = nullptr;
diff --git a/Engine/Source/Runtime/UMG/Private/Components/RetainerBox.cpp b/Engine/Source/Runtime/UMG/Private/Components/RetainerBox.cpp
--- a/Engine/Source/Runtime/UMG/Private/Components/RetainerBox.cpp
+++ b/Engine/Source/Runtime/UMG/Private/Components/RetainerBox.cpp
@@ -69,7 +69,7 @@ TSharedRef<SWidget> URetainerBox::RebuildWidget()
 		;
 
 	MyRetainerWidget->SetRetainedRendering(IsDesignTime() ? false : true);
-
+	
 	if ( GetChildrenCount() > 0 )
 	{
 		MyRetainerWidget->SetContent(GetContentSlot()->Content ? GetContentSlot()->Content->TakeWidget() : SNullWidget::NullWidget);
@@ -84,6 +84,7 @@ void URetainerBox::SynchronizeProperties()
 
 	MyRetainerWidget->SetEffectMaterial(EffectMaterial);
 	MyRetainerWidget->SetTextureParameter(TextureParameter);
+	MyRetainerWidget->SetWorld(GetWorld());
 }
 
 void URetainerBox::OnSlotAdded(UPanelSlot* InSlot)
diff --git a/Engine/Source/Runtime/UMG/Private/Components/Widget.cpp b/Engine/Source/Runtime/UMG/Private/Components/Widget.cpp
--- a/Engine/Source/Runtime/UMG/Private/Components/Widget.cpp
+++ b/Engine/Source/Runtime/UMG/Private/Components/Widget.cpp
@@ -20,8 +20,8 @@
 #include "Slate/SObjectWidget.h"
 #include "Blueprint/WidgetTree.h"
 #include "UMGStyle.h"
-
 #include "Types/ReflectionMetadata.h"
+#include "PropertyLocalizationDataGathering.h"
 
 #define LOCTEXT_NAMESPACE "UMG"
 
@@ -114,6 +114,30 @@ private:
 };
 
 
+#if WITH_EDITORONLY_DATA
+namespace
+{
+	void GatherWidgetForLocalization(const UObject* const Object, FPropertyLocalizationDataGatherer& PropertyLocalizationDataGatherer, const EPropertyLocalizationGathererTextFlags GatherTextFlags)
+	{
+		const UWidget* const Widget = CastChecked<UWidget>(Object);
+
+		EPropertyLocalizationGathererTextFlags WidgetGatherTextFlags = GatherTextFlags;
+
+		// If we've instanced this widget from another asset, then we only want to process the widget itself (to process any overrides against the archetype), but skip all of its children
+		if (UObject* WidgetGenerator = Widget->WidgetGeneratedBy.Get())
+		{
+			if (WidgetGenerator->GetOutermost() != Widget->GetOutermost())
+			{
+				WidgetGatherTextFlags |= EPropertyLocalizationGathererTextFlags::SkipSubObjects;
+			}
+		}
+
+		PropertyLocalizationDataGatherer.GatherLocalizationDataFromObject(Widget, WidgetGatherTextFlags);
+	}
+}
+#endif
+
+
 /////////////////////////////////////////////////////
 // UWidget
 
@@ -130,6 +154,10 @@ UWidget::UWidget(const FObjectInitializer& ObjectInitializer)
 	Visiblity_DEPRECATED = Visibility = ESlateVisibility::Visible;	
 	RenderTransformPivot = FVector2D(0.5f, 0.5f);
 	Cursor = EMouseCursor::Default;
+
+#if WITH_EDITORONLY_DATA
+	{ static const FAutoRegisterLocalizationDataGatheringCallback AutomaticRegistrationOfLocalizationGatherer(UWidget::StaticClass(), &GatherWidgetForLocalization); }
+#endif
 }
 
 void UWidget::SetRenderTransform(FWidgetTransform Transform)
@@ -522,6 +550,17 @@ void UWidget::RemoveFromParent()
 	}
 }
 
+const FGeometry& UWidget::GetCachedGeometry() const
+{
+	TSharedPtr<SWidget> SafeWidget = GetCachedWidget();
+	if ( SafeWidget.IsValid() )
+	{
+		return SafeWidget->GetCachedGeometry();
+	}
+
+	return SNullWidget::NullWidget->GetCachedGeometry();
+}
+
 #if !(UE_BUILD_SHIPPING || UE_BUILD_TEST)
 void UWidget::VerifySynchronizeProperties()
 {
diff --git a/Engine/Source/Runtime/UMG/Private/Components/WidgetComponent.cpp b/Engine/Source/Runtime/UMG/Private/Components/WidgetComponent.cpp
--- a/Engine/Source/Runtime/UMG/Private/Components/WidgetComponent.cpp
+++ b/Engine/Source/Runtime/UMG/Private/Components/WidgetComponent.cpp
@@ -20,6 +20,7 @@
 #include "Slate/SGameLayerManager.h"
 #include "Slate/WidgetRenderer.h"
 #include "Slate/SWorldWidgetScreenLayer.h"
+#include "SViewport.h"
 
 DECLARE_CYCLE_STAT(TEXT("3DHitTesting"), STAT_Slate3DHitTesting, STATGROUP_Slate);
 
@@ -96,6 +97,173 @@ private:
 };
 
 
+
+
+/**
+* The hit tester used by all Widget Component objects.
+*/
+class FWidget3DHitTester : public ICustomHitTestPath
+{
+public:
+	FWidget3DHitTester( UWorld* InWorld )
+		: World( InWorld )
+		, CachedFrame(-1)
+	{}
+
+	// ICustomHitTestPath implementation
+	virtual TArray<FWidgetAndPointer> GetBubblePathAndVirtualCursors(const FGeometry& InGeometry, FVector2D DesktopSpaceCoordinate, bool bIgnoreEnabledStatus) const override
+	{
+		SCOPE_CYCLE_COUNTER(STAT_Slate3DHitTesting);
+
+		if( World.IsValid() /*&& ensure( World->IsGameWorld() )*/ )
+		{
+			UWorld* SafeWorld = World.Get();
+			if ( SafeWorld )
+			{
+				ULocalPlayer* const TargetPlayer = GEngine->GetLocalPlayerFromControllerId(SafeWorld, 0);
+
+				if( TargetPlayer && TargetPlayer->PlayerController )
+				{
+					FVector2D LocalMouseCoordinate = InGeometry.AbsoluteToLocal(DesktopSpaceCoordinate) * InGeometry.Scale;
+
+					if ( UPrimitiveComponent* HitComponent = GetHitResultAtScreenPositionAndCache(TargetPlayer->PlayerController, LocalMouseCoordinate) )
+					{
+						if ( UWidgetComponent* WidgetComponent = Cast<UWidgetComponent>(HitComponent) )
+						{
+							if ( WidgetComponent->GetReceiveHardwareInput() )
+							{
+								if ( WidgetComponent->GetDrawSize().X != 0 && WidgetComponent->GetDrawSize().Y != 0 )
+								{
+									// Get the "forward" vector based on the current rotation system.
+									const FVector ForwardVector = WidgetComponent->GetForwardVector();
+
+									// Make sure the player is interacting with the front of the widget
+									if ( FVector::DotProduct(ForwardVector, CachedHitResult.ImpactPoint - CachedHitResult.TraceStart) < 0.f )
+									{
+										return WidgetComponent->GetHitWidgetPath(CachedHitResult.Location, bIgnoreEnabledStatus);
+									}
+								}
+							}
+						}
+					}
+				}
+			}
+		}
+
+		return TArray<FWidgetAndPointer>();
+	}
+
+	virtual void ArrangeChildren( FArrangedChildren& ArrangedChildren ) const override
+	{
+		for( TWeakObjectPtr<UWidgetComponent> Component : RegisteredComponents )
+		{
+			UWidgetComponent* WidgetComponent = Component.Get();
+			// Check if visible;
+			if ( WidgetComponent && WidgetComponent->GetSlateWindow().IsValid() )
+			{
+				FGeometry WidgetGeom;
+
+				ArrangedChildren.AddWidget( FArrangedWidget( WidgetComponent->GetSlateWindow().ToSharedRef(), WidgetGeom.MakeChild( WidgetComponent->GetDrawSize(), FSlateLayoutTransform() ) ) );
+			}
+		}
+	}
+
+	virtual TSharedPtr<struct FVirtualPointerPosition> TranslateMouseCoordinateFor3DChild( const TSharedRef<SWidget>& ChildWidget, const FGeometry& ViewportGeometry, const FVector2D& ScreenSpaceMouseCoordinate, const FVector2D& LastScreenSpaceMouseCoordinate ) const override
+	{
+		if ( World.IsValid() && ensure(World->IsGameWorld()) )
+		{
+			ULocalPlayer* const TargetPlayer = GEngine->GetLocalPlayerFromControllerId(World.Get(), 0);
+			if ( TargetPlayer && TargetPlayer->PlayerController )
+			{
+				FVector2D LocalMouseCoordinate = ViewportGeometry.AbsoluteToLocal(ScreenSpaceMouseCoordinate) * ViewportGeometry.Scale;
+
+				// Check for a hit against any widget components in the world
+				for ( TWeakObjectPtr<UWidgetComponent> Component : RegisteredComponents )
+				{
+					UWidgetComponent* WidgetComponent = Component.Get();
+					// Check if visible;
+					if ( WidgetComponent && WidgetComponent->GetSlateWindow() == ChildWidget )
+					{
+						if ( UPrimitiveComponent* HitComponent = GetHitResultAtScreenPositionAndCache(TargetPlayer->PlayerController, LocalMouseCoordinate) )
+						{
+							if ( WidgetComponent->GetReceiveHardwareInput() )
+							{
+								if ( WidgetComponent->GetDrawSize().X != 0 && WidgetComponent->GetDrawSize().Y != 0 )
+								{
+									if ( WidgetComponent == HitComponent )
+									{
+										TSharedPtr<FVirtualPointerPosition> VirtualCursorPos = MakeShareable(new FVirtualPointerPosition);
+
+										FVector2D LocalHitLocation;
+										WidgetComponent->GetLocalHitLocation(CachedHitResult.Location, LocalHitLocation);
+
+										VirtualCursorPos->CurrentCursorPosition = LocalHitLocation;
+										VirtualCursorPos->LastCursorPosition = LocalHitLocation;
+
+										return VirtualCursorPos;
+									}
+								}
+							}
+						}
+					}
+				}
+			}
+		}
+
+		return nullptr;
+	}
+	// End ICustomHitTestPath
+
+	UPrimitiveComponent* GetHitResultAtScreenPositionAndCache(APlayerController* PlayerController, FVector2D ScreenPosition) const
+	{
+		UPrimitiveComponent* HitComponent = nullptr;
+
+		if ( GFrameNumber != CachedFrame || CachedScreenPosition != ScreenPosition )
+		{
+			CachedFrame = GFrameNumber;
+			CachedScreenPosition = ScreenPosition;
+
+			if ( PlayerController )
+			{
+				if ( PlayerController->GetHitResultAtScreenPosition(ScreenPosition, ECC_Visibility, true, CachedHitResult) )
+				{
+					return CachedHitResult.Component.Get();
+				}
+			}
+		}
+		else
+		{
+			return CachedHitResult.Component.Get();
+		}
+
+		return nullptr;
+	}
+
+	void RegisterWidgetComponent( UWidgetComponent* InComponent )
+	{
+		RegisteredComponents.AddUnique( InComponent );
+	}
+
+	void UnregisterWidgetComponent( UWidgetComponent* InComponent )
+	{
+		RegisteredComponents.RemoveSingleSwap( InComponent );
+	}
+
+	uint32 GetNumRegisteredComponents() const { return RegisteredComponents.Num(); }
+	
+	UWorld* GetWorld() const { return World.Get(); }
+
+private:
+	TArray< TWeakObjectPtr<UWidgetComponent> > RegisteredComponents;
+	TWeakObjectPtr<UWorld> World;
+
+	mutable int64 CachedFrame;
+	mutable FVector2D CachedScreenPosition;
+	mutable FHitResult CachedHitResult;
+};
+
+
+
 /** Represents a billboard sprite to the scene manager. */
 class FWidget3DSceneProxy : public FPrimitiveSceneProxy
 {
@@ -130,19 +298,21 @@ public:
 
 		Collector.RegisterOneFrameMaterialProxy(WireframeMaterialInstance);
 
-		FMaterialRenderProxy* MaterialProxy = nullptr;
+		FMaterialRenderProxy* ParentMaterialProxy = nullptr;
 		if ( bWireframe )
 		{
-			MaterialProxy = WireframeMaterialInstance;
+			ParentMaterialProxy = WireframeMaterialInstance;
 		}
 		else
 		{
-			MaterialProxy = MaterialInstance->GetRenderProxy(IsSelected());
+			ParentMaterialProxy = MaterialInstance->GetRenderProxy(IsSelected());
 		}
 #else
-		FMaterialRenderProxy* MaterialProxy = MaterialInstance->GetRenderProxy(IsSelected());
+		FMaterialRenderProxy* ParentMaterialProxy = MaterialInstance->GetRenderProxy(IsSelected());
 #endif
 
+		//FSpriteTextureOverrideRenderProxy* TextureOverrideMaterialProxy = new FSpriteTextureOverrideRenderProxy(ParentMaterialProxy,
+
 		const FMatrix& ViewportLocalToWorld = GetLocalToWorld();
 	
 		if( RenderTarget )
@@ -173,7 +343,7 @@ public:
 							MeshBuilder.AddTriangle(VertexIndices[0], VertexIndices[1], VertexIndices[2]);
 							MeshBuilder.AddTriangle(VertexIndices[0], VertexIndices[2], VertexIndices[3]);
 
-							MeshBuilder.GetMesh(ViewportLocalToWorld, MaterialProxy, SDPG_World, false, true, ViewIndex, Collector);
+							MeshBuilder.GetMesh(ViewportLocalToWorld, ParentMaterialProxy, SDPG_World, false, true, ViewIndex, Collector);
 						}
 					}
 				}
@@ -250,7 +420,7 @@ public:
 								LastTangentY = TangentY;
 								LastTangentZ = TangentZ;
 							}
-							MeshBuilder.GetMesh(ViewportLocalToWorld, MaterialProxy, SDPG_World, false, true, ViewIndex, Collector);
+							MeshBuilder.GetMesh(ViewportLocalToWorld, ParentMaterialProxy, SDPG_World, false, true, ViewIndex, Collector);
 						}
 					}
 				}
@@ -379,6 +549,7 @@ UWidgetComponent::UWidgetComponent( const FObjectInitializer& PCIP )
 	, bRedrawRequested(true)
 	, RedrawTime(0)
 	, LastWidgetRenderTime(0)
+	, bReceiveHardwareInput(false)
 	, bWindowFocusable(true)
 	, BackgroundColor( FLinearColor::Transparent )
 	, TintColorAndOpacity( FLinearColor::White )
@@ -519,6 +690,12 @@ void UWidgetComponent::OnRegister()
 	{
 		if ( Space != EWidgetSpace::Screen )
 		{
+			if ( bReceiveHardwareInput && GetWorld()->IsGameWorld() )
+			{
+				TSharedPtr<SViewport> GameViewportWidget = GEngine->GetGameViewportWidget();
+				RegisterHitTesterWithViewport(GameViewportWidget);
+			}
+
 			if ( !WidgetRenderer.IsValid() && !GUsingNullRHI )
 			{
 				WidgetRenderer = MakeShareable(new FWidgetRenderer());
@@ -532,8 +709,62 @@ void UWidgetComponent::OnRegister()
 #endif // !UE_SERVER
 }
 
+void UWidgetComponent::RegisterHitTesterWithViewport(TSharedPtr<SViewport> ViewportWidget)
+{
+#if !UE_SERVER
+	if ( ViewportWidget.IsValid() )
+	{
+		TSharedPtr<ICustomHitTestPath> CustomHitTestPath = ViewportWidget->GetCustomHitTestPath();
+		if ( !CustomHitTestPath.IsValid() )
+		{
+			CustomHitTestPath = MakeShareable(new FWidget3DHitTester(GetWorld()));
+			ViewportWidget->SetCustomHitTestPath(CustomHitTestPath);
+		}
+
+		TSharedPtr<FWidget3DHitTester> Widget3DHitTester = StaticCastSharedPtr<FWidget3DHitTester>(CustomHitTestPath);
+		if ( Widget3DHitTester->GetWorld() == GetWorld() )
+		{
+			Widget3DHitTester->RegisterWidgetComponent(this);
+		}
+	}
+#endif
+
+}
+
+void UWidgetComponent::UnregisterHitTesterWithViewport(TSharedPtr<SViewport> ViewportWidget)
+{
+#if !UE_SERVER
+	if ( bReceiveHardwareInput )
+	{
+		TSharedPtr<ICustomHitTestPath> CustomHitTestPath = ViewportWidget->GetCustomHitTestPath();
+		if ( CustomHitTestPath.IsValid() )
+		{
+			TSharedPtr<FWidget3DHitTester> WidgetHitTestPath = StaticCastSharedPtr<FWidget3DHitTester>(CustomHitTestPath);
+
+			WidgetHitTestPath->UnregisterWidgetComponent(this);
+
+			if ( WidgetHitTestPath->GetNumRegisteredComponents() == 0 )
+			{
+				ViewportWidget->SetCustomHitTestPath(nullptr);
+			}
+		}
+	}
+#endif
+}
+
 void UWidgetComponent::OnUnregister()
 {
+#if !UE_SERVER
+	if ( GetWorld()->IsGameWorld() )
+	{
+		TSharedPtr<SViewport> GameViewportWidget = GEngine->GetGameViewportWidget();
+		if ( GameViewportWidget.IsValid() )
+		{
+			UnregisterHitTesterWithViewport(GameViewportWidget);
+		}
+	}
+#endif
+
 #if WITH_EDITOR
 	if (!GetWorld()->IsGameWorld())
 	{
@@ -569,7 +800,7 @@ void UWidgetComponent::RegisterWindow()
 {
 	if ( SlateWindow.IsValid() )
 	{
-		if ( FSlateApplication::IsInitialized() )
+		if (!bReceiveHardwareInput && FSlateApplication::IsInitialized() )
 		{
 			FSlateApplication::Get().RegisterVirtualWindow(SlateWindow.ToSharedRef());
 		}
@@ -580,7 +811,7 @@ void UWidgetComponent::UnregisterWindow()
 {
 	if ( SlateWindow.IsValid() )
 	{
-		if ( FSlateApplication::IsInitialized() )
+		if ( !bReceiveHardwareInput && FSlateApplication::IsInitialized() )
 		{
 			FSlateApplication::Get().UnregisterVirtualWindow(SlateWindow.ToSharedRef());
 		}
@@ -607,7 +838,11 @@ void UWidgetComponent::TickComponent(float DeltaTime, enum ELevelTick TickType,
 	    {
 			if ( ShouldDrawWidget() )
 		    {
-			    DrawWidgetToRenderTarget(DeltaTime);
+				// Calculate the actual delta time since we last drew, this handles the case where we're ticking when
+				// the world is paused, this also takes care of the case where the widget component is rendering at
+				// a different rate than the rest of the world.
+				const float DeltaTimeFromLastDraw = LastWidgetRenderTime == 0 ? 0 : (FApp::GetCurrentTime() - LastWidgetRenderTime );
+			    DrawWidgetToRenderTarget(DeltaTimeFromLastDraw);
 		    }
 	    }
 	    else
@@ -673,7 +908,7 @@ bool UWidgetComponent::ShouldDrawWidget() const
 		// If we don't tick when off-screen, don't bother ticking if it hasn't been rendered recently
 		if ( TickWhenOffscreen || GetWorld()->TimeSince(LastRenderTime) <= RenderTimeThreshold )
 		{
-			if ( GetWorld()->TimeSince(LastWidgetRenderTime) >= RedrawTime )
+			if ( ( FApp::GetCurrentTime() - LastWidgetRenderTime) >= RedrawTime )
 			{
 				return bManuallyRedraw ? bRedrawRequested : true;
 			}
@@ -740,7 +975,7 @@ void UWidgetComponent::DrawWidgetToRenderTarget(float DeltaTime)
 		CurrentDrawSize,
 		DeltaTime);
 
-	LastWidgetRenderTime = GetWorld()->TimeSeconds;
+	LastWidgetRenderTime = FApp::GetCurrentTime();
 }
 
 float UWidgetComponent::ComputeComponentWidth() const
@@ -834,7 +1069,7 @@ void UWidgetComponent::ApplyComponentInstanceData(FWidgetComponentInstanceData*
 	}
 
 	RenderTarget = WidgetInstanceData->RenderTarget;
-	if( MaterialInstance && RenderTarget )
+	if ( MaterialInstance && RenderTarget )
 	{
 		MaterialInstance->SetTextureParameterValue("SlateUI", RenderTarget);
 	}
diff --git a/Engine/Source/Runtime/UMG/Private/Components/WidgetInteractionComponent.cpp b/Engine/Source/Runtime/UMG/Private/Components/WidgetInteractionComponent.cpp
--- a/Engine/Source/Runtime/UMG/Private/Components/WidgetInteractionComponent.cpp
+++ b/Engine/Source/Runtime/UMG/Private/Components/WidgetInteractionComponent.cpp
@@ -23,6 +23,7 @@ UWidgetInteractionComponent::UWidgetInteractionComponent(const FObjectInitialize
 	, DebugColor(FLinearColor::Red)
 {
 	PrimaryComponentTick.bCanEverTick = true;
+	PrimaryComponentTick.bTickEvenWhenPaused = true;
 	TraceChannel = ECC_Visibility;
 	bAutoActivate = true;
 
@@ -221,6 +222,18 @@ void UWidgetInteractionComponent::GetRelatedComponentsToIgnoreInAutomaticHitTest
 	}
 }
 
+bool UWidgetInteractionComponent::CanInteractWithComponent(UWidgetComponent* Component) const
+{
+	bool bCanInteract = false;
+
+	if (Component)
+	{
+		bCanInteract = !GetWorld()->IsPaused() || Component->PrimaryComponentTick.bTickEvenWhenPaused;
+	}
+
+	return bCanInteract;
+}
+
 void UWidgetInteractionComponent::SimulatePointerMovement()
 {
 	bIsHoveredWidgetInteractable = false;
diff --git a/Engine/Source/Runtime/UMG/Private/Slate/SRetainerWidget.cpp b/Engine/Source/Runtime/UMG/Private/Slate/SRetainerWidget.cpp
--- a/Engine/Source/Runtime/UMG/Private/Slate/SRetainerWidget.cpp
+++ b/Engine/Source/Runtime/UMG/Private/Slate/SRetainerWidget.cpp
@@ -6,6 +6,7 @@
 #include "Materials/MaterialInstanceDynamic.h"
 #include "Engine/TextureRenderTarget2D.h"
 #include "Framework/Application/SlateApplication.h"
+#include "Engine/World.h"
 
 
 DECLARE_CYCLE_STAT(TEXT("Retainer Widget Tick"), STAT_SlateRetainerWidgetTick, STATGROUP_Slate);
@@ -198,6 +199,11 @@ void SRetainerWidget::SetTextureParameter(FName TextureParameter)
 {
 	DynamicEffectTextureParameter = TextureParameter;
 }
+ 
+void SRetainerWidget::SetWorld(UWorld* World)
+{
+	OuterWorld = World;
+}
 
 void SRetainerWidget::AddReferencedObjects(FReferenceCollector& Collector)
 {
@@ -229,9 +235,21 @@ void SRetainerWidget::OnTickRetainers(float DeltaTime)
 	// we should not be added to tick if we're not rendering
 	checkSlow(FApp::CanEverRender());
 
+
 	const bool bShouldRenderAnything = IsAnythingVisibleToRender();
 	if ( bRenderingOffscreen && bShouldRenderAnything )
 	{
+
+		// In order to get material parameter collections to function properly, we need the current world's Scene
+		// properly propagated through to any widgets that depend on that functionality. The SceneViewport and RetainerWidget the 
+		// only location where this information exists in Slate, so we push the current scene onto the current
+		// Slate application so that we can leverage it in later calls.
+		UWorld* TickWorld = OuterWorld.Get();
+		if (TickWorld && TickWorld->Scene && IsInGameThread())
+		{
+			FSlateApplication::Get().GetRenderer()->RegisterCurrentScene(TickWorld->Scene);
+		}
+
 		SCOPE_CYCLE_COUNTER( STAT_SlateRetainerWidgetTick );
 		if ( LastTickedFrame != GFrameCounter && ( GFrameCounter % PhaseCount ) == Phase )
 		{
diff --git a/Engine/Source/Runtime/UMG/Private/Slate/SWorldWidgetScreenLayer.cpp b/Engine/Source/Runtime/UMG/Private/Slate/SWorldWidgetScreenLayer.cpp
--- a/Engine/Source/Runtime/UMG/Private/Slate/SWorldWidgetScreenLayer.cpp
+++ b/Engine/Source/Runtime/UMG/Private/Slate/SWorldWidgetScreenLayer.cpp
@@ -94,7 +94,7 @@ void SWorldWidgetScreenLayer::Tick(const FGeometry& AllottedGeometry, const doub
 							FVector2D ComponentDrawSize = Entry.WidgetComponent->GetDrawSize();
 							FVector2D ComponentPivot = Entry.WidgetComponent->GetPivot();
 
-							CanvasSlot->AutoSize(ComponentDrawSize.IsZero());
+							CanvasSlot->AutoSize(ComponentDrawSize.IsZero() || Entry.WidgetComponent->GetDrawAtDesiredSize());
 							CanvasSlot->Offset(FMargin(ScreenPosition.X, ScreenPosition.Y, ComponentDrawSize.X, ComponentDrawSize.Y));
 							CanvasSlot->Anchors(FAnchors(0, 0, 0, 0));
 							CanvasSlot->Alignment(ComponentPivot);
diff --git a/Engine/Source/Runtime/UMG/Private/SlateBlueprintLibrary.cpp b/Engine/Source/Runtime/UMG/Private/SlateBlueprintLibrary.cpp
--- a/Engine/Source/Runtime/UMG/Private/SlateBlueprintLibrary.cpp
+++ b/Engine/Source/Runtime/UMG/Private/SlateBlueprintLibrary.cpp
@@ -6,6 +6,11 @@
 #include "Engine/Engine.h"
 
 #include "Engine/UserInterfaceSettings.h"
+#include "Slate/SlateBrushAsset.h"
+#include "Runtime/Engine/Classes/Engine/UserInterfaceSettings.h"
+#include "SlateBlueprintLibrary.h"
+#include "Slate/SceneViewport.h"
+#include "Blueprint/WidgetLayoutLibrary.h"
 
 #define LOCTEXT_NAMESPACE "UMG"
 
@@ -55,17 +60,15 @@ void USlateBlueprintLibrary::AbsoluteToViewport(UObject* WorldContextObject, FVe
 	{
 		if ( UGameViewportClient* ViewportClient = World->GetGameViewport() )
 		{
-			if ( FViewport* Viewport = ViewportClient->Viewport )
+			if ( FSceneViewport* Viewport = ViewportClient->GetGameViewport() )
 			{
 				FVector2D ViewportSize;
 				ViewportClient->GetViewportSize(ViewportSize);
 
 				PixelPosition = Viewport->VirtualDesktopPixelToViewport(FIntPoint((int32)AbsoluteDesktopCoordinate.X, (int32)AbsoluteDesktopCoordinate.Y)) * ViewportSize;
 
-				float CurrentViewportScale = GetDefault<UUserInterfaceSettings>()->GetDPIScaleBasedOnSize(FIntPoint(ViewportSize.X, ViewportSize.Y));
-
 				// Remove DPI Scaling.
-				ViewportPosition = PixelPosition / CurrentViewportScale;
+				ViewportPosition = PixelPosition / UWidgetLayoutLibrary::GetViewportScale(ViewportClient);
 
 				return;
 			}
@@ -108,4 +111,12 @@ void USlateBlueprintLibrary::ScreenToWidgetAbsolute(UObject* WorldContextObject,
 	AbsoluteCoordinate = FVector2D(0, 0);
 }
 
-#undef LOCTEXT_NAMESPACE
+void USlateBlueprintLibrary::ScreenToViewport(UObject* WorldContextObject, FVector2D ScreenPosition, FVector2D& ViewportPosition)
+{
+	FVector2D AbsolutePosition;
+	USlateBlueprintLibrary::ScreenToWidgetAbsolute(WorldContextObject, ScreenPosition, AbsolutePosition);
+	FVector2D PixelPosition;
+	USlateBlueprintLibrary::AbsoluteToViewport(WorldContextObject, AbsolutePosition, PixelPosition, ViewportPosition);
+}
+
+#undef LOCTEXT_NAMESPACE
\ No newline at end of file
diff --git a/Engine/Source/Runtime/UMG/Private/WidgetLayoutLibrary.cpp b/Engine/Source/Runtime/UMG/Private/WidgetLayoutLibrary.cpp
--- a/Engine/Source/Runtime/UMG/Private/WidgetLayoutLibrary.cpp
+++ b/Engine/Source/Runtime/UMG/Private/WidgetLayoutLibrary.cpp
@@ -14,6 +14,12 @@
 #include "Components/OverlaySlot.h"
 #include "Components/BorderSlot.h"
 #include "Engine/UserInterfaceSettings.h"
+#include "Slate/SlateBrushAsset.h"
+#include "WidgetLayoutLibrary.h"
+#include "Runtime/Engine/Classes/Engine/UserInterfaceSettings.h"
+#include "Runtime/Engine/Classes/Engine/RendererSettings.h"
+#include "Blueprint/SlateBlueprintLibrary.h"
+#include "Slate/SceneViewport.h"
 
 #define LOCTEXT_NAMESPACE "UMG"
 
@@ -25,44 +31,41 @@ UWidgetLayoutLibrary::UWidgetLayoutLibrary(const FObjectInitializer& ObjectIniti
 {
 }
 
-bool UWidgetLayoutLibrary::ProjectWorldLocationToWidgetPosition(APlayerController* PlayerController, FVector WorldLocation, FVector2D& ScreenPosition)
+bool UWidgetLayoutLibrary::ProjectWorldLocationToWidgetPosition(APlayerController* PlayerController, FVector WorldLocation, FVector2D& ViewportPosition)
 {
 	FVector ScreenPosition3D;
 	const bool bSuccess = ProjectWorldLocationToWidgetPositionWithDistance(PlayerController, WorldLocation, ScreenPosition3D);
-	ScreenPosition = FVector2D(ScreenPosition3D.X, ScreenPosition3D.Y);
+	ViewportPosition = FVector2D(ScreenPosition3D.X, ScreenPosition3D.Y);
 	return bSuccess;
 }
 
-bool UWidgetLayoutLibrary::ProjectWorldLocationToWidgetPositionWithDistance(APlayerController* PlayerController, FVector WorldLocation, FVector& ScreenPosition)
+bool UWidgetLayoutLibrary::ProjectWorldLocationToWidgetPositionWithDistance(APlayerController* PlayerController, FVector WorldLocation, FVector& ViewportPosition)
 {	
 	if ( PlayerController )
 	{
 		FVector PixelLocation;
-		const bool bProjected = PlayerController->ProjectWorldLocationToScreenWithDistance(WorldLocation, PixelLocation, /*bPlayerViewportRelative=*/ true);
+		const bool bPlayerViewportRelative = false;
+		const bool bProjected = PlayerController->ProjectWorldLocationToScreenWithDistance(WorldLocation, PixelLocation, bPlayerViewportRelative);
 
 		if ( bProjected )
 		{
-			ScreenPosition = PixelLocation;
-
 			// Round the pixel projected value to reduce jittering due to layout rounding,
 			// I do this before I remove scaling, because scaling is going to be applied later
 			// in the opposite direction, so as long as we round, before inverse scale, scale should
 			// result in more or less the same value, especially after slate does layout rounding.
-			ScreenPosition.X = FMath::RoundToInt(ScreenPosition.X);
-			ScreenPosition.Y = FMath::RoundToInt(ScreenPosition.Y);
-			ScreenPosition.Z = PixelLocation.Z;
-
-			// Get the application / DPI scale
-			const float Scale = UWidgetLayoutLibrary::GetViewportScale(PlayerController);
+			FVector2D ScreenPosition(FMath::RoundToInt(PixelLocation.X), FMath::RoundToInt(PixelLocation.Y));
 
-			// Apply inverse DPI scale so that the widget ends up in the expected position
-			ScreenPosition = ScreenPosition / Scale;
+			FVector2D ViewportPosition2D;
+			USlateBlueprintLibrary::ScreenToViewport(PlayerController, ScreenPosition, ViewportPosition2D);
+			ViewportPosition.X = ViewportPosition2D.X;
+			ViewportPosition.Y = ViewportPosition2D.Y;
+			ViewportPosition.Z = PixelLocation.Z;
 
 			return true;
 		}
 	}
 
-	ScreenPosition = FVector::ZeroVector;
+	ViewportPosition = FVector::ZeroVector;
 
 	return false;
 }
@@ -74,17 +77,34 @@ float UWidgetLayoutLibrary::GetViewportScale(UObject* WorldContextObject)
 	{
 		if ( UGameViewportClient* ViewportClient = World->GetGameViewport() )
 		{
-			FVector2D ViewportSize;
-			ViewportClient->GetViewportSize(ViewportSize);
-			return GetDefault<UUserInterfaceSettings>()->GetDPIScaleBasedOnSize(FIntPoint(ViewportSize.X, ViewportSize.Y));
+			return GetViewportScale(ViewportClient);
 		}
 	}
 
 	return 1;
 }
 
+float UWidgetLayoutLibrary::GetViewportScale(UGameViewportClient* ViewportClient)
+{
+	FVector2D ViewportSize;
+	ViewportClient->GetViewportSize(ViewportSize);
+	
+	float UserResolutionScale = GetDefault<UUserInterfaceSettings>()->GetDPIScaleBasedOnSize(FIntPoint(ViewportSize.X, ViewportSize.Y));
+
+	// Normally we'd factor in native DPI Scale here too, but because the SGameLayerManager already
+	// accounts for the native DPI scale, and extracts it from the calculations, the UMG/Slate portion of the
+	// game can more or less assume the platform scale is 1.0 for DPI.
+	//
+	// This id done because UMG already provides a solution for scaling UI across devices, that if influenced
+	// by the platform scale would produce undesirable results.
+
+	return UserResolutionScale;
+}
+
 bool UWidgetLayoutLibrary::GetMousePositionScaledByDPI(APlayerController* Player, float& LocationX, float& LocationY)
 {
+	// TODO NDarnell We should deprecate this function, it's not super useful.
+
 	if ( Player && Player->GetMousePosition(LocationX, LocationY) )
 	{
 		float Scale = UWidgetLayoutLibrary::GetViewportScale(Player);
diff --git a/Engine/Source/Runtime/UMG/Public/Blueprint/SlateBlueprintLibrary.h b/Engine/Source/Runtime/UMG/Public/Blueprint/SlateBlueprintLibrary.h
--- a/Engine/Source/Runtime/UMG/Public/Blueprint/SlateBlueprintLibrary.h
+++ b/Engine/Source/Runtime/UMG/Public/Blueprint/SlateBlueprintLibrary.h
@@ -75,9 +75,21 @@ public:
 	UFUNCTION(BlueprintPure, Category="User Interface|Geometry", meta=( WorldContext="WorldContextObject" ))
 	static void AbsoluteToViewport(UObject* WorldContextObject, FVector2D AbsoluteDesktopCoordinate, FVector2D& PixelPosition, FVector2D& ViewportPosition);
 
-	UFUNCTION(Category="User Interface|Geometry", meta=( WorldContext="WorldContextObject" ))
+	/**
+	 * Translates a screen position in pixels into the local space of a widget with the given geometry. 
+	 */
+	UFUNCTION(Category="User Interface|Geometry", meta=( WorldContext="WorldContextObject", DisplayName="ScreenToLocal" ))
 	static void ScreenToWidgetLocal(UObject* WorldContextObject, const FGeometry& Geometry, FVector2D ScreenPosition, FVector2D& LocalCoordinate);
 
-	UFUNCTION(Category="User Interface|Geometry", meta=( WorldContext="WorldContextObject" ))
+	/**
+	 * Translates a screen position in pixels into absolute application coordinates.
+	 */
+	UFUNCTION(Category="User Interface|Geometry", meta=( WorldContext="WorldContextObject", DisplayName="ScreenToAbsolute" ))
 	static void ScreenToWidgetAbsolute(UObject* WorldContextObject, FVector2D ScreenPosition, FVector2D& AbsoluteCoordinate);
+
+	/**
+	 * Translates a screen position in pixels into the local space of the viewport widget.
+	 */
+	UFUNCTION(Category="User Interface|Geometry", meta=( WorldContext="WorldContextObject" ))
+	static void ScreenToViewport(UObject* WorldContextObject, FVector2D ScreenPosition, FVector2D& ViewportPosition);
 };
diff --git a/Engine/Source/Runtime/UMG/Public/Blueprint/WidgetLayoutLibrary.h b/Engine/Source/Runtime/UMG/Public/Blueprint/WidgetLayoutLibrary.h
--- a/Engine/Source/Runtime/UMG/Public/Blueprint/WidgetLayoutLibrary.h
+++ b/Engine/Source/Runtime/UMG/Public/Blueprint/WidgetLayoutLibrary.h
@@ -48,6 +48,11 @@ public:
 	UFUNCTION(BlueprintPure, BlueprintCosmetic, Category="Viewport", meta=( WorldContext="WorldContextObject" ))
 	static float GetViewportScale(UObject* WorldContextObject);
 
+	/**
+	 * Gets the current DPI Scale being applied to the viewport and all the Widgets.
+	 */
+	static float GetViewportScale(UGameViewportClient* ViewportClient);
+
 	/**
 	 * Gets the size of the game viewport.
 	 */
diff --git a/Engine/Source/Runtime/UMG/Public/Components/BackgroundBlur.h b/Engine/Source/Runtime/UMG/Public/Components/BackgroundBlur.h
new file mode 100644
--- /dev/null
+++ b/Engine/Source/Runtime/UMG/Public/Components/BackgroundBlur.h
@@ -0,0 +1,104 @@
+// Copyright 1998-2017 Epic Games, Inc. All Rights Reserved.
+
+#pragma once
+
+#include "ContentWidget.h"
+#include "BackgroundBlur.generated.h"
+
+/**
+ * A background blur is a container widget that can contain one child widget, providing an opportunity 
+ * to surround it with adjustable padding and apply a post-process Gaussian blur to all content beneath the widget.
+ *
+ * * Single Child
+ * * Blur Effect
+ */
+UCLASS()
+class UMG_API UBackgroundBlur : public UContentWidget
+{
+	GENERATED_UCLASS_BODY()
+
+public:
+	/** The padding area between the slot and the content it contains. */
+	UPROPERTY(EditAnywhere, BlueprintReadOnly, Category="Content")
+	FMargin Padding;
+
+	/** The alignment of the content horizontally. */
+	UPROPERTY(EditAnywhere, BlueprintReadOnly, Category="Content")
+	TEnumAsByte<EHorizontalAlignment> HorizontalAlignment;
+
+	/** The alignment of the content vertically. */
+	UPROPERTY(EditAnywhere, BlueprintReadOnly, Category="Content")
+	TEnumAsByte<EVerticalAlignment> VerticalAlignment;
+
+	/** True to modulate the strength of the blur based on the widget alpha. */
+	UPROPERTY(EditAnywhere, BlueprintReadOnly, Category=Appearance)
+	bool bApplyAlphaToBlur;
+
+	/**
+	 * How blurry the background is.  Larger numbers mean more blurry but will result in larger runtime cost on the gpu.
+	 */
+	UPROPERTY(EditAnywhere, BlueprintReadOnly, Category=Appearance, meta=(ClampMin=0,ClampMax=100))
+	float BlurStrength;
+
+	/** Whether or not the radius should be computed automatically or if it should use the radius */
+	UPROPERTY()
+	bool bOverrideAutoRadiusCalculation;
+
+	/**
+	 * This is the number of pixels which will be weighted in each direction from any given pixel when computing the blur
+	 * A larger value is more costly but allows for stronger blurs.  
+	 */
+	UPROPERTY(EditAnywhere, AdvancedDisplay, BlueprintReadOnly, Category=Appearance, meta=(ClampMin=0, ClampMax=255, EditCondition="bOverrideAutoRadiusCalculation"))
+	int32 BlurRadius;
+
+	/**
+	 * An image to draw instead of applying a blur when low quality override mode is enabled. 
+	 * You can enable low quality mode for background blurs by setting the cvar Slate.ForceBackgroundBlurLowQualityOverride to 1. 
+	 * This is usually done in the project's scalability settings
+	 */
+	UPROPERTY(EditAnywhere, BlueprintReadOnly, Category=Appearance)
+	FSlateBrush LowQualityFallbackBrush;
+
+public:
+	virtual void ReleaseSlateResources(bool bReleaseChildren) override;
+
+#if WITH_EDITOR
+	virtual const FText GetPaletteCategory() override;
+#endif
+
+	UFUNCTION(BlueprintCallable, Category="Appearance")
+	void SetPadding(FMargin InPadding);
+
+	UFUNCTION(BlueprintCallable, Category="Appearance")
+	void SetHorizontalAlignment(EHorizontalAlignment InHorizontalAlignment);
+
+	UFUNCTION(BlueprintCallable, Category="Appearance")
+	void SetVerticalAlignment(EVerticalAlignment InVerticalAlignment);
+
+	UFUNCTION(BlueprintCallable, Category = "Appearance")
+	void SetApplyAlphaToBlur(bool bInApplyAlphaToBlur);
+
+	UFUNCTION(BlueprintCallable, Category = "Appearance")
+	void SetBlurRadius(int32 InBlurRadius);
+
+	UFUNCTION(BlueprintCallable, Category = "Appearance")
+	virtual void SetBlurStrength(float InStrength);
+
+	UFUNCTION(BlueprintCallable, Category = "Appearance")
+	void SetLowQualityFallbackBrush(const FSlateBrush& InBrush);
+
+	/** UObject interface */
+	virtual void PostLoad() override;
+protected:
+	/** UWidget interface */
+	virtual UClass* GetSlotClass() const override;
+	virtual TSharedRef<SWidget> RebuildWidget() override;
+	virtual void SynchronizeProperties() override;
+
+	/** UPanelWidget interface */
+	virtual void OnSlotAdded(UPanelSlot* Slot) override;
+	virtual void OnSlotRemoved(UPanelSlot* Slot) override;
+
+protected:
+	TSharedPtr<class SBackgroundBlur> MyBackgroundBlur;
+};
diff --git a/Engine/Source/Runtime/UMG/Public/Components/BackgroundBlurSlot.h b/Engine/Source/Runtime/UMG/Public/Components/BackgroundBlurSlot.h
new file mode 100644
--- /dev/null
+++ b/Engine/Source/Runtime/UMG/Public/Components/BackgroundBlurSlot.h
@@ -0,0 +1,62 @@
+// Copyright 1998-2016 Epic Games, Inc. All Rights Reserved.
+
+#pragma once
+
+#include "SlateWrapperTypes.h"
+#include "PanelSlot.h"
+#include "BackgroundBlurSlot.generated.h"
+
+
+class SBackgroundBlur;
+class UBackgroundBlur;
+
+/**
+ * The Slot for the UBackgroundBlurSlot, contains the widget displayed in a BackgroundBlur's single slot
+ */
+UCLASS()
+class UMG_API UBackgroundBlurSlot : public UPanelSlot
+{
+	GENERATED_UCLASS_BODY()
+
+public:
+
+	UFUNCTION(BlueprintCallable, Category="Layout|Background Blur Slot")
+	void SetPadding(FMargin InPadding);
+
+	UFUNCTION(BlueprintCallable, Category="Layout|Background Blur Slot")
+	void SetHorizontalAlignment(EHorizontalAlignment InHorizontalAlignment);
+
+	UFUNCTION(BlueprintCallable, Category="Layout|Background Blur Slot")
+	void SetVerticalAlignment(EVerticalAlignment InVerticalAlignment);
+
+protected:
+	/** The padding area between the slot and the content it contains. */
+	UPROPERTY(EditAnywhere, Category="Layout|Background Blur Slot")
+	FMargin Padding;
+
+	/** The alignment of the object horizontally. */
+	UPROPERTY(EditAnywhere, Category="Layout|Background Blur Slot")
+	TEnumAsByte<EHorizontalAlignment> HorizontalAlignment;
+
+	/** The alignment of the object vertically. */
+	UPROPERTY(EditAnywhere, Category="Layout|Background Blur Slot")
+	TEnumAsByte<EVerticalAlignment> VerticalAlignment;
+
+public:
+
+	// UPanelSlot interface
+	virtual void SynchronizeProperties() override;
+	// End of UPanelSlot interface
+
+	/** Builds the underlying slot for the slate BackgroundBlur. */
+	void BuildSlot(TSharedRef<SBackgroundBlur> InBackgroundBlur);
+
+	virtual void ReleaseSlateResources(bool bReleaseChildren) override;
+
+private:
+
+	/** A pointer to the BackgroundBlur to allow us to adjust the size, padding...etc at runtime. */
+	TSharedPtr<SBackgroundBlur> BackgroundBlur;
+
+	friend UBackgroundBlur;
+};
diff --git a/Engine/Source/Runtime/UMG/Public/Components/ComboBox.h b/Engine/Source/Runtime/UMG/Public/Components/ComboBox.h
--- a/Engine/Source/Runtime/UMG/Public/Components/ComboBox.h
+++ b/Engine/Source/Runtime/UMG/Public/Components/ComboBox.h
@@ -33,6 +33,9 @@ public:
 	virtual const FText GetPaletteCategory() override;
 #endif
 
+	UPROPERTY(EditAnywhere, BlueprintReadOnly, Category=Interaction)
+	bool bIsFocusable;
+
 protected:
 	TSharedRef<SWidget> HandleGenerateWidget(UObject* Item) const;
 
diff --git a/Engine/Source/Runtime/UMG/Public/Components/ComboBoxString.h b/Engine/Source/Runtime/UMG/Public/Components/ComboBoxString.h
--- a/Engine/Source/Runtime/UMG/Public/Components/ComboBoxString.h
+++ b/Engine/Source/Runtime/UMG/Public/Components/ComboBoxString.h
@@ -77,6 +77,9 @@ public:
 	UPROPERTY(EditAnywhere, BlueprintReadOnly, Category=Style, meta=(DesignerRebuild))
 	FSlateColor ForegroundColor;
 
+	UPROPERTY(EditAnywhere, BlueprintReadOnly, Category=Interaction)
+	bool bIsFocusable;
+
 public: // Events
 
 	/** Called when the widget is needed for the item. */
diff --git a/Engine/Source/Runtime/UMG/Public/Components/PanelWidget.h b/Engine/Source/Runtime/UMG/Public/Components/PanelWidget.h
--- a/Engine/Source/Runtime/UMG/Public/Components/PanelWidget.h
+++ b/Engine/Source/Runtime/UMG/Public/Components/PanelWidget.h
@@ -6,6 +6,7 @@
 #include "UObject/ObjectMacros.h"
 #include "Components/PanelSlot.h"
 #include "Components/Widget.h"
+#include "Blueprint/UserWidget.h"
 #include "PanelWidget.generated.h"
 
 /** The base class for all UMG panel widgets.  Panel widgets layout a collection of child widgets. */
diff --git a/Engine/Source/Runtime/UMG/Public/Components/Widget.h b/Engine/Source/Runtime/UMG/Public/Components/Widget.h
--- a/Engine/Source/Runtime/UMG/Public/Components/Widget.h
+++ b/Engine/Source/Runtime/UMG/Public/Components/Widget.h
@@ -36,6 +36,7 @@ namespace UMWidget
 	{
 		/// [PropertyMetadata] This property if changed will rebuild the widget designer preview.  Use sparingly, try to update most properties by
 		/// setting them in the SynchronizeProperties function.
+		/// UPROPERTY(meta=(DesignerRebuild))
 		DesignerRebuild,
 		/// [PropertyMetadata] This property requires a widget be bound to it in the designer.  Allows easy native access to designer defined controls.
 		/// UPROPERTY(meta=(BindWidget))
@@ -127,7 +128,7 @@ namespace EWidgetDesignFlags
 /**
  * This is the base class for all wrapped Slate controls that are exposed to UObjects.
  */
-UCLASS(Abstract, BlueprintType)
+UCLASS(Abstract, BlueprintType, Blueprintable)
 class UMG_API UWidget : public UVisual
 {
 	GENERATED_UCLASS_BODY()
@@ -419,6 +420,18 @@ public:
 	UFUNCTION(BlueprintCallable, Category="Widget")
 	virtual void RemoveFromParent();
 
+	/**
+	 * Gets the last geometry used to Tick the widget.  This data may not exist yet if this call happens prior to 
+	 * the widget having been ticked/painted, or it may be out of date, or a frame behind.
+	 *
+	 * We recommend not to use this data unless there's no other way to solve your problem.  Normally in Slate we
+	 * try and handle these issues by making a dependent widget part of the hierarchy, as to avoid frame behind
+	 * or what are referred to as hysteresis problems, both caused by depending on geometry from the previous frame
+	 * being used to advise how to layout a dependent object the current frame.
+	 */
+	UFUNCTION(BlueprintCallable, Category="Widget")
+	const FGeometry& GetCachedGeometry() const;
+
 	/**
 	 * Gets the underlying slate widget or constructs it if it doesn't exist.  If you're looking to replace
 	 * what slate widget gets constructed look for RebuildWidget.  For extremely special cases where you actually
diff --git a/Engine/Source/Runtime/UMG/Public/Components/WidgetComponent.h b/Engine/Source/Runtime/UMG/Public/Components/WidgetComponent.h
--- a/Engine/Source/Runtime/UMG/Public/Components/WidgetComponent.h
+++ b/Engine/Source/Runtime/UMG/Public/Components/WidgetComponent.h
@@ -216,8 +216,15 @@ public:
 	/** @return The pivot point where the UI is rendered about the origin. */
 	FVector2D GetPivot() const { return Pivot; }
 
+	/**  */
 	void SetPivot( const FVector2D& InPivot ) { Pivot = InPivot; }
 
+	/**  */
+	bool GetDrawAtDesiredSize() const { return bDrawAtDesiredSize; }
+
+	/**  */
+	void SetDrawAtDesiredSize(bool InbDrawAtDesiredSize) { bDrawAtDesiredSize = InbDrawAtDesiredSize; }
+
 	/** Get the fake window we create for widgets displayed in the world. */
 	TSharedPtr< SWindow > GetVirtualWindow() const;
 	
@@ -238,10 +245,15 @@ public:
 	/** @see EWidgetGeometryMode, @see GetCylinderArcAngle() */
 	EWidgetGeometryMode GetGeometryMode() const { return GeometryMode; }
 
+	bool GetReceiveHardwareInput() const { return bReceiveHardwareInput; }
+
 	/** Defines the curvature of the widget component when using EWidgetGeometryMode::Cylinder; ignored otherwise.  */
 	float GetCylinderArcAngle() const { return CylinderArcAngle; }
 
 protected:
+	void RegisterHitTesterWithViewport(TSharedPtr<SViewport> ViewportWidget);
+	void UnregisterHitTesterWithViewport(TSharedPtr<SViewport> ViewportWidget);
+
 	void RegisterWindow();
 	void UnregisterWindow();
 	void RemoveWidgetFromScreen();
@@ -285,12 +297,7 @@ protected:
 	float RedrawTime;
 
 	/** What was the last time we rendered the widget? */
-	UPROPERTY()
-	float LastWidgetRenderTime;
-
-	/** Is the virtual window created to host the widget focusable? */
-	UPROPERTY(EditAnywhere, Category=UserInterface)
-	bool bWindowFocusable;
+	double LastWidgetRenderTime;
 
 	/**
 	 * The actual draw size, this changes based on DrawSize - or the desired size of the widget if
@@ -313,6 +320,23 @@ protected:
 	UPROPERTY(EditAnywhere, Category=UserInterface)
 	FVector2D Pivot;
 
+	/**
+	 * Register with the viewport for hardware input from the true mouse and keyboard.  These widgets
+	 * will more or less react like regular 2D widgets in the viewport, e.g. they can and will steal focus
+	 * from the viewport.
+	 * 
+	 * WARNING: If you are making a VR game, definitely do not change this to true.  This option should ONLY be used
+	 * if you're making what would otherwise be a normal menu for a game, just in 3D.  If you also need the game to 
+	 * remain responsive and for the player to be able to interact with UI and move around the world (such as a keypad on a door), 
+	 * use the WidgetInteractionComponent instead.
+	 */
+	UPROPERTY(EditAnywhere, Category=Interaction)
+	bool bReceiveHardwareInput;
+
+	/** Is the virtual window created to host the widget focusable? */
+	UPROPERTY(EditAnywhere, Category=Interaction)
+	bool bWindowFocusable;
+
 	/**
 	 * The owner player for a widget component, if this widget is drawn on the screen, this controls
 	 * what player's screen it appears on for split screen, if not set, users player 0.
diff --git a/Engine/Source/Runtime/UMG/Public/Components/WidgetInteractionComponent.h b/Engine/Source/Runtime/UMG/Public/Components/WidgetInteractionComponent.h
--- a/Engine/Source/Runtime/UMG/Public/Components/WidgetInteractionComponent.h
+++ b/Engine/Source/Runtime/UMG/Public/Components/WidgetInteractionComponent.h
@@ -280,6 +280,9 @@ protected:
 	 */
 	void GetRelatedComponentsToIgnoreInAutomaticHitTesting(TArray<UPrimitiveComponent*>& IgnorePrimitives) const;
 
+	/** Returns true if the inteaction component can interact with the supplied widget component */
+	bool CanInteractWithComponent(UWidgetComponent* Component) const;
+
 protected:
 
 	/** The last widget path under the hit result. */
diff --git a/Engine/Source/Runtime/UMG/Public/Slate/SRetainerWidget.h b/Engine/Source/Runtime/UMG/Public/Slate/SRetainerWidget.h
--- a/Engine/Source/Runtime/UMG/Public/Slate/SRetainerWidget.h
+++ b/Engine/Source/Runtime/UMG/Public/Slate/SRetainerWidget.h
@@ -84,6 +84,8 @@ public:
 		return CachedAllottedGeometry;
 	}
 
+	void SetWorld(UWorld* World);
+
 protected:
 	// BEGIN SLeafWidget interface
 	virtual int32 OnPaint(const FPaintArgs& Args, const FGeometry& AllottedGeometry, const FSlateRect& MyClippingRect, FSlateWindowElementList& OutDrawElements, int32 LayerId, const FWidgetStyle& InWidgetStyle, bool bParentEnabled) const override;
@@ -123,6 +125,7 @@ private:
 
 	TSharedPtr<SVirtualWindow> Window;
 	TSharedPtr<FHittestGrid> HitTestGrid;
+	TWeakObjectPtr<UWorld> OuterWorld;
 
 	STAT(TStatId MyStatId;)
 
