UMG - Fixing an uninitialized variable in the FWidgetRenderer leading to an ensure in some usages where the hittest grid isn't explictly cleared and needs to be.

diff --git a/Engine/Source/Runtime/UMG/Private/Slate/WidgetRenderer.cpp b/Engine/Source/Runtime/UMG/Private/Slate/WidgetRenderer.cpp
--- a/Engine/Source/Runtime/UMG/Private/Slate/WidgetRenderer.cpp
+++ b/Engine/Source/Runtime/UMG/Private/Slate/WidgetRenderer.cpp
@@ -17,6 +17,7 @@
 
 FWidgetRenderer::FWidgetRenderer(bool bUseGammaCorrection, bool bInClearTarget)
 	: bPrepassNeeded(true)
+	, bClearHitTestGrid(true)
 	, bUseGammaSpace(bUseGammaCorrection)
 	, bClearTarget(bInClearTarget)
 	, ViewOffset(0.0f, 0.0f)
