Copying //UE4/Dev-Framework to //UE4/Dev-Main (Source: //UE4/Dev-Framework @ 3628051)

diff --git a/Engine/Source/Runtime/SlateCore/Private/Styling/SlateBrush.cpp b/Engine/Source/Runtime/SlateCore/Private/Styling/SlateBrush.cpp
--- a/Engine/Source/Runtime/SlateCore/Private/Styling/SlateBrush.cpp
+++ b/Engine/Source/Runtime/SlateCore/Private/Styling/SlateBrush.cpp
@@ -8,16 +8,16 @@
 
 FSlateBrush::FSlateBrush( ESlateBrushDrawType::Type InDrawType, const FName InResourceName, const FMargin& InMargin, ESlateBrushTileType::Type InTiling, ESlateBrushImageType::Type InImageType, const FVector2D& InImageSize, const FLinearColor& InTint, UObject* InObjectResource, bool bInDynamicallyLoaded )
 	: ImageSize( InImageSize )
-	, DrawAs( InDrawType )
 	, Margin( InMargin )
 	, TintColor( InTint )
+	, ResourceObject( InObjectResource )
+	, ResourceName( InResourceName )
+	, UVRegion( ForceInit )
+	, DrawAs( InDrawType )
 	, Tiling( InTiling )
 	, Mirroring( ESlateBrushMirrorType::NoMirror )
 	, ImageType( InImageType )
-	, ResourceObject( InObjectResource )
-	, ResourceName( InResourceName )
 	, bIsDynamicallyLoaded( bInDynamicallyLoaded )
-	, UVRegion( ForceInit )
 {
 	bHasUObject_DEPRECATED = (InObjectResource != nullptr) || InResourceName.ToString().StartsWith(FSlateBrush::UTextureIdentifier());
 
@@ -31,16 +31,16 @@ FSlateBrush::FSlateBrush( ESlateBrushDrawType::Type InDrawType, const FName InRe
 
 FSlateBrush::FSlateBrush( ESlateBrushDrawType::Type InDrawType, const FName InResourceName, const FMargin& InMargin, ESlateBrushTileType::Type InTiling, ESlateBrushImageType::Type InImageType, const FVector2D& InImageSize, const TSharedRef< FLinearColor >& InTint, UObject* InObjectResource, bool bInDynamicallyLoaded )
 	: ImageSize( InImageSize )
-	, DrawAs( InDrawType )
 	, Margin( InMargin )
 	, TintColor( InTint )
+	, ResourceObject( InObjectResource )
+	, ResourceName( InResourceName )
+	, UVRegion( ForceInit )
+	, DrawAs( InDrawType )
 	, Tiling( InTiling )
 	, Mirroring( ESlateBrushMirrorType::NoMirror )
 	, ImageType( InImageType )
-	, ResourceObject( InObjectResource )
-	, ResourceName( InResourceName )
 	, bIsDynamicallyLoaded( bInDynamicallyLoaded )
-	, UVRegion( ForceInit )
 {
 	bHasUObject_DEPRECATED = (InObjectResource != nullptr) || InResourceName.ToString().StartsWith(FSlateBrush::UTextureIdentifier());
 
@@ -54,16 +54,16 @@ FSlateBrush::FSlateBrush( ESlateBrushDrawType::Type InDrawType, const FName InRe
 
 FSlateBrush::FSlateBrush( ESlateBrushDrawType::Type InDrawType, const FName InResourceName, const FMargin& InMargin, ESlateBrushTileType::Type InTiling, ESlateBrushImageType::Type InImageType, const FVector2D& InImageSize, const FSlateColor& InTint, UObject* InObjectResource, bool bInDynamicallyLoaded )
 	: ImageSize(InImageSize)
-	, DrawAs(InDrawType)
 	, Margin(InMargin)
 	, TintColor(InTint)
+	, ResourceObject(InObjectResource)
+	, ResourceName(InResourceName)
+	, UVRegion(ForceInit)
+	, DrawAs(InDrawType)
 	, Tiling(InTiling)
 	, Mirroring( ESlateBrushMirrorType::NoMirror )
 	, ImageType(InImageType)
-	, ResourceObject(InObjectResource)
-	, ResourceName(InResourceName)
 	, bIsDynamicallyLoaded(bInDynamicallyLoaded)
-	, UVRegion(ForceInit)
 {
 	bHasUObject_DEPRECATED = (InObjectResource != nullptr) || InResourceName.ToString().StartsWith(FSlateBrush::UTextureIdentifier());
 
diff --git a/Engine/Source/Runtime/SlateCore/Private/Widgets/SBoxPanel.cpp b/Engine/Source/Runtime/SlateCore/Private/Widgets/SBoxPanel.cpp
--- a/Engine/Source/Runtime/SlateCore/Private/Widgets/SBoxPanel.cpp
+++ b/Engine/Source/Runtime/SlateCore/Private/Widgets/SBoxPanel.cpp
@@ -1,4 +1,4 @@
-// Copyright 1998-2017 Epic Games, Inc. All Rights Reserved.
+﻿// Copyright 1998-2017 Epic Games, Inc. All Rights Reserved.
 
 #include "Widgets/SBoxPanel.h"
 #include "Layout/LayoutUtils.h"
@@ -280,8 +280,8 @@ void SBoxPanel::ClearChildren()
  * @param InOrientation   The orientation of the Box Panel
  */
 SBoxPanel::SBoxPanel( EOrientation InOrientation )
-: Orientation(InOrientation)
-, Children()
+: Children()
+, Orientation(InOrientation)
 {
 
 }
diff --git a/Engine/Source/Runtime/SlateCore/Public/Fonts/SlateFontInfo.h b/Engine/Source/Runtime/SlateCore/Public/Fonts/SlateFontInfo.h
--- a/Engine/Source/Runtime/SlateCore/Public/Fonts/SlateFontInfo.h
+++ b/Engine/Source/Runtime/SlateCore/Public/Fonts/SlateFontInfo.h
@@ -124,6 +124,10 @@ struct SLATECORE_API FSlateFontInfo
 	UPROPERTY(EditAnywhere, BlueprintReadWrite, Category=SlateStyleRules, meta=(ClampMin=1, ClampMax=1000))
 	int32 Size;
 
+	/** The font fallback level. Runtime only, don't set on shared FSlateFontInfo, as it may change the font elsewhere (make a copy). */
+	EFontFallback FontFallback;
+
+#if WITH_EDITORONLY_DATA
 private:
 
 	/** The name of the font */
@@ -133,11 +137,7 @@ private:
 	/** The hinting algorithm to use with the font */
 	UPROPERTY()
 	EFontHinting Hinting_DEPRECATED;
-
-public:
-
-	/** The font fallback level. Runtime only, don't set on shared FSlateFontInfo, as it may change the font elsewhere (make a copy). */
-	EFontFallback FontFallback;
+#endif
 
 public:
 
@@ -259,19 +259,22 @@ public:
 		return Hash;
 	}
 
+#if WITH_EDITORONLY_DATA
 	/**
 	 * Used to upgrade legacy font into so that it uses composite fonts
 	 */
 	void PostSerialize(const FArchive& Ar);
+#endif
 
 private:
 
 	/**
 	 * Used to upgrade legacy font into so that it uses composite fonts
 	 */
-	void UpgradeLegacyFontInfo();
+	void UpgradeLegacyFontInfo(FName LegacyFontName, EFontHinting LegacyHinting);
 };
 
+#if WITH_EDITORONLY_DATA
 template<>
 struct TStructOpsTypeTraits<FSlateFontInfo>
 	: public TStructOpsTypeTraitsBase2<FSlateFontInfo>
@@ -281,3 +284,4 @@ struct TStructOpsTypeTraits<FSlateFontInfo>
 		WithPostSerialize = true,
 	};
 };
+#endif
\ No newline at end of file
diff --git a/Engine/Source/Runtime/SlateCore/Public/Layout/Visibility.h b/Engine/Source/Runtime/SlateCore/Public/Layout/Visibility.h
--- a/Engine/Source/Runtime/SlateCore/Public/Layout/Visibility.h
+++ b/Engine/Source/Runtime/SlateCore/Public/Layout/Visibility.h
@@ -108,5 +108,5 @@ private:
 
 public:
 
-	Private Value;
+	TEnumAsByte<Private> Value;
 };
diff --git a/Engine/Source/Runtime/SlateCore/Public/Widgets/SBoxPanel.h b/Engine/Source/Runtime/SlateCore/Public/Widgets/SBoxPanel.h
--- a/Engine/Source/Runtime/SlateCore/Public/Widgets/SBoxPanel.h
+++ b/Engine/Source/Runtime/SlateCore/Public/Widgets/SBoxPanel.h
@@ -35,21 +35,22 @@ public:
 	public:		
 		/** Horizontal and Vertical Boxes inherit from FSlot */
 		virtual ~FSlot(){}
-		/**
-		 * How much space this slot should occupy along panel's direction.
-		 *   When SizeRule is SizeRule_Auto, the widget's DesiredSize will be used as the space required.
-		 *   When SizeRule is SizeRule_Stretch, the available space will be distributed proportionately between
-		 *   peer Widgets depending on the Value property. Available space is space remaining after all the
-		 *   peers' SizeRule_Auto requirements have been satisfied.
-		 */
-		FSizeParam SizeParam;
 		
 		/** Horizontal positioning of child within the allocated slot */
-		EHorizontalAlignment HAlignment;
+		TEnumAsByte<EHorizontalAlignment> HAlignment;
 		
 		/** Vertical positioning of child within the allocated slot */
-		EVerticalAlignment VAlignment;
-		
+		TEnumAsByte<EVerticalAlignment> VAlignment;
+
+		/**
+		* How much space this slot should occupy along panel's direction.
+		*   When SizeRule is SizeRule_Auto, the widget's DesiredSize will be used as the space required.
+		*   When SizeRule is SizeRule_Stretch, the available space will be distributed proportionately between
+		*   peer Widgets depending on the Value property. Available space is space remaining after all the
+		*   peers' SizeRule_Auto requirements have been satisfied.
+		*/
+		FSizeParam SizeParam;
+
 		/** The padding to add around the child. */
 		TAttribute<FMargin> SlotPadding;
 		
@@ -60,9 +61,9 @@ public:
 		/** Default values for a slot. */
 		FSlot()
 			: TSlotBase<FSlot>()
-			, SizeParam( FStretch(1) )
 			, HAlignment( HAlign_Fill )
 			, VAlignment( VAlign_Fill )
+			, SizeParam( FStretch(1) )
 			, SlotPadding( FMargin(0) )
 			, MaxSize( 0.0f )
 		{ }
@@ -101,10 +102,11 @@ protected:
 	 */
 	SBoxPanel( EOrientation InOrientation );
 
-	/** The Box Panel's orientation; determined at construct time. */
-	const EOrientation Orientation;
 	/** The Box Panel's children. */
 	TPanelChildren<FSlot> Children;
+
+	/** The Box Panel's orientation; determined at construct time. */
+	const EOrientation Orientation;
 };
 
 
diff --git a/Engine/Source/Runtime/SlateCore/Public/Widgets/SOverlay.h b/Engine/Source/Runtime/SlateCore/Public/Widgets/SOverlay.h
--- a/Engine/Source/Runtime/SlateCore/Public/Widgets/SOverlay.h
+++ b/Engine/Source/Runtime/SlateCore/Public/Widgets/SOverlay.h
@@ -42,10 +42,10 @@ public:
 	public:
 		FOverlaySlot()
 			: TSlotBase<FOverlaySlot>()
+			, ZOrder(0)
 			, HAlignment(HAlign_Fill)
 			, VAlignment(VAlign_Fill)
 			, SlotPadding(0.0f)
-			, ZOrder(0)
 		{ }
 
 		FOverlaySlot& HAlign( EHorizontalAlignment InHAlignment )
@@ -84,16 +84,15 @@ public:
 			return *this;
 		}
 
-
-		/** The child widget contained in this slot. */
-		EHorizontalAlignment HAlignment;
-		EVerticalAlignment VAlignment;
-		TAttribute< FMargin > SlotPadding;
-
 		/** Slots with larger ZOrder values will draw above slots with smaller ZOrder values.  Slots
-			with the same ZOrder will simply draw in the order they were added.  Currently this only
-			works for overlay slots that are added dynamically with AddWidget() and RemoveWidget() */
+		with the same ZOrder will simply draw in the order they were added.  Currently this only
+		works for overlay slots that are added dynamically with AddWidget() and RemoveWidget() */
 		int32 ZOrder;
+
+		TEnumAsByte<EHorizontalAlignment> HAlignment;
+		TEnumAsByte<EVerticalAlignment> VAlignment;
+
+		TAttribute< FMargin > SlotPadding;
 	};
 
 
diff --git a/Engine/Source/Runtime/SlateCore/Public/Widgets/SWidget.h b/Engine/Source/Runtime/SlateCore/Public/Widgets/SWidget.h
--- a/Engine/Source/Runtime/SlateCore/Public/Widgets/SWidget.h
+++ b/Engine/Source/Runtime/SlateCore/Public/Widgets/SWidget.h
@@ -1,4 +1,4 @@
-// Copyright 1998-2017 Epic Games, Inc. All Rights Reserved.
+﻿// Copyright 1998-2017 Epic Games, Inc. All Rights Reserved.
 
 #pragma once
 
@@ -1322,19 +1322,19 @@ protected:
 
 protected:
 	/** Is this widget hovered? */
-	bool bIsHovered : 1;
+	uint8 bIsHovered : 1;
 
 	/** Can the widget ever be ticked. */
-	bool bCanTick : 1;
+	uint8 bCanTick : 1;
 
 	/** Can the widget ever support keyboard focus */
-	bool bCanSupportFocus : 1;
+	uint8 bCanSupportFocus : 1;
 
 	/**
 	 * Can the widget ever support children?  This will be false on SLeafWidgets, 
 	 * rather than setting this directly, you should probably inherit from SLeafWidget.
 	 */
-	bool bCanHaveChildren : 1;
+	uint8 bCanHaveChildren : 1;
 
 	/**
 	  * Some widgets might be a complex hierarchy of child widgets you never see.  Some of those widgets
@@ -1342,7 +1342,15 @@ protected:
 	  * so even though it may be set to clip, this flag is used to inform painting that this widget doesn't
 	  * really do the clipping.
 	  */
-	bool bClippingProxy : 1;
+	uint8 bClippingProxy : 1;
+
+#if SLATE_DEFERRED_DESIRED_SIZE
+	/** Has the desired size of the widget been cached? */
+	uint8 bCachedDesiredSize : 1;
+
+	/** Are we currently updating the desired size? */
+	mutable uint8 bUpdatingDesiredSize : 1;
+#endif
 
 private:
 
@@ -1350,16 +1358,16 @@ private:
 	 * Whether this widget is a "tool tip force field".  That is, tool-tips should never spawn over the area
 	 * occupied by this widget, and will instead be repelled to an outside edge
 	 */
-	bool bToolTipForceFieldEnabled : 1;
+	uint8 bToolTipForceFieldEnabled : 1;
 
 	/** Should we be forcing this widget to be volatile at all times and redrawn every frame? */
-	bool bForceVolatile : 1;
+	uint8 bForceVolatile : 1;
 
 	/** The last cached volatility of this widget.  Cached so that we don't need to recompute volatility every frame. */
-	bool bCachedVolatile : 1;
+	uint8 bCachedVolatile : 1;
 
 	/** If we're owned by a volatile widget, we need inherit that volatility and use as part of our volatility, but don't cache it. */
-	mutable bool bInheritedVolatility : 1;
+	mutable uint8 bInheritedVolatility : 1;
 
 protected:
 	/**
@@ -1414,13 +1422,6 @@ protected:
 	TAttribute< FVector2D > RenderTransformPivot;
 
 protected:
-#if SLATE_DEFERRED_DESIRED_SIZE
-	/** Has the desired size of the widget been cached? */
-	bool bCachedDesiredSize : 1;
-
-	/** Are we currently updating the desired size? */
-	mutable bool bUpdatingDesiredSize : 1;
-#endif
 
 	/** Debugging information on the type of widget we're creating for the Widget Reflector. */
 	FName TypeOfWidget;
diff --git a/Engine/Source/Runtime/Sockets/Private/BSDSockets/SocketsBSD.cpp b/Engine/Source/Runtime/Sockets/Private/BSDSockets/SocketsBSD.cpp
--- a/Engine/Source/Runtime/Sockets/Private/BSDSockets/SocketsBSD.cpp
+++ b/Engine/Source/Runtime/Sockets/Private/BSDSockets/SocketsBSD.cpp
@@ -41,8 +41,8 @@ bool FSocketBSD::Connect(const FInternetAddr& Addr)
 	check(SocketSubsystem);
 	ESocketErrors Error = SocketSubsystem->TranslateErrorCode(Return);
 
-	// "would block" is not an error
-	return ((Error == SE_NO_ERROR) || (Error == SE_EWOULDBLOCK));
+	// EWOULDBLOCK is not an error, and EINPROGRESS is fine on initial connection as it may still be creating for nonblocking sockets
+	return ((Error == SE_NO_ERROR) || (Error == SE_EWOULDBLOCK) || (Error == SE_EINPROGRESS));
 }
 
 
diff --git a/Engine/Source/Runtime/UMG/Private/Components/CanvasPanel.cpp b/Engine/Source/Runtime/UMG/Private/Components/CanvasPanel.cpp
--- a/Engine/Source/Runtime/UMG/Private/Components/CanvasPanel.cpp
+++ b/Engine/Source/Runtime/UMG/Private/Components/CanvasPanel.cpp
@@ -15,7 +15,7 @@ UCanvasPanel::UCanvasPanel(const FObjectInitializer& ObjectInitializer)
 	bIsVariable = false;
 
 	SConstraintCanvas::FArguments Defaults;
-	Visiblity_DEPRECATED = Visibility = UWidget::ConvertRuntimeToSerializedVisibility(Defaults._Visibility.Get());
+	Visibility = UWidget::ConvertRuntimeToSerializedVisibility(Defaults._Visibility.Get());
 }
 
 void UCanvasPanel::ReleaseSlateResources(bool bReleaseChildren)
diff --git a/Engine/Source/Runtime/UMG/Private/Components/GridPanel.cpp b/Engine/Source/Runtime/UMG/Private/Components/GridPanel.cpp
--- a/Engine/Source/Runtime/UMG/Private/Components/GridPanel.cpp
+++ b/Engine/Source/Runtime/UMG/Private/Components/GridPanel.cpp
@@ -16,7 +16,7 @@ UGridPanel::UGridPanel(const FObjectInitializer& ObjectInitializer)
 	bIsVariable = false;
 
 	SGridPanel::FArguments Defaults;
-	Visiblity_DEPRECATED = Visibility = UWidget::ConvertRuntimeToSerializedVisibility(Defaults._Visibility.Get());
+	Visibility = UWidget::ConvertRuntimeToSerializedVisibility(Defaults._Visibility.Get());
 }
 
 void UGridPanel::ReleaseSlateResources(bool bReleaseChildren)
diff --git a/Engine/Source/Runtime/UMG/Private/Components/Widget.cpp b/Engine/Source/Runtime/UMG/Private/Components/Widget.cpp
--- a/Engine/Source/Runtime/UMG/Private/Components/Widget.cpp
+++ b/Engine/Source/Runtime/UMG/Private/Components/Widget.cpp
@@ -151,7 +151,7 @@ UWidget::UWidget(const FObjectInitializer& ObjectInitializer)
 #if WITH_EDITOR
 	DesignerFlags = EWidgetDesignFlags::None;
 #endif
-	Visiblity_DEPRECATED = Visibility = ESlateVisibility::Visible;	
+	Visibility = ESlateVisibility::Visible;
 	RenderTransformPivot = FVector2D(0.5f, 0.5f);
 	Cursor = EMouseCursor::Default;
 
@@ -1000,7 +1000,7 @@ void UWidget::SynchronizeProperties()
 			SafeWidget->SetCursor(Cursor);// PROPERTY_BINDING(EMouseCursor::Type, Cursor));
 		}
 
-		SafeWidget->SetEnabled(PROPERTY_BINDING( bool, bIsEnabled ));
+		SafeWidget->SetEnabled(BITFIELD_PROPERTY_BINDING( bIsEnabled ));
 		SafeWidget->SetVisibility(OPTIONAL_BINDING_CONVERT(ESlateVisibility, Visibility, EVisibility, ConvertVisibility));
 	}
 
@@ -1086,16 +1086,6 @@ UWorld* UWidget::GetWorld() const
 	return nullptr;
 }
 
-void UWidget::PostLoad()
-{
-	Super::PostLoad();
-
-	if ( GetLinkerUE4Version() < VER_UE4_RENAME_WIDGET_VISIBILITY )
-	{
-		Visibility = Visiblity_DEPRECATED;
-	}
-}
-
 EVisibility UWidget::ConvertSerializedVisibilityToRuntime(ESlateVisibility Input)
 {
 	switch ( Input )
diff --git a/Engine/Source/Runtime/UMG/Private/Components/WidgetSwitcher.cpp b/Engine/Source/Runtime/UMG/Private/Components/WidgetSwitcher.cpp
--- a/Engine/Source/Runtime/UMG/Private/Components/WidgetSwitcher.cpp
+++ b/Engine/Source/Runtime/UMG/Private/Components/WidgetSwitcher.cpp
@@ -15,7 +15,7 @@ UWidgetSwitcher::UWidgetSwitcher(const FObjectInitializer& ObjectInitializer)
 	bIsVariable = true;
 
 	SWidgetSwitcher::FArguments Defaults;
-	Visibility = Visiblity_DEPRECATED = UWidget::ConvertRuntimeToSerializedVisibility(Defaults._Visibility.Get());
+	Visibility = UWidget::ConvertRuntimeToSerializedVisibility(Defaults._Visibility.Get());
 }
 
 void UWidgetSwitcher::ReleaseSlateResources(bool bReleaseChildren)
diff --git a/Engine/Source/Runtime/UMG/Private/Components/WrapBox.cpp b/Engine/Source/Runtime/UMG/Private/Components/WrapBox.cpp
--- a/Engine/Source/Runtime/UMG/Private/Components/WrapBox.cpp
+++ b/Engine/Source/Runtime/UMG/Private/Components/WrapBox.cpp
@@ -14,7 +14,7 @@ UWrapBox::UWrapBox(const FObjectInitializer& ObjectInitializer)
 	bIsVariable = false;
 
 	SWrapBox::FArguments Defaults;
-	Visibility = Visiblity_DEPRECATED = UWidget::ConvertRuntimeToSerializedVisibility(Defaults._Visibility.Get());
+	Visibility = UWidget::ConvertRuntimeToSerializedVisibility(Defaults._Visibility.Get());
 
 	WrapWidth = 500;
 	bExplicitWrapWidth = false;
diff --git a/Engine/Source/Runtime/UMG/Public/Blueprint/UserWidget.h b/Engine/Source/Runtime/UMG/Public/Blueprint/UserWidget.h
--- a/Engine/Source/Runtime/UMG/Public/Blueprint/UserWidget.h
+++ b/Engine/Source/Runtime/UMG/Public/Blueprint/UserWidget.h
@@ -863,13 +863,6 @@ public:
 	UPROPERTY(EditAnywhere, BlueprintReadOnly, Category="Appearance")
 	FMargin Padding;
 
-	UPROPERTY()
-	bool bSupportsKeyboardFocus_DEPRECATED;
-
-	/** Setting this flag to true, allows this widget to accept focus when clicked, or when navigated to. */
-	UPROPERTY(EditAnywhere, BlueprintReadWrite, Category="Interaction")
-	bool bIsFocusable;
-
 	/** All the sequence players currently playing */
 	UPROPERTY(Transient)
 	TArray<UUMGSequencePlayer*> ActiveSequencePlayers;
@@ -878,12 +871,6 @@ public:
 	UPROPERTY(Transient)
 	TArray<UUMGSequencePlayer*> StoppedSequencePlayers;
 
-	UPROPERTY(EditAnywhere, BlueprintReadWrite, Category = "Input")
-	bool bStopAction;
-
-	UPROPERTY(EditAnywhere, BlueprintReadOnly, Category = "Input")
-	int32 Priority;
-
 private:
 	/** Stores the widgets being assigned to named slots */
 	UPROPERTY()
@@ -918,18 +905,31 @@ public:
 
 #endif
 
+	UPROPERTY(EditAnywhere, BlueprintReadOnly, Category = "Input")
+	int32 Priority;
+
+	UPROPERTY()
+	uint8 bSupportsKeyboardFocus_DEPRECATED:1;
+
+	/** Setting this flag to true, allows this widget to accept focus when clicked, or when navigated to. */
+	UPROPERTY(EditAnywhere, BlueprintReadWrite, Category="Interaction")
+	uint8 bIsFocusable:1;
+
+	UPROPERTY(EditAnywhere, BlueprintReadWrite, Category = "Input")
+	uint8 bStopAction:1;
+
 	/** If a widget doesn't ever need to tick the blueprint, setting this to false is an optimization. */
 	UPROPERTY()
-	uint32 bCanEverTick : 1;
+	uint8 bCanEverTick : 1;
 
 	/** If a widget doesn't ever need to do custom painting in the blueprint, setting this to false is an optimization. */
 	UPROPERTY()
-	uint32 bCanEverPaint : 1;
+	uint8 bCanEverPaint : 1;
 
 protected:
 
 	/** Has this widget been initialized by its class yet? */
-	uint32 bInitialized : 1;
+	uint8 bInitialized : 1;
 
 public:
 	/**
@@ -937,7 +937,7 @@ public:
 	 * initialization logic for widgets, because these widgets have already been initialized.
 	 */
 	UPROPERTY()
-	uint32 bCookedWidgetTree : 1;
+	uint8 bCookedWidgetTree : 1;
 
 protected:
 
diff --git a/Engine/Source/Runtime/UMG/Public/Blueprint/WidgetBlueprintGeneratedClass.h b/Engine/Source/Runtime/UMG/Public/Blueprint/WidgetBlueprintGeneratedClass.h
--- a/Engine/Source/Runtime/UMG/Public/Blueprint/WidgetBlueprintGeneratedClass.h
+++ b/Engine/Source/Runtime/UMG/Public/Blueprint/WidgetBlueprintGeneratedClass.h
@@ -67,13 +67,25 @@ public:
 #if WITH_EDITORONLY_DATA
 
 	UPROPERTY()
-	bool bCookSlowConstructionWidgetTree;
+	uint8 bCookSlowConstructionWidgetTree:1;
 
 #endif
 
 	UPROPERTY()
-	bool bAllowTemplate;
+	uint8 bAllowTemplate:1;
+
+private:
 
+	UPROPERTY()
+	uint8 bValidTemplate:1;
+
+	UPROPERTY(Transient)
+	uint8 bTemplateInitialized:1;
+
+	UPROPERTY(Transient)
+	uint8 bCookedTemplate:1;
+
+public:
 	UPROPERTY()
 	TArray< FDelegateRuntimeBinding > Bindings;
 
@@ -120,14 +132,6 @@ public:
 private:
 	void InitializeTemplate(const ITargetPlatform* TargetPlatform);
 
-private:
-
-	UPROPERTY()
-	bool bValidTemplate;
-
-	UPROPERTY(Transient)
-	bool bTemplateInitialized;
-
 	UPROPERTY()
 	TSoftObjectPtr<UUserWidget> TemplateAsset;
 
diff --git a/Engine/Source/Runtime/UMG/Public/Components/Border.h b/Engine/Source/Runtime/UMG/Public/Components/Border.h
--- a/Engine/Source/Runtime/UMG/Public/Components/Border.h
+++ b/Engine/Source/Runtime/UMG/Public/Components/Border.h
@@ -1,4 +1,4 @@
-// Copyright 1998-2017 Epic Games, Inc. All Rights Reserved.
+﻿// Copyright 1998-2017 Epic Games, Inc. All Rights Reserved.
 
 #pragma once
 
@@ -32,6 +32,18 @@ class UMG_API UBorder : public UContentWidget
 
 public:
 
+	/** The alignment of the content horizontally. */
+	UPROPERTY(EditAnywhere, BlueprintReadOnly, Category="Content")
+	TEnumAsByte<EHorizontalAlignment> HorizontalAlignment;
+
+	/** The alignment of the content vertically. */
+	UPROPERTY(EditAnywhere, BlueprintReadOnly, Category="Content")
+	TEnumAsByte<EVerticalAlignment> VerticalAlignment;
+
+	/** Whether or not to show the disabled effect when this border is disabled */
+	UPROPERTY(EditAnywhere, BlueprintReadOnly, Category=Appearance, AdvancedDisplay)
+	uint8 bShowEffectWhenDisabled:1;
+
 	/** Color and opacity multiplier of content in the border */
 	UPROPERTY(EditAnywhere, BlueprintReadOnly, Category="Content", meta=( sRGB="true" ))
 	FLinearColor ContentColorAndOpacity;
@@ -44,14 +56,6 @@ public:
 	UPROPERTY(EditAnywhere, BlueprintReadOnly, Category="Content")
 	FMargin Padding;
 
-	/** The alignment of the content horizontally. */
-	UPROPERTY(EditAnywhere, BlueprintReadOnly, Category="Content")
-	TEnumAsByte<EHorizontalAlignment> HorizontalAlignment;
-
-	/** The alignment of the content vertically. */
-	UPROPERTY(EditAnywhere, BlueprintReadOnly, Category="Content")
-	TEnumAsByte<EVerticalAlignment> VerticalAlignment;
-
 	/** Brush to drag as the background */
 	UPROPERTY(EditAnywhere, BlueprintReadOnly, Category=Appearance, meta=( DisplayName="Brush" ))
 	FSlateBrush Background;
@@ -77,10 +81,6 @@ public:
 	UPROPERTY(EditAnywhere, BlueprintReadOnly, Category = Appearance)
 	FVector2D DesiredSizeScale;
 
-	/** Whether or not to show the disabled effect when this border is disabled */
-	UPROPERTY(EditAnywhere, BlueprintReadOnly, Category=Appearance, AdvancedDisplay)
-	bool bShowEffectWhenDisabled;
-
 public:
 
 	UPROPERTY(EditAnywhere, Category=Events, meta=( IsBindableEvent="True" ))
@@ -189,10 +189,12 @@ protected:
 	/** Translates the bound brush data and assigns it to the cached brush used by this widget. */
 	const FSlateBrush* ConvertImage(TAttribute<FSlateBrush> InImageAsset) const;
 
+#if WITH_EDITORONLY_DATA
 protected:
 	/** Image to use for the border */
 	UPROPERTY()
 	USlateBrushAsset* Brush_DEPRECATED;
+#endif
 
 	PROPERTY_BINDING_IMPLEMENTATION(FLinearColor, ContentColorAndOpacity)
 };
diff --git a/Engine/Source/Runtime/UMG/Public/Components/Image.h b/Engine/Source/Runtime/UMG/Public/Components/Image.h
--- a/Engine/Source/Runtime/UMG/Public/Components/Image.h
+++ b/Engine/Source/Runtime/UMG/Public/Components/Image.h
@@ -29,9 +29,11 @@ class UMG_API UImage : public UWidget
 
 public:
 
+#if WITH_EDITORONLY_DATA
 	/** Image to draw */
 	UPROPERTY()
 	USlateBrushAsset* Image_DEPRECATED;
+#endif
 
 	/** Image to draw */
 	UPROPERTY(EditAnywhere, BlueprintReadOnly, Category=Appearance)
@@ -96,9 +98,11 @@ public:
 	virtual void ReleaseSlateResources(bool bReleaseChildren) override;
 	//~ End UVisual Interface
 
+#if WITH_EDITORONLY_DATA
 	//~ Begin UObject Interface
 	virtual void PostLoad() override;
 	//~ End UObject Interface
+#endif
 
 #if WITH_EDITOR
 	//~ Begin UWidget Interface
diff --git a/Engine/Source/Runtime/UMG/Public/Components/Widget.h b/Engine/Source/Runtime/UMG/Public/Components/Widget.h
--- a/Engine/Source/Runtime/UMG/Public/Components/Widget.h
+++ b/Engine/Source/Runtime/UMG/Public/Components/Widget.h
@@ -77,6 +77,13 @@
 	:																\
 		TAttribute< ReturnType >(MemberName)
 
+#define BITFIELD_PROPERTY_BINDING(MemberName)						\
+	( MemberName ## Delegate.IsBound() && !IsDesignTime() )			\
+	?																\
+		BIND_UOBJECT_ATTRIBUTE(bool, K2_Gate_ ## MemberName)		\
+	:																\
+		TAttribute< bool >(MemberName != 0)
+
 #define PROPERTY_BINDING_IMPLEMENTATION(ReturnType, MemberName)			\
 	ReturnType K2_Cache_ ## MemberName;									\
 	ReturnType K2_Gate_ ## MemberName()									\
@@ -98,6 +105,13 @@
 	:															\
 		TAttribute< ReturnType >(MemberName)
 
+#define BITFIELD_PROPERTY_BINDING(MemberName)					\
+	( MemberName ## Delegate.IsBound() && !IsDesignTime() )		\
+	?															\
+		TAttribute< bool >::Create(MemberName ## Delegate.GetUObject(), MemberName ## Delegate.GetFunctionName()) \
+	:															\
+		TAttribute< bool >(MemberName != 0)
+
 #define PROPERTY_BINDING_IMPLEMENTATION(Type, MemberName) 
 
 #endif
@@ -196,22 +210,11 @@
 	typedef TFunctionRef<TSharedPtr<SObjectWidget>( UUserWidget*, TSharedRef<SWidget> )> ConstructMethodType;
 
 	/**
-	 * Allows controls to be exposed as variables in a blueprint.  Not all controls need to be exposed
-	 * as variables, so this allows only the most useful ones to end up being exposed.
-	 */
-	UPROPERTY()
-	bool bIsVariable;
-
-	/**
 	 * The parent slot of the UWidget.  Allows us to easily inline edit the layout controlling this widget.
 	 */
 	UPROPERTY(Instanced, EditAnywhere, BlueprintReadOnly, Category=Layout, meta=(ShowOnlyInnerProperties))
 	UPanelSlot* Slot;
 
-	/** Sets whether this widget can be modified interactively by the user */
-	UPROPERTY(EditAnywhere, BlueprintReadOnly, Category="Behavior")
-	bool bIsEnabled;
-
 	/** A bindable delegate for bIsEnabled */
 	UPROPERTY()
 	FGetBool bIsEnabledDelegate;
@@ -232,37 +235,46 @@
 	UPROPERTY()
 	FGetWidget ToolTipWidgetDelegate;
 
-	/** The visibility of the widget */
-	UPROPERTY()
-	ESlateVisibility Visiblity_DEPRECATED;
 
-	/** The visibility of the widget */
-	UPROPERTY(EditAnywhere, Category="Behavior")
-	ESlateVisibility Visibility;
-
 	/** A bindable delegate for Visibility */
 	UPROPERTY()
 	FGetSlateVisibility VisibilityDelegate;
 
+	/** A bindable delegate for Cursor */
+	//UPROPERTY()
+	//FGetMouseCursor CursorDelegate;
+
+public:
+
+	/** The render transform of the widget allows for arbitrary 2D transforms to be applied to the widget. */
+	UPROPERTY(EditAnywhere, BlueprintReadOnly, Category="Render Transform", meta=( DisplayName="Transform" ))
+	FWidgetTransform RenderTransform;
+
+	/**
+	 * The render transform pivot controls the location about which transforms are applied.  
+	 * This value is a normalized coordinate about which things like rotations will occur.
+	 */
+	UPROPERTY(EditAnywhere, BlueprintReadOnly, Category="Render Transform", meta=( DisplayName="Pivot" ))
+	FVector2D RenderTransformPivot;
+
+	/**
+	 * Allows controls to be exposed as variables in a blueprint.  Not all controls need to be exposed
+	 * as variables, so this allows only the most useful ones to end up being exposed.
+	 */
+	UPROPERTY()
+	uint8 bIsVariable:1;
+
+	/** Flag if the Widget was created from a blueprint */
+	UPROPERTY(Transient)
+	uint8 bCreatedByConstructionScript:1;
+
+	/** Sets whether this widget can be modified interactively by the user */
+	UPROPERTY(EditAnywhere, BlueprintReadOnly, Category="Behavior")
+	uint8 bIsEnabled:1;
+
 	/**  */
 	UPROPERTY(EditAnywhere, Category="Behavior", meta=(InlineEditConditionToggle))
-	uint32 bOverride_Cursor : 1;
+	uint8 bOverride_Cursor : 1;
 
-	/** The cursor to show when the mouse is over the widget */
-	UPROPERTY(EditAnywhere, BlueprintReadOnly, Category="Behavior", AdvancedDisplay, meta=( editcondition="bOverride_Cursor" ))
-	TEnumAsByte<EMouseCursor::Type> Cursor;
-
-	/**
-	 * Controls how the clipping behavior of this widget.  Normally content that overflows the
-	 * bounds of the widget continues rendering.  Enabling clipping prevents that overflowing content
-	 * from being seen.
-	 *
-	 * NOTE: Elements in different clipping spaces can not be batched together, and so there is a
-	 * performance cost to clipping.  Do not enable clipping unless a panel actually needs to prevent
-	 * content from showing up outside its bounds.
-	 */
-	UPROPERTY(EditAnywhere, Category = "Clipping")
-	EWidgetClipping Clipping;
-
 protected:
 
@@ -273,20 +285,42 @@
 	 * ever caching anything.
 	 */
 	UPROPERTY(EditAnywhere, BlueprintReadOnly, Category="Performance")
-	bool bIsVolatile;
+	uint8 bIsVolatile:1;
 
 public:
+#if WITH_EDITORONLY_DATA
+	/** Stores the design time flag setting if the widget is hidden inside the designer */
+	UPROPERTY()
+	uint8 bHiddenInDesigner:1;
 
-	/** The render transform of the widget allows for arbitrary 2D transforms to be applied to the widget. */
-	UPROPERTY(EditAnywhere, BlueprintReadOnly, Category="Render Transform", meta=( DisplayName="Transform" ))
-	FWidgetTransform RenderTransform;
+	/** Stores the design time flag setting if the widget is expanded inside the designer */
+	UPROPERTY()
+	uint8 bExpandedInDesigner:1;
+
+	/** Stores the design time flag setting if the widget is locked inside the designer */
+	UPROPERTY()
+	uint8 bLockedInDesigner:1;
+#endif
+
+	/** The cursor to show when the mouse is over the widget */
+	UPROPERTY(EditAnywhere, BlueprintReadOnly, Category="Behavior", AdvancedDisplay, meta=( editcondition="bOverride_Cursor" ))
+	TEnumAsByte<EMouseCursor::Type> Cursor;
 
 	/**
-	 * The render transform pivot controls the location about which transforms are applied.  
-	 * This value is a normalized coordinate about which things like rotations will occur.
+	 * Controls how the clipping behavior of this widget.  Normally content that overflows the
+	 * bounds of the widget continues rendering.  Enabling clipping prevents that overflowing content
+	 * from being seen.
+	 *
+	 * NOTE: Elements in different clipping spaces can not be batched together, and so there is a
+	 * performance cost to clipping.  Do not enable clipping unless a panel actually needs to prevent
+	 * content from showing up outside its bounds.
 	 */
-	UPROPERTY(EditAnywhere, BlueprintReadOnly, Category="Render Transform", meta=( DisplayName="Pivot" ))
-	FVector2D RenderTransformPivot;
+	UPROPERTY(EditAnywhere, Category = "Clipping")
+	EWidgetClipping Clipping;
+
+	/** The visibility of the widget */
+	UPROPERTY(EditAnywhere, Category="Behavior")
+	ESlateVisibility Visibility;
 
 	/**
 	 * The navigation object for this widget is optionally created if the user has configured custom
@@ -298,18 +332,6 @@
 
 #if WITH_EDITORONLY_DATA
 
-	/** Stores the design time flag setting if the widget is hidden inside the designer */
-	UPROPERTY()
-	bool bHiddenInDesigner;
-
-	/** Stores the design time flag setting if the widget is expanded inside the designer */
-	UPROPERTY()
-	bool bExpandedInDesigner;
-
-	/** Stores the design time flag setting if the widget is locked inside the designer */
-	UPROPERTY()
-	bool bLockedInDesigner;
-
 	/** Stores a reference to the asset responsible for this widgets construction. */
 	TWeakObjectPtr<UObject> WidgetGeneratedBy;
 	
@@ -632,7 +654,6 @@
 
 	// Begin UObject
 	virtual UWorld* GetWorld() const override;
-	virtual void PostLoad() override;
 	// End UObject
 
 #if WITH_EDITOR
