Copying //UE4/Dev-Sequencer to //UE4/Dev-Main (Source: //UE4/Dev-Sequencer @ 3946692)

diff --git a/Engine/Source/Runtime/MovieSceneTracks/Public/Tracks/MovieSceneParticleParameterTrack.h b/Engine/Source/Runtime/MovieSceneTracks/Public/Tracks/MovieSceneParticleParameterTrack.h
--- a/Engine/Source/Runtime/MovieSceneTracks/Public/Tracks/MovieSceneParticleParameterTrack.h
+++ b/Engine/Source/Runtime/MovieSceneTracks/Public/Tracks/MovieSceneParticleParameterTrack.h
@@ -28,7 +28,6 @@ public:
 	virtual void AddSection(UMovieSceneSection& Section) override;
 	virtual void RemoveSection(UMovieSceneSection& Section) override;
 	virtual bool IsEmpty() const override;
-	virtual TRange<float> GetSectionBoundaries() const override;
 	virtual const TArray<UMovieSceneSection*>& GetAllSections() const override;
 
 #if WITH_EDITORONLY_DATA
@@ -44,7 +43,7 @@ public:
 	 * @param Time The time to add the new key.
 	 * @param The value for the new key.
 	 */
-	void MOVIESCENETRACKS_API AddScalarParameterKey( FName ParameterName, float Position, float Value );
+	void MOVIESCENETRACKS_API AddScalarParameterKey( FName ParameterName, FFrameNumber Position, float Value );
 
 	/**
 	* Adds a Vector parameter key to the track.
@@ -52,7 +51,7 @@ public:
 	* @param Time The time to add the new key.
 	* @param The value for the new key.
 	*/
-	void MOVIESCENETRACKS_API AddVectorParameterKey( FName ParameterName, float Position, FVector Value );
+	void MOVIESCENETRACKS_API AddVectorParameterKey( FName ParameterName, FFrameNumber Position, FVector Value );
 
 	/**
 	* Adds a Vector parameter key to the track.
@@ -60,7 +59,7 @@ public:
 	* @param Time The time to add the new key.
 	* @param The value for the new key.
 	*/
-	void MOVIESCENETRACKS_API AddColorParameterKey( FName ParameterName, float Position, FLinearColor Value );
+	void MOVIESCENETRACKS_API AddColorParameterKey( FName ParameterName, FFrameNumber Position, FLinearColor Value );
 
 private:
 
diff --git a/Engine/Source/Runtime/MovieSceneTracks/Public/Tracks/MovieSceneParticleTrack.h b/Engine/Source/Runtime/MovieSceneTracks/Public/Tracks/MovieSceneParticleTrack.h
--- a/Engine/Source/Runtime/MovieSceneTracks/Public/Tracks/MovieSceneParticleTrack.h
+++ b/Engine/Source/Runtime/MovieSceneTracks/Public/Tracks/MovieSceneParticleTrack.h
@@ -37,9 +37,8 @@ public:
 	virtual void AddSection(UMovieSceneSection& Section) override;
 	virtual void RemoveSection(UMovieSceneSection& Section) override;
 	virtual bool IsEmpty() const override;
-	virtual TRange<float> GetSectionBoundaries() const override;
 	virtual const TArray<UMovieSceneSection*>& GetAllSections() const override;
-	virtual void AddNewSection(float SectionTime);
+	virtual void AddNewSection(FFrameNumber SectionTime);
 	virtual UMovieSceneSection* CreateNewSection() override;
 
 #if WITH_EDITORONLY_DATA
diff --git a/Engine/Source/Runtime/MovieSceneTracks/Public/Tracks/MovieScenePropertyTrack.h b/Engine/Source/Runtime/MovieSceneTracks/Public/Tracks/MovieScenePropertyTrack.h
--- a/Engine/Source/Runtime/MovieSceneTracks/Public/Tracks/MovieScenePropertyTrack.h
+++ b/Engine/Source/Runtime/MovieSceneTracks/Public/Tracks/MovieScenePropertyTrack.h
@@ -25,7 +25,6 @@ public:
 	virtual void AddSection(UMovieSceneSection& Section) override;
 	virtual void RemoveSection(UMovieSceneSection& Section) override;
 	virtual bool IsEmpty() const override;
-	virtual TRange<float> GetSectionBoundaries() const override;
 	virtual const TArray<UMovieSceneSection*>& GetAllSections() const override;
 
 #if WITH_EDITORONLY_DATA
@@ -51,6 +50,22 @@ public:
 	/** @return The property path for this track */
 	const FString& GetPropertyPath() const { return PropertyPath; }
 
+	/**
+	 * Finds a section at the current time.
+	 *
+	 * @param Time	The time relative to the owning movie scene where the section should be
+	 * @return The found section.
+	 */
+	class UMovieSceneSection* FindSection(FFrameNumber Time);
+
+	/**
+	 * Finds a section at the current time or extends an existing one
+	 *
+	 * @param Time	The time relative to the owning movie scene where the section should be
+	 * @return The found section.
+	 */
+	class UMovieSceneSection* FindOrExtendSection(FFrameNumber Time);
+
 	/**
 	 * Finds a section at the current time.
 	 *
@@ -58,7 +73,7 @@ public:
 	 * @param bSectionAdded Whether a section was added or not
 	 * @return The found section, or the new section.
 	 */
-	class UMovieSceneSection* FindOrAddSection(float Time, bool& bSectionAdded);
+	class UMovieSceneSection* FindOrAddSection(FFrameNumber Time, bool& bSectionAdded);
 
 #if WITH_EDITORONLY_DATA
 	/** Unique name for this track to afford multiple tracks on a given object (i.e. for array properties) */
diff --git a/Engine/Source/Runtime/MovieSceneTracks/Public/Tracks/MovieSceneSkeletalAnimationTrack.h b/Engine/Source/Runtime/MovieSceneTracks/Public/Tracks/MovieSceneSkeletalAnimationTrack.h
--- a/Engine/Source/Runtime/MovieSceneTracks/Public/Tracks/MovieSceneSkeletalAnimationTrack.h
+++ b/Engine/Source/Runtime/MovieSceneTracks/Public/Tracks/MovieSceneSkeletalAnimationTrack.h
@@ -19,10 +19,13 @@ class UMovieSceneSkeletalAnimationTrack
 public:
 
 	/** Adds a new animation to this track */
-	virtual void AddNewAnimation(float KeyTime, class UAnimSequenceBase* AnimSequence);
+	virtual UMovieSceneSection* AddNewAnimationOnRow(FFrameNumber KeyTime, class UAnimSequenceBase* AnimSequence, int32 RowIndex);
+
+	/** Adds a new animation to this track on the next available/non-overlapping row */
+	virtual UMovieSceneSection* AddNewAnimation(FFrameNumber KeyTime, class UAnimSequenceBase* AnimSequence) { return AddNewAnimationOnRow(KeyTime, AnimSequence, INDEX_NONE); }
 
 	/** Gets the animation sections at a certain time */
-	TArray<UMovieSceneSection*> GetAnimSectionsAtTime(float Time);
+	TArray<UMovieSceneSection*> GetAnimSectionsAtTime(FFrameNumber Time);
 
 public:
 
@@ -34,7 +37,6 @@ public:
 	virtual void AddSection(UMovieSceneSection& Section) override;
 	virtual void RemoveSection(UMovieSceneSection& Section) override;
 	virtual bool IsEmpty() const override;
-	virtual TRange<float> GetSectionBoundaries() const override;
 	virtual const TArray<UMovieSceneSection*>& GetAllSections() const override;
 	virtual UMovieSceneSection* CreateNewSection() override;
 	virtual FMovieSceneTrackRowSegmentBlenderPtr GetRowSegmentBlender() const override;
diff --git a/Engine/Source/Runtime/MovieSceneTracks/Public/Tracks/MovieSceneStringTrack.h b/Engine/Source/Runtime/MovieSceneTracks/Public/Tracks/MovieSceneStringTrack.h
--- a/Engine/Source/Runtime/MovieSceneTracks/Public/Tracks/MovieSceneStringTrack.h
+++ b/Engine/Source/Runtime/MovieSceneTracks/Public/Tracks/MovieSceneStringTrack.h
@@ -26,19 +26,6 @@ public:
 #endif
 	}
 
-public:
-
-	/**
-	 * Adds an string to the appropriate section.
-	 *
-	 * This method will create a new section if no appropriate one exists.
-	 *
-	 * @param Time The time at which the string should be triggered.
-	 * @param String The string to add.
-	 * @param KeyParams The keying parameters
-	 */
-	bool AddKeyToSection(float Time, const FString& String);
-
 public:
 
 	//~ UMovieSceneTrack interface
@@ -47,7 +34,6 @@ public:
 	virtual UMovieSceneSection* CreateNewSection() override;
 	virtual FMovieSceneEvalTemplatePtr CreateTemplateForSection(const UMovieSceneSection& InSection) const override;
 	virtual const TArray<UMovieSceneSection*>& GetAllSections() const override;
-	virtual TRange<float> GetSectionBoundaries() const override;
 	virtual bool HasSection(const UMovieSceneSection& Section) const override;
 	virtual bool IsEmpty() const override;
 	virtual void RemoveAllAnimationData() override;
diff --git a/Engine/Source/Runtime/Networking/Public/Common/UdpSocketReceiver.h b/Engine/Source/Runtime/Networking/Public/Common/UdpSocketReceiver.h
--- a/Engine/Source/Runtime/Networking/Public/Common/UdpSocketReceiver.h
+++ b/Engine/Source/Runtime/Networking/Public/Common/UdpSocketReceiver.h
@@ -4,12 +4,14 @@
 
 #include "CoreMinimal.h"
 #include "HAL/Runnable.h"
-#include "IPAddress.h"
-#include "SocketSubsystem.h"
-#include "Interfaces/IPv4/IPv4Endpoint.h"
 #include "HAL/RunnableThread.h"
+#include "Misc/SingleThreadRunnable.h"
 #include "Serialization/ArrayReader.h"
 #include "Sockets.h"
+#include "SocketSubsystem.h"
+
+#include "Interfaces/IPv4/IPv4Endpoint.h"
+#include "IPAddress.h"
 
 /**
  * Temporary fix for concurrency crashes. This whole class will be redesigned.
@@ -30,6 +32,7 @@ DECLARE_DELEGATE_TwoParams(FOnSocketDataReceived, const FArrayReaderPtr&, const
  */
 class FUdpSocketReceiver
 	: public FRunnable
+	, private FSingleThreadRunnable
 {
 public:
 
@@ -87,7 +90,12 @@ public:
 
 public:
 
-	// FRunnable interface
+	//~ FRunnable interface
+
+	virtual FSingleThreadRunnable* GetSingleThreadInterface() override
+	{
+		return this;
+	}
 
 	virtual bool Init() override
 	{
@@ -96,30 +104,9 @@ public:
 
 	virtual uint32 Run() override
 	{
-		TSharedRef<FInternetAddr> Sender = SocketSubsystem->CreateInternetAddr();
-
 		while (!Stopping)
 		{
-			if (!Socket->Wait(ESocketWaitConditions::WaitForRead, WaitTime))
-			{
-				continue;
-			}
-
-			uint32 Size;
-
-			while (Socket->HasPendingData(Size))
-			{
-				FArrayReaderPtr Reader = MakeShareable(new FArrayReader(true));
-				Reader->SetNumUninitialized(FMath::Min(Size, 65507u));
-
-				int32 Read = 0;
-				
-				if (Socket->RecvFrom(Reader->GetData(), Reader->Num(), Read, *Sender))
-				{
-					Reader->RemoveAt(Read, Reader->Num() - Read, false);
-					DataReceivedDelegate.ExecuteIfBound(Reader, FIPv4Endpoint(Sender));
-				}
-			}
+			Update(WaitTime);
 		}
 
 		return 0;
@@ -132,24 +119,61 @@ public:
 
 	virtual void Exit() override { }
 
+protected:
+
+	/** Update this socket receiver. */
+	void Update(const FTimespan& SocketWaitTime)
+	{
+		if (!Socket->Wait(ESocketWaitConditions::WaitForRead, SocketWaitTime))
+		{
+			return;
+		}
+
+		TSharedRef<FInternetAddr> Sender = SocketSubsystem->CreateInternetAddr();
+		uint32 Size;
+
+		while (Socket->HasPendingData(Size))
+		{
+			FArrayReaderPtr Reader = MakeShareable(new FArrayReader(true));
+			Reader->SetNumUninitialized(FMath::Min(Size, 65507u));
+
+			int32 Read = 0;
+
+			if (Socket->RecvFrom(Reader->GetData(), Reader->Num(), Read, *Sender))
+			{
+				Reader->RemoveAt(Read, Reader->Num() - Read, false);
+				DataReceivedDelegate.ExecuteIfBound(Reader, FIPv4Endpoint(Sender));
+			}
+		}
+	}
+
+protected:
+
+	//~ FSingleThreadRunnable interface
+
+	virtual void Tick() override
+	{
+		Update(FTimespan::Zero());
+	}
+
 private:
 
-	/** Holds the network socket. */
+	/** The network socket. */
 	FSocket* Socket;
 
-	/** Holds a pointer to the socket sub-system. */
+	/** Pointer to the socket sub-system. */
 	ISocketSubsystem* SocketSubsystem;
 
-	/** Holds a flag indicating that the thread is stopping. */
+	/** Flag indicating that the thread is stopping. */
 	bool Stopping;
 
-	/** Holds the thread object. */
+	/** The thread object. */
 	FRunnableThread* Thread;
 
 	/** The receiver thread's name. */
 	FString ThreadName;
 
-	/** Holds the amount of time to wait for inbound packets. */
+	/** The amount of time to wait for inbound packets. */
 	FTimespan WaitTime;
 
 private:
diff --git a/Engine/Source/Runtime/Networking/Public/Common/UdpSocketSender.h b/Engine/Source/Runtime/Networking/Public/Common/UdpSocketSender.h
--- a/Engine/Source/Runtime/Networking/Public/Common/UdpSocketSender.h
+++ b/Engine/Source/Runtime/Networking/Public/Common/UdpSocketSender.h
@@ -3,12 +3,14 @@
 #pragma once
 
 #include "CoreMinimal.h"
-#include "HAL/Runnable.h"
 #include "Containers/Queue.h"
-#include "Interfaces/IPv4/IPv4Endpoint.h"
+#include "HAL/Runnable.h"
 #include "HAL/RunnableThread.h"
+#include "Misc/SingleThreadRunnable.h"
 #include "Sockets.h"
 
+#include "Interfaces/IPv4/IPv4Endpoint.h"
+
 /**
  * Asynchronously sends data to an UDP socket.
  *
@@ -16,6 +18,7 @@
  */
 class FUdpSocketSender
 	: public FRunnable
+	, private FSingleThreadRunnable
 {
 	// Structure for outbound packets.
 	struct FPacket
@@ -138,7 +141,12 @@ public:
 
 public:
 
-	// FRunnable interface
+	//~ FRunnable interface
+
+	virtual FSingleThreadRunnable* GetSingleThreadInterface() override
+	{
+		return this;
+	}
 
 	virtual bool Init() override
 	{
@@ -149,23 +157,11 @@ public:
 	{
 		while (!Stopping)
 		{
-			while (!SendQueue.IsEmpty())
+			if (!Update(WaitTime))
 			{
-				if (Socket->Wait(ESocketWaitConditions::WaitForWrite, WaitTime))
-				{
-					FPacket Packet;
-					int32 Sent = 0;
-
-					SendQueue.Dequeue(Packet);
-					Socket->SendTo(Packet.Data->GetData(), Packet.Data->Num(), Sent, *Packet.Recipient.ToInternetAddr());
-					
-					if (Sent != Packet.Data->Num())
-					{
-						Stopping = true;
-
-						return 0;
-					}
-				}
+				Stopping = true;
+
+				return 0;
 			}
 
 			WorkEvent->Wait(WaitTime);
@@ -182,29 +178,70 @@ public:
 
 	virtual void Exit() override { }
 
+protected:
+
+	/**
+	 * Update this socket sender.
+	 *
+	 * @param Time to wait for the socket.
+	 * @return true on success, false otherwise.
+	 */
+	bool Update(const FTimespan& SocketWaitTime)
+	{
+		while (!SendQueue.IsEmpty())
+		{
+			if (!Socket->Wait(ESocketWaitConditions::WaitForWrite, SocketWaitTime))
+			{
+				break;
+			}
+
+			FPacket Packet;
+			int32 Sent = 0;
+
+			SendQueue.Dequeue(Packet);
+			Socket->SendTo(Packet.Data->GetData(), Packet.Data->Num(), Sent, *Packet.Recipient.ToInternetAddr());
+
+			if (Sent != Packet.Data->Num())
+			{
+				return false;
+			}
+		}
+
+		return true;
+	}
+
+protected:
+
+	//~ FSingleThreadRunnable interface
+
+	virtual void Tick() override
+	{
+		Update(FTimespan::Zero());
+	}
+
 private:
 
-	/** Holds the send queue. */
+	/** The send queue. */
 	TQueue<FPacket, EQueueMode::Mpsc> SendQueue;
 
-	/** Holds the send rate. */
+	/** The send rate. */
 	uint32 SendRate;
 
-	/** Holds the network socket. */
+	/** The network socket. */
 	FSocket* Socket;
 
-	/** Holds a flag indicating that the thread is stopping. */
+	/** Flag indicating that the thread is stopping. */
 	bool Stopping;
 
-	/** Holds the thread object. */
+	/** The thread object. */
 	FRunnableThread* Thread;
 
-	/** Holds the current throughput. */
+	/** The current throughput. */
 	uint32 Throughput;
 
-	/** Holds the amount of time to wait for inbound packets. */
+	/** The amount of time to wait for inbound packets. */
 	FTimespan WaitTime;
 
-	/** Holds an event signaling that inbound messages need to be processed. */
+	/** An event signaling that inbound messages need to be processed. */
 	FEvent* WorkEvent;
 };
diff --git a/Engine/Source/Runtime/Slate/Public/Widgets/Input/SSpinBox.h b/Engine/Source/Runtime/Slate/Public/Widgets/Input/SSpinBox.h
--- a/Engine/Source/Runtime/Slate/Public/Widgets/Input/SSpinBox.h
+++ b/Engine/Source/Runtime/Slate/Public/Widgets/Input/SSpinBox.h
@@ -101,6 +101,8 @@ public:
 		SLATE_ATTRIBUTE( NumericType, Delta )
 		/** How many pixel the mouse must move to change the value of the delta step */
 		SLATE_ATTRIBUTE( int32, ShiftMouseMovePixelPerDelta )
+		/** If we're an unbounded spinbox, what value do we divide mouse movement by before multiplying by Delta. Requires Delta to be set. */
+		SLATE_ATTRIBUTE( int32, LinearDeltaSensitivity)
 		/** Tell us if we want to support dynamically changing of the max value using ctrl */
 		SLATE_ATTRIBUTE(bool, SupportDynamicSliderMaxValue)
 		/** Tell us if we want to support dynamically changing of the min value using ctrl */
@@ -197,6 +199,7 @@ public:
 
 		Delta = InArgs._Delta;
 		ShiftMouseMovePixelPerDelta = InArgs._ShiftMouseMovePixelPerDelta;
+		LinearDeltaSensitivity = InArgs._LinearDeltaSensitivity;
 	
 		BackgroundHoveredBrush = &InArgs._Style->HoveredBackgroundBrush;
 		BackgroundBrush = &InArgs._Style->BackgroundBrush;
@@ -567,11 +570,19 @@ public:
 				}
 				else
 				{
-					//we have no range specified, so increment scaled by our current value		
-					const double CurrentValue = FMath::Clamp<double>(FMath::Abs(InternalValue),1.0,TNumericLimits<NumericType>::Max());
-					const float MouseDelta = FMath::Abs(MouseEvent.GetCursorDelta().X / SliderWidthInSlateUnits);
+					// If this control has a specified delta and sensitivity then we use that instead of the current value for determining how much to change.
 					const float Sign = (MouseEvent.GetCursorDelta().X > 0) ? 1.f : -1.f;
-					NewValue = InternalValue + (Sign * MouseDelta * FMath::Pow(CurrentValue, SliderExponent.Get()));
+					if (LinearDeltaSensitivity.IsSet() && Delta.IsSet() && Delta.Get() > 0)
+					{
+						const float MouseDelta = FMath::Abs(MouseEvent.GetCursorDelta().X / LinearDeltaSensitivity.Get());
+						NewValue = InternalValue + (Sign * MouseDelta * FMath::Pow(Delta.Get(), SliderExponent.Get()));
+					}
+					else
+					{
+						const float MouseDelta = FMath::Abs(MouseEvent.GetCursorDelta().X / SliderWidthInSlateUnits);
+						const double CurrentValue = FMath::Clamp<double>(FMath::Abs(InternalValue),1.0,TNumericLimits<NumericType>::Max());
+						NewValue = InternalValue + (Sign * MouseDelta * FMath::Pow(CurrentValue, SliderExponent.Get()));
+					}
 				}
 			
 				CommitValue( NewValue, CommittedViaSpin, ETextCommit::OnEnter );
@@ -936,6 +947,7 @@ private:
 	float DistanceDragged;
 	TAttribute<NumericType> Delta;
 	TAttribute<int32> ShiftMouseMovePixelPerDelta;
+	TAttribute<int32> LinearDeltaSensitivity;
 	TAttribute<float> SliderExponent;
 	TAttribute<NumericType> SliderExponentNeutralValue;
 	TAttribute< TOptional<NumericType> > MinValue;
diff --git a/Engine/Source/Runtime/Slate/Public/Widgets/Views/STableRow.h b/Engine/Source/Runtime/Slate/Public/Widgets/Views/STableRow.h
--- a/Engine/Source/Runtime/Slate/Public/Widgets/Views/STableRow.h
+++ b/Engine/Source/Runtime/Slate/Public/Widgets/Views/STableRow.h
@@ -679,7 +679,7 @@ public:
 				if (ReportedZone.IsSet())
 				{
 					FReply DropReply = OnAcceptDrop.Execute(DragDropEvent, ReportedZone.GetValue(), *MyItem);
-					if (DropReply.IsEventHandled())
+					if (DropReply.IsEventHandled() && ReportedZone.GetValue() == EItemDropZone::OntoItem)
 					{
 						// Expand the drop target just in case, so that what we dropped is visible.
 						OwnerWidget->Private_SetItemExpansion(*MyItem, true);
diff --git a/Engine/Source/Runtime/TimeManagement/Private/TimeManagementModule.cpp b/Engine/Source/Runtime/TimeManagement/Private/TimeManagementModule.cpp
new file mode 100644
--- /dev/null
+++ b/Engine/Source/Runtime/TimeManagement/Private/TimeManagementModule.cpp
@@ -0,0 +1,5 @@
+// Copyright 1998-2018 Epic Games, Inc. All Rights Reserved.
+
+#include "Modules/ModuleManager.h"
+
+IMPLEMENT_MODULE(FDefaultModuleImpl, TimeManagement);
\ No newline at end of file
diff --git a/Engine/Source/Runtime/TimeManagement/Public/CommonFrameRates.h b/Engine/Source/Runtime/TimeManagement/Public/CommonFrameRates.h
new file mode 100644
--- /dev/null
+++ b/Engine/Source/Runtime/TimeManagement/Public/CommonFrameRates.h
@@ -0,0 +1,66 @@
+// Copyright 1998-2018 Epic Games, Inc. All Rights Reserved.
+
+#pragma once
+
+#include "FrameRate.h"
+#include "Containers/ArrayView.h"
+
+enum class ECommonFrameRate : uint8
+{
+	FPS_12,
+	FPS_15,
+	FPS_24,
+	FPS_25,
+	FPS_30,
+	FPS_48,
+	FPS_50,
+	FPS_60,
+	FPS_100,
+	FPS_120,
+	FPS_240,
+	NTSC_24,
+	NTSC_30,
+	NTSC_60,
+
+	Private_Num
+};
+
+struct FCommonFrameRateInfo
+{
+	FFrameRate FrameRate;
+	FText DisplayName;
+	FText Description;
+};
+
+struct TIMEMANAGEMENT_API FCommonFrameRates
+{
+	typedef __underlying_type(ECommonFrameRate) NumericType;
+
+	FORCEINLINE static FFrameRate FPS_12()  { return AllFrameRates[(NumericType)ECommonFrameRate::FPS_12].FrameRate;  }
+	FORCEINLINE static FFrameRate FPS_15()  { return AllFrameRates[(NumericType)ECommonFrameRate::FPS_15].FrameRate;  }
+	FORCEINLINE static FFrameRate FPS_24()  { return AllFrameRates[(NumericType)ECommonFrameRate::FPS_24].FrameRate;  }
+	FORCEINLINE static FFrameRate FPS_25()  { return AllFrameRates[(NumericType)ECommonFrameRate::FPS_25].FrameRate;  }
+	FORCEINLINE static FFrameRate FPS_30()  { return AllFrameRates[(NumericType)ECommonFrameRate::FPS_30].FrameRate;  }
+	FORCEINLINE static FFrameRate FPS_48()  { return AllFrameRates[(NumericType)ECommonFrameRate::FPS_48].FrameRate;  }
+	FORCEINLINE static FFrameRate FPS_50()  { return AllFrameRates[(NumericType)ECommonFrameRate::FPS_50].FrameRate;  }
+	FORCEINLINE static FFrameRate FPS_60()  { return AllFrameRates[(NumericType)ECommonFrameRate::FPS_60].FrameRate;  }
+	FORCEINLINE static FFrameRate FPS_100() { return AllFrameRates[(NumericType)ECommonFrameRate::FPS_100].FrameRate; }
+	FORCEINLINE static FFrameRate FPS_120() { return AllFrameRates[(NumericType)ECommonFrameRate::FPS_120].FrameRate; }
+	FORCEINLINE static FFrameRate FPS_240() { return AllFrameRates[(NumericType)ECommonFrameRate::FPS_240].FrameRate; }
+
+	FORCEINLINE static FFrameRate NTSC_24() { return AllFrameRates[(NumericType)ECommonFrameRate::NTSC_24].FrameRate; }
+	FORCEINLINE static FFrameRate NTSC_30() { return AllFrameRates[(NumericType)ECommonFrameRate::NTSC_30].FrameRate; }
+	FORCEINLINE static FFrameRate NTSC_60() { return AllFrameRates[(NumericType)ECommonFrameRate::NTSC_60].FrameRate; }
+
+	static TArrayView<const FCommonFrameRateInfo> GetAll();
+
+	static bool Contains(FFrameRate FrameRateToCheck)
+	{
+		return Find(FrameRateToCheck) != nullptr;
+	}
+
+	static const FCommonFrameRateInfo* Find(FFrameRate InFrameRate);
+
+private:
+	static const FCommonFrameRateInfo AllFrameRates[(int32)ECommonFrameRate::Private_Num];
+};
\ No newline at end of file
diff --git a/Engine/Source/Runtime/TimeManagement/Public/FrameNumberNumericInterface.h b/Engine/Source/Runtime/TimeManagement/Public/FrameNumberNumericInterface.h
new file mode 100644
--- /dev/null
+++ b/Engine/Source/Runtime/TimeManagement/Public/FrameNumberNumericInterface.h
@@ -0,0 +1,168 @@
+// Copyright 1998-2018 Epic Games, Inc. All Rights Reserved.
+
+#pragma once
+
+#include "CoreMinimal.h"
+#include "FrameNumberDisplayFormat.h"
+#include "Widgets/Input/NumericTypeInterface.h"
+#include "FrameRate.h"
+#include "FrameNumberTimeEvaluator.h"
+#include "Timecode.h"
+
+
+DECLARE_DELEGATE_RetVal(EFrameNumberDisplayFormats, FOnGetDisplayFormat)
+DECLARE_DELEGATE_RetVal(uint8, FOnGetZeroPad)
+DECLARE_DELEGATE_RetVal(FFrameRate, FOnGetFrameRate)
+
+/**
+* A large portion of the Sequencer UI is built around using SpinBox and NumericBox because the framerate
+* used to be stored in (float) seconds. This creates a convenient UI as it allows the users to type in
+* numbers (for frames or seconds), or to drag the mouse to change the time. When Sequencer was converted
+* to using integer based frames and frame rates (expressed as numerator/denominator) the spinbox-based
+* UI became an issue as SpinBox and NumericBox internally convert to a double to perform various calculations.
+* This is an issue if your Spinbox type (ie: FQualifiedFrameTime) is not implicitly convertible to double,
+* which the new Sequencer types are not (due to accidental type coercion when using outside the UI).
+*
+* To solve this, the Spinboxes will still use double as their type, but double now means frame number and not
+* time. A double can store the entire range of int32 (which backs FFrameNumber) without precision loss, and
+* we can execute callbacks to find out what framerate the sequence is running at. This allows us to display
+* in Timecode, Time or Frames and convert back and forth to double for the UI, and from double into the backing
+* FFrameNumber for the sequence.
+*/
+struct FFrameNumberInterface : public INumericTypeInterface<double>
+{
+	FFrameNumberInterface(FOnGetDisplayFormat InGetDisplayFormat, FOnGetZeroPad InOnGetZeroPadFrameNumber, FOnGetFrameRate InGetFrameResolution, FOnGetFrameRate InGetPlayRate)
+		: GetDisplayFormat(InGetDisplayFormat)
+		, GetFrameResolution(InGetFrameResolution)
+		, GetPlayRate(InGetPlayRate)
+		, GetZeroPadFrames(InOnGetZeroPadFrameNumber)
+	{
+		check(InGetDisplayFormat.IsBound());
+		check(InGetFrameResolution.IsBound());
+		check(InGetPlayRate.IsBound());
+	}
+
+private:
+	FOnGetDisplayFormat GetDisplayFormat;
+	FOnGetFrameRate GetFrameResolution;
+	FOnGetFrameRate GetPlayRate;
+	FOnGetZeroPad GetZeroPadFrames;
+
+	/** Check whether the typed character is valid */
+	virtual bool IsCharacterValid(TCHAR InChar) const override
+	{
+		auto IsValidLocalizedCharacter = [InChar]() -> bool
+		{
+			const FDecimalNumberFormattingRules& NumberFormattingRules = ExpressionParser::GetLocalizedNumberFormattingRules();
+			return InChar == NumberFormattingRules.GroupingSeparatorCharacter
+				|| InChar == NumberFormattingRules.DecimalSeparatorCharacter
+				|| Algo::Find(NumberFormattingRules.DigitCharacters, InChar) != 0;
+		};
+
+		static const FString ValidChars = TEXT("1234567890()-+=\\/.,*^%%hrmsf[]:; ");
+		return InChar != 0 && (ValidChars.GetCharArray().Contains(InChar) || IsValidLocalizedCharacter());
+	}
+
+	virtual FString ToString(const double& Value) const override
+	{
+		FFrameRate SourceFrameRate = GetFrameResolution.Execute();
+		FFrameRate DestinationFrameRate = GetPlayRate.Execute();
+		EFrameNumberDisplayFormats Format = GetDisplayFormat.Execute();
+
+		// If they want Drop Frame Timecode format but we're in an unsupported frame rate, we'll override it and say they want non drop frame.
+		bool bIsValidRateForDropFrame = FTimecode::IsDropFormatTimecodeSupported(DestinationFrameRate);
+		if (Format == EFrameNumberDisplayFormats::DropFrameTimecode && !bIsValidRateForDropFrame)
+		{
+			Format = EFrameNumberDisplayFormats::NonDropFrameTimecode;
+		}
+
+		switch (Format)
+		{
+		case EFrameNumberDisplayFormats::Frames:
+		{
+			// Convert from sequence resolution into display rate frames.
+			FFrameTime DisplayTime = FFrameRate::TransformTime(FFrameTime::FromDecimal(Value), SourceFrameRate, DestinationFrameRate);
+			FString SubframeIndicator = FMath::IsNearlyZero(DisplayTime.GetSubFrame()) ? TEXT("") : TEXT("*");
+
+			return FString::Printf(TEXT("%0*d%s"), GetZeroPadFrames.Execute(), DisplayTime.GetFrame().Value, *SubframeIndicator);
+		}
+		case EFrameNumberDisplayFormats::Seconds:
+		{
+			double TimeInSeconds = SourceFrameRate.AsSeconds(FFrameTime::FromDecimal(Value));
+			return FString::Printf(TEXT("%.2f s"), TimeInSeconds);
+		}
+		case EFrameNumberDisplayFormats::NonDropFrameTimecode:
+		{
+			FFrameTime InternalFrameNumber = FFrameTime::FromDecimal(Value);
+			FFrameNumber PlayRateFrameNumber = FFrameRate::TransformTime(InternalFrameNumber, SourceFrameRate, DestinationFrameRate).FloorToFrame();
+
+			FTimecode AsNonDropTimecode = FTimecode::FromFrameNumber(PlayRateFrameNumber, DestinationFrameRate, false);
+			return FString::Printf(TEXT("[%s]"), *AsNonDropTimecode.ToString(false));
+		}
+		case EFrameNumberDisplayFormats::DropFrameTimecode:
+		{
+			FFrameTime InternalFrameNumber = FFrameTime::FromDecimal(Value);
+			FFrameNumber PlayRateFrameNumber = FFrameRate::TransformTime(InternalFrameNumber, SourceFrameRate, DestinationFrameRate).FloorToFrame();
+
+			FTimecode AsDropTimecode = FTimecode::FromFrameNumber(PlayRateFrameNumber, DestinationFrameRate, true);
+			return FString::Printf(TEXT("[%s]"), *AsDropTimecode.ToString(false));
+		}
+		default:
+			return FString(TEXT("Unsupported Format"));
+		}
+	}
+
+	virtual TOptional<double> FromString(const FString& InString, const double& InExistingValue) override
+	{
+		FFrameRate SourceFrameRate = GetPlayRate.Execute();
+		FFrameRate DestinationFrameRate = GetFrameResolution.Execute();
+		EFrameNumberDisplayFormats FallbackFormat = GetDisplayFormat.Execute();
+
+		// We allow input in any format (time, frames or timecode) and we just convert it into into the internal sequence resolution.
+		// The user's input can be ambiguous though (does "5" mean 5 frames or 5 seconds?) so when we check each possible result we
+		// also check to see if they explicitly specified that format, or if the evaluator just happens to be able to parse that.
+
+		// All of these will convert into the frame resolution from the user's input before returning.
+		FFrameNumberTimeEvaluator Eval;
+		TValueOrError<FFrameTime, FExpressionError> TimecodeResult = Eval.EvaluateTimecode(*InString, SourceFrameRate, DestinationFrameRate);
+		bool bWasFrameText;
+		TValueOrError<FFrameTime, FExpressionError> FrameResult = Eval.EvaluateFrame(*InString, SourceFrameRate, DestinationFrameRate, /*Out*/ bWasFrameText);
+		bool bWasTimeText;
+		TValueOrError<FFrameTime, FExpressionError> TimeResult = Eval.EvaluateTime(*InString, DestinationFrameRate, /*Out*/ bWasTimeText);
+
+		if (TimecodeResult.IsValid())
+		{
+			return TOptional<double>(TimecodeResult.GetValue().GetFrame().Value);
+		}
+		else if (!bWasFrameText && !bWasTimeText)
+		{
+			// If the value was ambiguous ("5") and not specifically a frame ("5f") or a time "(5s") then we defer to the current display format.
+			if (TimeResult.IsValid() && FallbackFormat == EFrameNumberDisplayFormats::Seconds)
+			{
+				return TOptional<double>(TimeResult.GetValue().GetFrame().Value);
+			}
+			else if (FrameResult.IsValid() && FallbackFormat == EFrameNumberDisplayFormats::Frames)
+			{
+				return TOptional<double>(FrameResult.GetValue().GetFrame().Value);
+			}
+			else
+			{ 
+				// We don't support ambiguous conversion to Timecode
+				return TOptional<double>();
+			}
+		}
+		else if (TimeResult.IsValid() && !bWasFrameText)
+		{
+			// Both time and frames can parse "5" successfully, so if we know it's not specifically frame text
+			return TOptional<double>(TimeResult.GetValue().GetFrame().Value);
+		}
+		else if (FrameResult.IsValid() && !bWasTimeText)
+		{
+			// Both time and frames can parse "5" successfully and we know it's not specifically time
+			return TOptional<double>(FrameResult.GetValue().GetFrame().Value);
+		}
+
+		// We're not sure what they typed in
+		return TOptional<double>();
+	}
+};
\ No newline at end of file
diff --git a/Engine/Source/Runtime/TimeManagement/Public/FrameNumberTimeEvaluator.h b/Engine/Source/Runtime/TimeManagement/Public/FrameNumberTimeEvaluator.h
new file mode 100644
--- /dev/null
+++ b/Engine/Source/Runtime/TimeManagement/Public/FrameNumberTimeEvaluator.h
@@ -0,0 +1,50 @@
+// Copyright 1998-2018 Epic Games, Inc. All Rights Reserved.
+
+#pragma once
+
+#include "CoreTypes.h"
+#include "Misc/Optional.h"
+#include "Templates/ValueOrError.h"
+#include "Misc/ExpressionParserTypes.h"
+
+struct FFrameTime;
+struct FFrameRate;
+
+#define DEFINE_EXPRESSION_OPERATOR_NODE2(TYPE, ...) \
+namespace ExpressionParser {\
+	struct TIMEMANAGEMENT_API TYPE { static const TCHAR* const Moniker; }; \
+}\
+DEFINE_EXPRESSION_NODE_TYPE(ExpressionParser::TYPE, __VA_ARGS__)
+
+// @todo sequencer-timecode This is a duplicate from the BasicMathEvaluator, but that macro exports with CORE_API so we can't use it here.
+// We should probably just explicitly define these nodes in the namespace instead of duplicating the macro.
+DEFINE_EXPRESSION_OPERATOR_NODE2(FHour, 0x702443E8, 0xEF9A48A0, 0x8DC56394, 0x48F2632A);
+DEFINE_EXPRESSION_OPERATOR_NODE2(FMinute, 0x4FED0D41, 0x298D481C, 0xAC022899, 0xF196E817);
+DEFINE_EXPRESSION_OPERATOR_NODE2(FSecond, 0x43365BB1, 0x9A6B457E, 0xA6C71BB3, 0xCE6D5ACD);
+DEFINE_EXPRESSION_OPERATOR_NODE2(FMillisecond, 0xBFA9E4FB, 0x45799275, 0x0EFA76E3, 0x76E3533F);
+DEFINE_EXPRESSION_OPERATOR_NODE2(FFrames, 0xA0D341C4, 0xA51D40D3, 0xA9DBDD45, 0x83E1A21A);
+DEFINE_EXPRESSION_OPERATOR_NODE2(FTimecodeDelimiter, 0x6E947008, 0x523D4A17, 0xB5F7BCDA, 0xA43B62FC);
+DEFINE_EXPRESSION_OPERATOR_NODE2(FDropcodeDelimiter, 0x80AF2C2C, 0x544A451F, 0x822E5B24, 0xAFC78EEF);
+DEFINE_EXPRESSION_OPERATOR_NODE2(FBracketStart, 0xA7358BD1, 0xB4EF466D, 0xA336CD84, 0x14C93D2E);
+DEFINE_EXPRESSION_OPERATOR_NODE2(FBracketEnd, 0xE49500E9, 0x03E64440, 0x87802630, 0xF1C1DDDF);
+
+struct TIMEMANAGEMENT_API FFrameNumberTimeEvaluator
+{
+	/** Constructor that sets up the parser's lexer and compiler */
+	FFrameNumberTimeEvaluator();
+
+	/** Evaluate the given expression, resulting in either a double value, or an error */
+	TValueOrError<FFrameTime, FExpressionError> EvaluateTimecode(const TCHAR* InExpression, const FFrameRate& InDisplayFrameRate, const FFrameRate& InFrameResolution) const;
+	TValueOrError<FFrameTime, FExpressionError> EvaluateFrame(const TCHAR* InExpression, const FFrameRate& InDisplayFrameRate, const FFrameRate& InFrameResolution, bool& OutDirectlyParsed) const;
+	TValueOrError<FFrameTime, FExpressionError> EvaluateTime(const TCHAR* InExpression, FFrameRate InFrameRate, bool& OutDirectlyParsed) const;
+
+	FTokenDefinitions TimecodeTokenDefinitions;
+	FTokenDefinitions FrameTokenDefinitions;
+	FTokenDefinitions TimeTokenDefinitions;
+
+	FExpressionGrammar FrameGrammar;
+	FExpressionGrammar TimecodeGrammar;
+	FExpressionGrammar TimeGrammar;
+
+	FOperatorJumpTable TimeJumpTable;
+};
\ No newline at end of file
diff --git a/Engine/Source/Runtime/TimeManagement/Public/FrameTime.h b/Engine/Source/Runtime/TimeManagement/Public/FrameTime.h
new file mode 100644
--- /dev/null
+++ b/Engine/Source/Runtime/TimeManagement/Public/FrameTime.h
@@ -0,0 +1,275 @@
+// Copyright 1998-2018 Epic Games, Inc. All Rights Reserved.
+
+#pragma once
+
+#include "FrameNumber.h"
+#include "FrameTime.generated.h"
+
+
+/**
+ * Structure representing a time by a context-free frame number, plus a sub frame value in the range [0:1)
+ * Conversion to and from time in seconds is achieved in combination with FFrameRate.
+ * Only the frame number part of this representation can be negative, sub frames are always a positive value between the frame number and its next logical frame
+ */
+USTRUCT(BlueprintType)
+struct FFrameTime
+{
+	GENERATED_BODY()
+
+	/**
+	 * Default constructor initializing to zero
+	 */
+	FFrameTime();
+
+	/**
+	 * Implicit construction from a single integer, while disallowing implicit conversion from any other numeric type
+	 */
+	template<typename T, typename = typename TEnableIf<TIsSame<T, int32>::Value>::Type>
+	FFrameTime(T /* int32 */ InFrameNumber);
+
+	/**
+	 * Implicit construction from a type-safe frame number
+	 */
+	FFrameTime(FFrameNumber InFrameNumber);
+
+	/**
+	 * Construction from a frame number and a sub frame
+	 */
+	FFrameTime(FFrameNumber InFrameNumber, float InSubFrame);
+
+	/**
+	 * Assignment from a type-safe frame number
+	 */
+	FFrameTime& operator=(FFrameNumber InFrameNumber);
+
+public:
+
+	/**
+	 * Access this time's frame number
+	 */
+	FORCEINLINE FFrameNumber GetFrame() const
+	{
+		return FrameNumber;
+	}
+
+	/**
+	 * Access this time's sub frame
+	 */
+	FORCEINLINE float GetSubFrame() const
+	{
+		return SubFrame;
+	}
+
+	/**
+	 * Return the first frame number less than or equal to this frame time
+	 */
+	FFrameNumber FloorToFrame() const;
+
+	/**
+	 * Return the next frame number greater than or equal to this frame time
+	 */
+	FFrameNumber CeilToFrame() const;
+
+	/**
+	 * Round to the nearest frame number
+	 */
+	FFrameNumber RoundToFrame() const;
+
+	/**
+	 * Retrieve a decimal representation of this frame time
+	 * Sub frames are always added to the current frame number, so for negative frame times, a time of -10 [sub frame 0.25] will yield a decimal value of -9.75.
+	 */
+	double AsDecimal() const;
+
+	/**
+	 * Convert a decimal representation to a frame time
+	 * Note that sub frames are always positive, so negative decimal representations result in an inverted sub frame and floored frame number
+	 */
+	static FFrameTime FromDecimal(double InDecimalFrame);
+
+
+private:
+
+	/*~ Built in operators */
+	friend bool operator==(FFrameTime A, FFrameTime B);
+	friend bool operator!=(FFrameTime A, FFrameTime B);
+	friend bool operator> (FFrameTime A, FFrameTime B);
+	friend bool operator>=(FFrameTime A, FFrameTime B);
+	friend bool operator< (FFrameTime A, FFrameTime B);
+	friend bool operator<=(FFrameTime A, FFrameTime B);
+
+	friend FFrameTime& operator+=(FFrameTime& LHS, FFrameTime RHS);
+	friend FFrameTime  operator+(FFrameTime A, FFrameTime B);
+	friend FFrameTime  operator-(FFrameTime A, FFrameTime B);
+
+	friend FFrameTime  operator-(FFrameTime A);
+
+	friend FFrameTime  operator*(FFrameTime A, float Scalar);
+	friend FFrameTime  operator/(FFrameTime A, float Scalar);
+
+public:
+
+	UPROPERTY(BlueprintReadWrite, Category = General)
+	FFrameNumber FrameNumber;
+
+private:
+
+	/** Must be 0.f <= SubFrame < 1.f */
+	UPROPERTY()
+	float SubFrame;
+};
+
+
+inline FFrameTime::FFrameTime()
+	: FrameNumber(0), SubFrame(0.f)
+{}
+
+
+template<typename T, typename>
+FFrameTime::FFrameTime(T InFrameNumber)
+	: FrameNumber(InFrameNumber), SubFrame(0.f)
+{}
+
+
+inline FFrameTime::FFrameTime(FFrameNumber InFrameNumber)
+	: FrameNumber(InFrameNumber), SubFrame(0.f)
+{}
+
+
+inline FFrameTime::FFrameTime(FFrameNumber InFrameNumber, float InSubFrame)
+	: FrameNumber(InFrameNumber), SubFrame(InSubFrame)
+{
+	// Hack to ensure that SubFrames are in a sensible range of precision to work around
+	// problems with FloorToXYZ returning the wrong thing for very small negative numbers
+	SubFrame = FMath::Clamp(SubFrame + 0.5f - 0.5f, 0.f, 0.99999994f);
+	checkSlow(InSubFrame >= 0.f && InSubFrame < 1.f);
+}
+
+
+inline FFrameTime& FFrameTime::operator=(FFrameNumber InFrameNumber)
+{
+	FrameNumber = InFrameNumber;
+	SubFrame    = 0.f;
+	return *this;
+}
+
+
+FORCEINLINE_DEBUGGABLE bool operator==(FFrameTime A, FFrameTime B)
+{
+	return A.FrameNumber == B.FrameNumber && A.SubFrame == B.SubFrame;
+}
+
+
+FORCEINLINE_DEBUGGABLE bool operator!=(FFrameTime A, FFrameTime B)
+{
+	return A.FrameNumber != B.FrameNumber || A.SubFrame != B.SubFrame;
+}
+
+
+FORCEINLINE_DEBUGGABLE bool operator> (FFrameTime A, FFrameTime B)
+{
+	return A.FrameNumber >  B.FrameNumber || ( A.FrameNumber == B.FrameNumber && A.SubFrame > B.SubFrame );
+}
+
+
+FORCEINLINE_DEBUGGABLE bool operator>=(FFrameTime A, FFrameTime B)
+{
+	return A.FrameNumber > B.FrameNumber || ( A.FrameNumber == B.FrameNumber && A.SubFrame >= B.SubFrame );
+}
+
+
+FORCEINLINE_DEBUGGABLE bool operator< (FFrameTime A, FFrameTime B)
+{
+	return A.FrameNumber <  B.FrameNumber || ( A.FrameNumber == B.FrameNumber && A.SubFrame < B.SubFrame );
+}
+
+
+FORCEINLINE_DEBUGGABLE bool operator<=(FFrameTime A, FFrameTime B)
+{
+	return A.FrameNumber < B.FrameNumber || ( A.FrameNumber == B.FrameNumber && A.SubFrame <= B.SubFrame );
+}
+
+
+FORCEINLINE_DEBUGGABLE FFrameTime& operator+=(FFrameTime& LHS, FFrameTime RHS)
+{
+	float NewSubFrame = LHS.SubFrame + RHS.SubFrame;
+
+	LHS.FrameNumber = LHS.FrameNumber + RHS.FrameNumber + FFrameNumber(FMath::FloorToInt(NewSubFrame));
+	LHS.SubFrame    = FMath::Frac(NewSubFrame);
+
+	return LHS;
+}
+
+
+FORCEINLINE_DEBUGGABLE FFrameTime operator+(FFrameTime A, FFrameTime B)
+{
+	const float        NewSubFrame    = A.SubFrame + B.SubFrame;
+	const FFrameNumber NewFrameNumber = A.FrameNumber + B.FrameNumber + FFrameNumber(FMath::FloorToInt(NewSubFrame));
+
+	return FFrameTime(NewFrameNumber, FMath::Frac(NewSubFrame));
+}
+
+
+FORCEINLINE_DEBUGGABLE FFrameTime operator-(FFrameTime A, FFrameTime B)
+{
+	// Ensure SubFrame is always between 0 and 1
+	// Note that the difference between frame -1.5 and 1.5 is 2, not 3, since sub frame positions are always positive
+	const float        NewSubFrame     = A.SubFrame - B.SubFrame;
+	const int          FlooredSubFrame = FMath::FloorToInt(NewSubFrame);
+	const FFrameNumber NewFrameNumber  = A.FrameNumber - B.FrameNumber + FFrameNumber(FlooredSubFrame);
+
+	return FFrameTime(NewFrameNumber, NewSubFrame - FlooredSubFrame);
+}
+
+
+FORCEINLINE_DEBUGGABLE FFrameTime operator-(FFrameTime A)
+{
+	return A.GetSubFrame() == 0.f
+		? FFrameTime(-A.FrameNumber)
+		: FFrameTime(-A.FrameNumber - 1, 1.f-A.GetSubFrame());
+}
+
+
+FORCEINLINE FFrameTime operator*(FFrameTime A, float Scalar)
+{
+	return FFrameTime::FromDecimal(A.AsDecimal() * Scalar);
+}
+
+
+FORCEINLINE FFrameTime operator/(FFrameTime A, float Scalar)
+{
+	return FFrameTime::FromDecimal(A.AsDecimal() / Scalar);
+}
+
+
+FORCEINLINE_DEBUGGABLE FFrameNumber FFrameTime::FloorToFrame() const
+{
+	return FrameNumber;
+}
+
+
+FORCEINLINE_DEBUGGABLE FFrameNumber FFrameTime::CeilToFrame() const
+{
+	return SubFrame == 0.f ? FrameNumber : FrameNumber+1;
+}
+
+
+FORCEINLINE_DEBUGGABLE FFrameNumber FFrameTime::RoundToFrame() const
+{
+	return SubFrame < .5f ? FrameNumber : FrameNumber+1;
+}
+
+
+FORCEINLINE_DEBUGGABLE double FFrameTime::AsDecimal() const
+{
+	return double(FrameNumber.Value) + SubFrame;
+}
+
+FORCEINLINE_DEBUGGABLE FFrameTime FFrameTime::FromDecimal(double InDecimalFrame)
+{
+	int32 NewFrame = static_cast<int32>(FMath::FloorToDouble(InDecimalFrame));
+
+	// Ensure fractional parts above the highest sub frame float precision do not round to 0.0
+	double Fraction = InDecimalFrame - FMath::FloorToDouble(InDecimalFrame);
+	return FFrameTime(NewFrame, FMath::Clamp(Fraction, 0.0, 0.999999940));
+}
\ No newline at end of file
diff --git a/Engine/Source/Runtime/TimeManagement/Public/TimeManagementBlueprintLibrary.h b/Engine/Source/Runtime/TimeManagement/Public/TimeManagementBlueprintLibrary.h
new file mode 100644
--- /dev/null
+++ b/Engine/Source/Runtime/TimeManagement/Public/TimeManagementBlueprintLibrary.h
@@ -0,0 +1,30 @@
+// Copyright 1998-2018 Epic Games, Inc. All Rights Reserved.
+
+#pragma once
+
+#include "CoreMinimal.h"
+#include "Kismet/BlueprintFunctionLibrary.h"
+#include "FrameRate.h"
+#include "QualifiedFrameTime.h"
+#include "TimeManagementBlueprintLibrary.generated.h"
+
+/**
+ * 
+ */
+UCLASS(meta = (BlueprintThreadSafe))
+class TIMEMANAGEMENT_API UTimeManagementBlueprintLibrary : public UBlueprintFunctionLibrary
+{
+	GENERATED_BODY()
+public:
+	/** Converts an FrameRate to a float ie: 1/30 returns 0.0333333 */
+	UFUNCTION(BlueprintPure, meta = (DisplayName = "FrameRate to Seconds", BlueprintAutocast), Category = "Utilities|Time Management")
+	static float Conv_FrameRateToSeconds(const FFrameRate& InFrameRate);
+
+	/** Converts an QualifiedFrameTime to seconds. */
+	UFUNCTION(BlueprintPure, meta = (DisplayName = "QualifiedFrameTime to Seconds", BlueprintAutocast), Category = "Utilities|Time Management")
+	static float Conv_QualifiedFrameTimeToSeconds(const FQualifiedFrameTime& InFrameTime);
+
+	/** Multiplies a value in seconds against a FrameRate to get a new FrameTime. */
+	UFUNCTION(BlueprintPure, meta = (DisplayName = "Seconds * FrameRate", CompactNodeTitle = "*"), Category = "Utilities|Time Management")
+	static FFrameTime Multiply_SecondsFrameRate(float TimeInSeconds, const FFrameRate& FrameRate);
+};
\ No newline at end of file
diff --git a/Engine/Source/Runtime/UMG/Private/Animation/MovieScene2DTransformSection.cpp b/Engine/Source/Runtime/UMG/Private/Animation/MovieScene2DTransformSection.cpp
--- a/Engine/Source/Runtime/UMG/Private/Animation/MovieScene2DTransformSection.cpp
+++ b/Engine/Source/Runtime/UMG/Private/Animation/MovieScene2DTransformSection.cpp
@@ -2,302 +2,125 @@
 
 #include "Animation/MovieScene2DTransformSection.h"
 #include "UObject/SequencerObjectVersion.h"
+#include "Channels/MovieSceneChannelProxy.h"
+#include "Slate/WidgetTransform.h"
 
-UMovieScene2DTransformSection::UMovieScene2DTransformSection(const FObjectInitializer& ObjectInitializer)
-	: Super(ObjectInitializer)
-{
-	EvalOptions.EnableAndSetCompletionMode
-		(GetLinkerCustomVersion(FSequencerObjectVersion::GUID) < FSequencerObjectVersion::WhenFinishedDefaultsToRestoreState ? 
-			EMovieSceneCompletionMode::KeepState : 
-			GetLinkerCustomVersion(FSequencerObjectVersion::GUID) < FSequencerObjectVersion::WhenFinishedDefaultsToProjectDefault ? 
-			EMovieSceneCompletionMode::RestoreState : 
-			EMovieSceneCompletionMode::ProjectDefault);
-	BlendType = EMovieSceneBlendType::Absolute;
-}
 
+#if WITH_EDITOR
 
-void UMovieScene2DTransformSection::MoveSection( float DeltaTime, TSet<FKeyHandle>& KeyHandles )
+struct F2DTransformSectionEditorData
 {
-	Super::MoveSection( DeltaTime, KeyHandles );
-
-	Rotation.ShiftCurve( DeltaTime, KeyHandles );
-
-	// Move all the curves in this section
-	for( int32 Axis = 0; Axis < 2; ++Axis )
+	F2DTransformSectionEditorData()
 	{
-		Translation[Axis].ShiftCurve( DeltaTime, KeyHandles );
-		Scale[Axis].ShiftCurve( DeltaTime, KeyHandles );
-		Shear[Axis].ShiftCurve( DeltaTime, KeyHandles );
-	}
-}
+		FText TranslationGroup = NSLOCTEXT("MovieScene2DTransformSection", "Translation", "Translation");
+		FText RotationGroup = NSLOCTEXT("MovieScene2DTransformSection", "Rotation", "Rotation");
+		FText ScaleGroup = NSLOCTEXT("MovieScene2DTransformSection", "Scale", "Scale");
+		FText ShearGroup = NSLOCTEXT("MovieScene2DTransformSection", "Shear", "Shear");
 
+		CommonData[0].SetIdentifiers("Translation.X", FCommonChannelData::ChannelX, TranslationGroup);
+		CommonData[0].SortOrder = 0;
 
-void UMovieScene2DTransformSection::DilateSection( float DilationFactor, float Origin, TSet<FKeyHandle>& KeyHandles )
-{
-	Super::DilateSection(DilationFactor, Origin, KeyHandles);
+		CommonData[1].SetIdentifiers("Translation.Y", FCommonChannelData::ChannelY, TranslationGroup);
+		CommonData[1].SortOrder = 1;
 
-	Rotation.ScaleCurve( Origin, DilationFactor, KeyHandles);
+		CommonData[2].SetIdentifiers("Angle", NSLOCTEXT("MovieScene2DTransformSection", "AngleText", "Angle"), RotationGroup);
+		CommonData[2].SortOrder = 2;
 
-	for( int32 Axis = 0; Axis < 2; ++Axis )
-	{
-		Translation[Axis].ScaleCurve( Origin, DilationFactor, KeyHandles );
-		Scale[Axis].ScaleCurve( Origin, DilationFactor, KeyHandles );
-		Shear[Axis].ScaleCurve( Origin, DilationFactor, KeyHandles );
-	}
-}
+		CommonData[3].SetIdentifiers("Scale.X", FCommonChannelData::ChannelX, ScaleGroup);
+		CommonData[3].SortOrder = 3;
+		
+		CommonData[4].SetIdentifiers("Scale.Y", FCommonChannelData::ChannelY, ScaleGroup);
+		CommonData[4].SortOrder = 4;
 
+		CommonData[5].SetIdentifiers("Shear.X", FCommonChannelData::ChannelX, ShearGroup);
+		CommonData[5].SortOrder = 5;
+		
+		CommonData[6].SetIdentifiers("Shear.Y", FCommonChannelData::ChannelY, ShearGroup);
+		CommonData[6].SortOrder = 6;
 
-void UMovieScene2DTransformSection::GetKeyHandles(TSet<FKeyHandle>& OutKeyHandles, TRange<float> TimeRange) const
-{
-	if (!TimeRange.Overlaps(GetRange()))
-	{
-		return;
+		ExternalValues[0].OnGetExternalValue = ExtractTranslationX;
+		ExternalValues[1].OnGetExternalValue = ExtractTranslationY;
+		ExternalValues[2].OnGetExternalValue = ExtractRotation;
+		ExternalValues[3].OnGetExternalValue = ExtractScaleX;
+		ExternalValues[4].OnGetExternalValue = ExtractScaleY;
+		ExternalValues[5].OnGetExternalValue = ExtractShearX;
+		ExternalValues[6].OnGetExternalValue = ExtractShearY;
 	}
 
-	for (auto It(Rotation.GetKeyHandleIterator()); It; ++It)
+	static TOptional<float> ExtractTranslationX(UObject& InObject, FTrackInstancePropertyBindings* Bindings)
 	{
-		float Time = Rotation.GetKeyTime(It.Key());
-		if (TimeRange.Contains(Time))
-		{
-			OutKeyHandles.Add(It.Key());
-		}
+		return Bindings ? Bindings->GetCurrentValue<FWidgetTransform>(InObject).Translation.X : TOptional<float>();
 	}
-
-	for (int32 Axis = 0; Axis < 2; ++Axis)
+	static TOptional<float> ExtractTranslationY(UObject& InObject, FTrackInstancePropertyBindings* Bindings)
 	{
-		for (auto It(Translation[Axis].GetKeyHandleIterator()); It; ++It)
-		{
-			float Time = Translation[Axis].GetKeyTime(It.Key());
-			if (TimeRange.Contains(Time))
-			{
-				OutKeyHandles.Add(It.Key());
-			}
-		}
-		for (auto It(Scale[Axis].GetKeyHandleIterator()); It; ++It)
-		{
-			float Time = Scale[Axis].GetKeyTime(It.Key());
-			if (TimeRange.Contains(Time))
-			{
-				OutKeyHandles.Add(It.Key());
-			}
-		}
-		for (auto It(Shear[Axis].GetKeyHandleIterator()); It; ++It)
-		{
-			float Time = Shear[Axis].GetKeyTime(It.Key());
-			if (TimeRange.Contains(Time))
-			{
-				OutKeyHandles.Add(It.Key());
-			}
-		}
+		return Bindings ? Bindings->GetCurrentValue<FWidgetTransform>(InObject).Translation.Y : TOptional<float>();
 	}
-}
 
-TOptional<float> UMovieScene2DTransformSection::GetKeyTime( FKeyHandle KeyHandle ) const
-{
-	if ( Rotation.IsKeyHandleValid( KeyHandle ) )
-	{
-		return TOptional<float>( Rotation.GetKeyTime( KeyHandle ) );
-	}
-	if ( Translation[0].IsKeyHandleValid( KeyHandle ) )
+	static TOptional<float> ExtractRotation(UObject& InObject, FTrackInstancePropertyBindings* Bindings)
 	{
-		return TOptional<float>( Translation[0].GetKeyTime( KeyHandle ) );
+		return Bindings ? Bindings->GetCurrentValue<FWidgetTransform>(InObject).Angle : TOptional<float>();
 	}
-	if ( Translation[1].IsKeyHandleValid( KeyHandle ) )
-	{
-		return TOptional<float>( Translation[1].GetKeyTime( KeyHandle ) );
-	}
-	if ( Scale[0].IsKeyHandleValid( KeyHandle ) )
-	{
-		return TOptional<float>( Scale[0].GetKeyTime( KeyHandle ) );
-	}
-	if ( Scale[1].IsKeyHandleValid( KeyHandle ) )
-	{
-		return TOptional<float>( Scale[1].GetKeyTime( KeyHandle ) );
-	}
-	return TOptional<float>();
-}
 
-void UMovieScene2DTransformSection::SetKeyTime( FKeyHandle KeyHandle, float Time )
-{
-	if( Rotation.IsKeyHandleValid( KeyHandle ) )
-	{
-		Rotation.SetKeyTime( KeyHandle, Time );
-	}
-	else if( Translation[0].IsKeyHandleValid( KeyHandle ) )
+	static TOptional<float> ExtractScaleX(UObject& InObject, FTrackInstancePropertyBindings* Bindings)
 	{
-		Translation[0].SetKeyTime( KeyHandle, Time );
+		return Bindings ? Bindings->GetCurrentValue<FWidgetTransform>(InObject).Scale.X : TOptional<float>();
 	}
-	else if( Translation[1].IsKeyHandleValid( KeyHandle ) )
+	static TOptional<float> ExtractScaleY(UObject& InObject, FTrackInstancePropertyBindings* Bindings)
 	{
-		Translation[1].SetKeyTime( KeyHandle, Time );
+		return Bindings ? Bindings->GetCurrentValue<FWidgetTransform>(InObject).Scale.Y : TOptional<float>();
 	}
-	else if( Scale[0].IsKeyHandleValid( KeyHandle ) )
-	{
-		Scale[0].SetKeyTime( KeyHandle, Time );
-	}
-	else if( Scale[1].IsKeyHandleValid( KeyHandle ) )
-	{
-		Scale[1].SetKeyTime( KeyHandle, Time );
-	}
-}
 
-/**
- * Chooses an appropriate curve from an axis and a set of curves
- */
-template<typename T>
-static T* ChooseCurve( EAxis::Type Axis, T* Curves )
-{
-	switch(Axis)
+	static TOptional<float> ExtractShearX(UObject& InObject, FTrackInstancePropertyBindings* Bindings)
 	{
-	case EAxis::X:
-		return &Curves[0];
-
-	case EAxis::Y:
-		return &Curves[1];
-
-	default:
-		check(false);
-		return nullptr;
+		return Bindings ? Bindings->GetCurrentValue<FWidgetTransform>(InObject).Shear.X : TOptional<float>();
 	}
-}
-
-
-FRichCurve& UMovieScene2DTransformSection::GetTranslationCurve(EAxis::Type Axis)
-{
-	return *ChooseCurve( Axis, Translation );
-}
-
-
-const FRichCurve& UMovieScene2DTransformSection::GetTranslationCurve(EAxis::Type Axis) const
-{
-	return *ChooseCurve( Axis, Translation );
-}
-
-
-FRichCurve& UMovieScene2DTransformSection::GetRotationCurve()
-{
-	return Rotation;
-}
-
-
-const FRichCurve& UMovieScene2DTransformSection::GetRotationCurve() const
-{
-	return Rotation;
-}
-
-
-FRichCurve& UMovieScene2DTransformSection::GetScaleCurve(EAxis::Type Axis)
-{
-	return *ChooseCurve(Axis, Scale);
-}
-
-
-const FRichCurve& UMovieScene2DTransformSection::GetScaleCurve(EAxis::Type Axis) const
-{
-	return *ChooseCurve(Axis, Scale);
-}
-
-
-FRichCurve& UMovieScene2DTransformSection::GetShearCurve(EAxis::Type Axis)
-{
-	return *ChooseCurve(Axis, Shear);
-}
-
-
-const FRichCurve& UMovieScene2DTransformSection::GetShearCurve(EAxis::Type Axis) const
-{
-	return *ChooseCurve(Axis, Shear);
-}
-
-
-FWidgetTransform UMovieScene2DTransformSection::Eval( float Position, const FWidgetTransform& DefaultValue ) const
-{
-	return FWidgetTransform(	FVector2D( Translation[0].Eval( Position, DefaultValue.Translation.X ), Translation[1].Eval( Position, DefaultValue.Translation.Y ) ),
-								FVector2D( Scale[0].Eval( Position, DefaultValue.Scale.X ),	Scale[1].Eval( Position, DefaultValue.Scale.Y ) ),
-								FVector2D( Shear[0].Eval( Position, DefaultValue.Shear.X ),	Shear[1].Eval( Position, DefaultValue.Shear.Y ) ),
-								Rotation.Eval( Position, DefaultValue.Angle ) );
-}
-
-// This function uses a template to avoid duplicating this for const and non-const versions
-template<typename CurveType>
-CurveType* GetCurveForChannelAndAxis( EKey2DTransformChannel Channel, EKey2DTransformAxis Axis, 
-	CurveType* TranslationCurves, CurveType* ScaleCurves, CurveType* ShearCurves, CurveType* RotationCurve )
-{
-	switch ( Channel )
+	static TOptional<float> ExtractShearY(UObject& InObject, FTrackInstancePropertyBindings* Bindings)
 	{
-	case EKey2DTransformChannel::Translation:
-		if ( Axis == EKey2DTransformAxis::X )
-		{
-			return &TranslationCurves[0];
-		}
-		else if ( Axis == EKey2DTransformAxis::Y )
-		{
-			return &TranslationCurves[1];
-		}
-	case EKey2DTransformChannel::Scale:
-		if ( Axis == EKey2DTransformAxis::X )
-		{
-			return &ScaleCurves[0];
-		}
-		else if ( Axis == EKey2DTransformAxis::Y )
-		{
-			return &ScaleCurves[1];
-		}
-	case EKey2DTransformChannel::Shear:
-		if ( Axis == EKey2DTransformAxis::X )
-		{
-			return &ShearCurves[0];
-		}
-		else if ( Axis == EKey2DTransformAxis::Y )
-		{
-			return &ShearCurves[1];
-		}
-	case EKey2DTransformChannel::Rotation:
-		if ( Axis == EKey2DTransformAxis::None )
-		{
-			return RotationCurve;
-		}
+		return Bindings ? Bindings->GetCurrentValue<FWidgetTransform>(InObject).Shear.Y : TOptional<float>();
 	}
-	checkf( false, TEXT( "Invalid channel and axis combination." ) );
-	return nullptr;
-}
 
+	FMovieSceneChannelEditorData    CommonData[7];
+	TMovieSceneExternalValue<float> ExternalValues[7];
+};
 
-bool UMovieScene2DTransformSection::NewKeyIsNewData( float Time, const struct F2DTransformKey& TransformKey ) const
+#endif	// WITH_EDITOR
+
+UMovieScene2DTransformSection::UMovieScene2DTransformSection(const FObjectInitializer& ObjectInitializer)
+	: Super(ObjectInitializer)
 {
-	const FRichCurve* KeyCurve = GetCurveForChannelAndAxis( TransformKey.Channel, TransformKey.Axis, Translation, Scale, Shear, &Rotation );
-	return FMath::IsNearlyEqual( KeyCurve->Eval( Time ), TransformKey.Value ) == false;
-}
+	EvalOptions.EnableAndSetCompletionMode
+		(GetLinkerCustomVersion(FSequencerObjectVersion::GUID) < FSequencerObjectVersion::WhenFinishedDefaultsToRestoreState ? 
+			EMovieSceneCompletionMode::KeepState : 
+			GetLinkerCustomVersion(FSequencerObjectVersion::GUID) < FSequencerObjectVersion::WhenFinishedDefaultsToProjectDefault ? 
+			EMovieSceneCompletionMode::RestoreState : 
+			EMovieSceneCompletionMode::ProjectDefault);
+	BlendType = EMovieSceneBlendType::Absolute;
 
+	FMovieSceneChannelData Channels;
 
-bool UMovieScene2DTransformSection::HasKeys( const struct F2DTransformKey& TransformKey ) const
-{
-	const FRichCurve* KeyCurve = GetCurveForChannelAndAxis( TransformKey.Channel, TransformKey.Axis, Translation, Scale, Shear, &Rotation );
-	return KeyCurve->GetNumKeys() != 0;
-}
+#if WITH_EDITOR
 
+	static const F2DTransformSectionEditorData EditorData;
 
-void UMovieScene2DTransformSection::AddKey( float Time, const struct F2DTransformKey& TransformKey, EMovieSceneKeyInterpolation KeyInterpolation )
-{
-	FRichCurve* KeyCurve = GetCurveForChannelAndAxis( TransformKey.Channel, TransformKey.Axis, Translation, Scale, Shear, &Rotation );
-	AddKeyToCurve( *KeyCurve, Time, TransformKey.Value, KeyInterpolation );
-}
+	Channels.Add(Translation[0], EditorData.CommonData[0], EditorData.ExternalValues[0]);
+	Channels.Add(Translation[1], EditorData.CommonData[1], EditorData.ExternalValues[1]);
+	Channels.Add(Rotation,       EditorData.CommonData[2], EditorData.ExternalValues[2]);
+	Channels.Add(Scale[0],       EditorData.CommonData[3], EditorData.ExternalValues[3]);
+	Channels.Add(Scale[1],       EditorData.CommonData[4], EditorData.ExternalValues[4]);
+	Channels.Add(Shear[0],       EditorData.CommonData[5], EditorData.ExternalValues[5]);
+	Channels.Add(Shear[1],       EditorData.CommonData[6], EditorData.ExternalValues[6]);
 
+#else
 
-void UMovieScene2DTransformSection::SetDefault( const struct F2DTransformKey& TransformKey )
-{
-	FRichCurve* KeyCurve = GetCurveForChannelAndAxis( TransformKey.Channel, TransformKey.Axis, Translation, Scale, Shear, &Rotation );
-	SetCurveDefault( *KeyCurve, TransformKey.Value );
-}
+	Channels.Add(Translation[0]);
+	Channels.Add(Translation[1]);
+	Channels.Add(Rotation);
+	Channels.Add(Scale[0]);
+	Channels.Add(Scale[1]);
+	Channels.Add(Shear[0]);
+	Channels.Add(Shear[1]);
 
+#endif
 
-void UMovieScene2DTransformSection::ClearDefaults()
-{
-	Rotation.ClearDefaultValue();
-	Translation[0].ClearDefaultValue();
-	Translation[1].ClearDefaultValue();
-	Scale[0].ClearDefaultValue();
-	Scale[1].ClearDefaultValue();
-	Shear[0].ClearDefaultValue();
-	Shear[1].ClearDefaultValue();
+	ChannelProxy = MakeShared<FMovieSceneChannelProxy>(MoveTemp(Channels));
 }
diff --git a/Engine/Source/Runtime/UMG/Private/Animation/MovieScene2DTransformTemplate.cpp b/Engine/Source/Runtime/UMG/Private/Animation/MovieScene2DTransformTemplate.cpp
--- a/Engine/Source/Runtime/UMG/Private/Animation/MovieScene2DTransformTemplate.cpp
+++ b/Engine/Source/Runtime/UMG/Private/Animation/MovieScene2DTransformTemplate.cpp
@@ -16,29 +16,30 @@ FMovieScene2DTransformSectionTemplate::FMovieScene2DTransformSectionTemplate(con
 	: FMovieScenePropertySectionTemplate(Track.GetPropertyName(), Track.GetPropertyPath())
 	, BlendType(Section.GetBlendType().Get())
 {
-	Translation[0] = Section.GetTranslationCurve(EAxis::X);
-	Translation[1] = Section.GetTranslationCurve(EAxis::Y);
+	Translation[0] = Section.Translation[0];
+	Translation[1] = Section.Translation[1];
 
-	Rotation = Section.GetRotationCurve();
+	Rotation = Section.Rotation;
 
-	Scale[0] = Section.GetScaleCurve(EAxis::X);
-	Scale[1] = Section.GetScaleCurve(EAxis::Y);
+	Scale[0] = Section.Scale[0];
+	Scale[1] = Section.Scale[1];
 
-	Shear[0] = Section.GetShearCurve(EAxis::X);
-	Shear[1] = Section.GetShearCurve(EAxis::Y);
+	Shear[0] = Section.Shear[0];
+	Shear[1] = Section.Shear[1];
 }
 
 void FMovieScene2DTransformSectionTemplate::Evaluate(const FMovieSceneEvaluationOperand& Operand, const FMovieSceneContext& Context, const FPersistentEvaluationData& PersistentData, FMovieSceneExecutionTokens& ExecutionTokens) const
 {
-	const float Time = Context.GetTime();
+	const FFrameTime Time = Context.GetTime();
 	MovieScene::TMultiChannelValue<float, 7> AnimatedData;
 
 	// Only activate channels if the curve has data associated with it
-	auto EvalChannel = [&AnimatedData, Time](uint8 ChanneIndex, const FRichCurve& Curve)
+	auto EvalChannel = [&AnimatedData, Time](uint8 ChanneIndex, const FMovieSceneFloatChannel& Channel)
 	{
-		if (Curve.HasAnyData())
+		float Value = 0.f;
+		if (Channel.Evaluate(Time, Value))
 		{
-			AnimatedData.Set(ChanneIndex, Curve.Eval(Time));
+			AnimatedData.Set(ChanneIndex, Value);
 		}
 	};
 
diff --git a/Engine/Source/Runtime/UMG/Private/Animation/MovieScene2DTransformTemplate.h b/Engine/Source/Runtime/UMG/Private/Animation/MovieScene2DTransformTemplate.h
--- a/Engine/Source/Runtime/UMG/Private/Animation/MovieScene2DTransformTemplate.h
+++ b/Engine/Source/Runtime/UMG/Private/Animation/MovieScene2DTransformTemplate.h
@@ -4,7 +4,7 @@
 
 #include "CoreMinimal.h"
 #include "UObject/ObjectMacros.h"
-#include "Curves/RichCurve.h"
+#include "Channels/MovieSceneFloatChannel.h"
 #include "Evaluation/MovieSceneEvalTemplate.h"
 #include "Evaluation/MovieScenePropertyTemplate.h"
 #include "Evaluation/Blending/MovieSceneMultiChannelBlending.h"
@@ -32,19 +32,19 @@ private:
 
 	/** Translation curves */
 	UPROPERTY()
-	FRichCurve Translation[2];
+	FMovieSceneFloatChannel Translation[2];
 	
 	/** Rotation curve */
 	UPROPERTY()
-	FRichCurve Rotation;
+	FMovieSceneFloatChannel Rotation;
 
 	/** Scale curves */
 	UPROPERTY()
-	FRichCurve Scale[2];
+	FMovieSceneFloatChannel Scale[2];
 
 	/** Shear curve */
 	UPROPERTY()
-	FRichCurve Shear[2];
+	FMovieSceneFloatChannel Shear[2];
 
 	/** Blending method */
 	UPROPERTY()
diff --git a/Engine/Source/Runtime/UMG/Private/Animation/MovieScene2DTransformTrack.cpp b/Engine/Source/Runtime/UMG/Private/Animation/MovieScene2DTransformTrack.cpp
--- a/Engine/Source/Runtime/UMG/Private/Animation/MovieScene2DTransformTrack.cpp
+++ b/Engine/Source/Runtime/UMG/Private/Animation/MovieScene2DTransformTrack.cpp
@@ -25,25 +25,4 @@ UMovieSceneSection* UMovieScene2DTransformTrack::CreateNewSection()
 FMovieSceneEvalTemplatePtr UMovieScene2DTransformTrack::CreateTemplateForSection(const UMovieSceneSection& InSection) const
 {
 	return FMovieScene2DTransformSectionTemplate(*CastChecked<const UMovieScene2DTransformSection>(&InSection), *this);
-}
-
-PRAGMA_DISABLE_DEPRECATION_WARNINGS
-bool UMovieScene2DTransformTrack::Eval(float Position, float LastPosition, FWidgetTransform& InOutTransform) const
-{
-	const UMovieSceneSection* Section = MovieSceneHelpers::FindNearestSectionAtTime(Sections, Position);
-
-	if(Section)
-	{
-		const UMovieScene2DTransformSection* TransformSection = CastChecked<UMovieScene2DTransformSection>(Section);
-
-		if (!Section->IsInfinite())
-		{
-			Position = FMath::Clamp(Position, Section->GetStartTime(), Section->GetEndTime());
-		}
-
-		InOutTransform = TransformSection->Eval(Position, InOutTransform);
-	}
-
-	return (Section != nullptr);
-}
-PRAGMA_ENABLE_DEPRECATION_WARNINGS
+}
\ No newline at end of file
diff --git a/Engine/Source/Runtime/UMG/Private/Animation/MovieSceneMarginSection.cpp b/Engine/Source/Runtime/UMG/Private/Animation/MovieSceneMarginSection.cpp
--- a/Engine/Source/Runtime/UMG/Private/Animation/MovieSceneMarginSection.cpp
+++ b/Engine/Source/Runtime/UMG/Private/Animation/MovieSceneMarginSection.cpp
@@ -1,167 +1,78 @@
 // Copyright 1998-2018 Epic Games, Inc. All Rights Reserved.
 
 #include "Animation/MovieSceneMarginSection.h"
+#include "Channels/MovieSceneChannelProxy.h"
 
-UMovieSceneMarginSection::UMovieSceneMarginSection( const FObjectInitializer& ObjectInitializer )
-	: Super( ObjectInitializer )
-{
-	BlendType = EMovieSceneBlendType::Absolute;
-}
+#if WITH_EDITOR
 
-void UMovieSceneMarginSection::MoveSection( float DeltaTime, TSet<FKeyHandle>& KeyHandles )
+struct FMarginSectionEditorData
 {
-	Super::MoveSection( DeltaTime, KeyHandles );
-
-	// Move all the curves in this section
-	LeftCurve.ShiftCurve(DeltaTime, KeyHandles);
-	TopCurve.ShiftCurve(DeltaTime, KeyHandles);
-	RightCurve.ShiftCurve(DeltaTime, KeyHandles);
-	BottomCurve.ShiftCurve(DeltaTime, KeyHandles);
-}
+	FMarginSectionEditorData()
+	{
+		CommonData[0].SetIdentifiers("Left", NSLOCTEXT("MovieSceneMarginSection", "LeftText", "Left"));
+		CommonData[0].SortOrder = 0;
 
-void UMovieSceneMarginSection::DilateSection( float DilationFactor, float Origin, TSet<FKeyHandle>& KeyHandles )
-{
-	Super::DilateSection(DilationFactor, Origin, KeyHandles);
+		CommonData[1].SetIdentifiers("Top", NSLOCTEXT("MovieSceneMarginSection", "TopText", "Top"));
+		CommonData[1].SortOrder = 1;
 
-	LeftCurve.ScaleCurve(Origin, DilationFactor, KeyHandles);
-	TopCurve.ScaleCurve(Origin, DilationFactor, KeyHandles);
-	RightCurve.ScaleCurve(Origin, DilationFactor, KeyHandles);
-	BottomCurve.ScaleCurve(Origin, DilationFactor, KeyHandles);
-}
+		CommonData[2].SetIdentifiers("Right", NSLOCTEXT("MovieSceneMarginSection", "RightText", "Right"));
+		CommonData[2].SortOrder = 2;
 
-void UMovieSceneMarginSection::GetKeyHandles(TSet<FKeyHandle>& OutKeyHandles, TRange<float> TimeRange) const
-{
-	if (!TimeRange.Overlaps(GetRange()))
-	{
-		return;
-	}
+		CommonData[3].SetIdentifiers("Bottom", NSLOCTEXT("MovieSceneMarginSection", "BottomText", "Bottom"));
+		CommonData[3].SortOrder = 3;
 
-	for (auto It(LeftCurve.GetKeyHandleIterator()); It; ++It)
-	{
-		float Time = LeftCurve.GetKeyTime(It.Key());
-		if (TimeRange.Contains(Time))
-		{
-			OutKeyHandles.Add(It.Key());
-		}
-	}
-	for (auto It(TopCurve.GetKeyHandleIterator()); It; ++It)
-	{
-		float Time = TopCurve.GetKeyTime(It.Key());
-		if (TimeRange.Contains(Time))
-		{
-			OutKeyHandles.Add(It.Key());
-		}
-	}
-	for (auto It(RightCurve.GetKeyHandleIterator()); It; ++It)
-	{
-		float Time = RightCurve.GetKeyTime(It.Key());
-		if (TimeRange.Contains(Time))
-		{
-			OutKeyHandles.Add(It.Key());
-		}
+		ExternalValues[0].OnGetExternalValue = ExtractLeftChannel;
+		ExternalValues[1].OnGetExternalValue = ExtractTopChannel;
+		ExternalValues[2].OnGetExternalValue = ExtractRightChannel;
+		ExternalValues[3].OnGetExternalValue = ExtractBottomChannel;
 	}
-	for (auto It(BottomCurve.GetKeyHandleIterator()); It; ++It)
-	{
-		float Time = BottomCurve.GetKeyTime(It.Key());
-		if (TimeRange.Contains(Time))
-		{
-			OutKeyHandles.Add(It.Key());
-		}
-	}
-}
-
 
-TOptional<float> UMovieSceneMarginSection::GetKeyTime( FKeyHandle KeyHandle ) const
-{
-	if ( LeftCurve.IsKeyHandleValid( KeyHandle ) )
+	static TOptional<float> ExtractLeftChannel(UObject& InObject, FTrackInstancePropertyBindings* Bindings)
 	{
-		return TOptional<float>( LeftCurve.GetKeyTime( KeyHandle ) );
+		return Bindings ? Bindings->GetCurrentValue<FMargin>(InObject).Left : TOptional<float>();
 	}
-	if ( TopCurve.IsKeyHandleValid( KeyHandle ) )
+	static TOptional<float> ExtractTopChannel(UObject& InObject, FTrackInstancePropertyBindings* Bindings)
 	{
-		return TOptional<float>( TopCurve.GetKeyTime( KeyHandle ) );
+		return Bindings ? Bindings->GetCurrentValue<FMargin>(InObject).Top : TOptional<float>();
 	}
-	if ( RightCurve.IsKeyHandleValid( KeyHandle ) )
+	static TOptional<float> ExtractRightChannel(UObject& InObject, FTrackInstancePropertyBindings* Bindings)
 	{
-		return TOptional<float>( RightCurve.GetKeyTime( KeyHandle ) );
+		return Bindings ? Bindings->GetCurrentValue<FMargin>(InObject).Right : TOptional<float>();
 	}
-	if ( BottomCurve.IsKeyHandleValid( KeyHandle ) )
+	static TOptional<float> ExtractBottomChannel(UObject& InObject, FTrackInstancePropertyBindings* Bindings)
 	{
-		return TOptional<float>( BottomCurve.GetKeyTime( KeyHandle ) );
+		return Bindings ? Bindings->GetCurrentValue<FMargin>(InObject).Bottom : TOptional<float>();
 	}
-	return TOptional<float>();
-}
-
 
-void UMovieSceneMarginSection::SetKeyTime( FKeyHandle KeyHandle, float Time )
-{
-	if ( LeftCurve.IsKeyHandleValid( KeyHandle ) )
-	{
-		LeftCurve.SetKeyTime( KeyHandle, Time );
-	}
-	else if ( TopCurve.IsKeyHandleValid( KeyHandle ) )
-	{
-		TopCurve.SetKeyTime( KeyHandle, Time );
-	}
-	else if ( RightCurve.IsKeyHandleValid( KeyHandle ) )
-	{
-		RightCurve.SetKeyTime( KeyHandle, Time );
-	}
-	else if ( BottomCurve.IsKeyHandleValid( KeyHandle ) )
-	{
-		BottomCurve.SetKeyTime( KeyHandle, Time );
-	}
-}
+	FMovieSceneChannelEditorData    CommonData[4];
+	TMovieSceneExternalValue<float> ExternalValues[4];
+};
 
+#endif	// WITH_EDITOR
 
-template<typename CurveType>
-CurveType* GetCurveForChannel( EKeyMarginChannel Channel, CurveType* LeftCurve, CurveType* TopCurve, CurveType* RightCurve, CurveType* BottomCurve )
+UMovieSceneMarginSection::UMovieSceneMarginSection( const FObjectInitializer& ObjectInitializer )
+	: Super( ObjectInitializer )
 {
-	switch ( Channel )
-	{
-	case EKeyMarginChannel::Left:
-		return LeftCurve;
-	case EKeyMarginChannel::Top:
-		return TopCurve;
-	case EKeyMarginChannel::Right:
-		return RightCurve;
-	case EKeyMarginChannel::Bottom:
-		return BottomCurve;
-	}
-	checkf(false, TEXT("Invalid curve channel"));
-	return nullptr;
-}
-
+	BlendType = EMovieSceneBlendType::Absolute;
 
-void UMovieSceneMarginSection::AddKey( float Time, const FMarginKey& Key, EMovieSceneKeyInterpolation KeyInterpolation )
-{
-	FRichCurve* KeyCurve = GetCurveForChannel( Key.Channel, &LeftCurve, &TopCurve, &RightCurve, &BottomCurve );
-	AddKeyToCurve( *KeyCurve, Time, Key.Value, KeyInterpolation );
-}
+	FMovieSceneChannelData Channels;
 
+#if WITH_EDITOR
 
-bool UMovieSceneMarginSection::NewKeyIsNewData( float Time, const FMarginKey& Key ) const
-{
-	const FRichCurve* KeyCurve = GetCurveForChannel( Key.Channel, &LeftCurve, &TopCurve, &RightCurve, &BottomCurve );
-	return FMath::IsNearlyEqual( KeyCurve->Eval( Time ), Key.Value ) == false;
-}
+	static const FMarginSectionEditorData EditorData;
+	Channels.Add(LeftCurve,   EditorData.CommonData[0], EditorData.ExternalValues[0]);
+	Channels.Add(TopCurve,    EditorData.CommonData[1], EditorData.ExternalValues[1]);
+	Channels.Add(RightCurve,  EditorData.CommonData[2], EditorData.ExternalValues[2]);
+	Channels.Add(BottomCurve, EditorData.CommonData[3], EditorData.ExternalValues[3]);
 
-bool UMovieSceneMarginSection::HasKeys( const FMarginKey& Key ) const
-{
-	const FRichCurve* KeyCurve = GetCurveForChannel( Key.Channel, &LeftCurve, &TopCurve, &RightCurve, &BottomCurve );
-	return KeyCurve->GetNumKeys() != 0;
-}
+#else
 
-void UMovieSceneMarginSection::SetDefault( const FMarginKey& Key )
-{
-	FRichCurve* KeyCurve = GetCurveForChannel( Key.Channel, &LeftCurve, &TopCurve, &RightCurve, &BottomCurve );
-	SetCurveDefault( *KeyCurve, Key.Value );
-}
+	Channels.Add(LeftCurve);
+	Channels.Add(TopCurve);
+	Channels.Add(RightCurve);
+	Channels.Add(BottomCurve);
 
-void UMovieSceneMarginSection::ClearDefaults()
-{
-	LeftCurve.ClearDefaultValue();
-	TopCurve.ClearDefaultValue();
-	RightCurve.ClearDefaultValue();
-	BottomCurve.ClearDefaultValue();
-}
+#endif
+	
+	ChannelProxy = MakeShared<FMovieSceneChannelProxy>(MoveTemp(Channels));
+}
\ No newline at end of file
diff --git a/Engine/Source/Runtime/UMG/Private/Animation/MovieSceneMarginTemplate.cpp b/Engine/Source/Runtime/UMG/Private/Animation/MovieSceneMarginTemplate.cpp
--- a/Engine/Source/Runtime/UMG/Private/Animation/MovieSceneMarginTemplate.cpp
+++ b/Engine/Source/Runtime/UMG/Private/Animation/MovieSceneMarginTemplate.cpp
@@ -14,37 +14,38 @@ template<> FMovieSceneAnimTypeID GetBlendingDataType<FMargin>()
 
 FMovieSceneMarginSectionTemplate::FMovieSceneMarginSectionTemplate(const UMovieSceneMarginSection& Section, const UMovieScenePropertyTrack& Track)
 	: FMovieScenePropertySectionTemplate(Track.GetPropertyName(), Track.GetPropertyPath())
-	, TopCurve(Section.GetTopCurve())
-	, LeftCurve(Section.GetLeftCurve())
-	, RightCurve(Section.GetRightCurve())
-	, BottomCurve(Section.GetBottomCurve())
+	, TopCurve(Section.TopCurve)
+	, LeftCurve(Section.LeftCurve)
+	, RightCurve(Section.RightCurve)
+	, BottomCurve(Section.BottomCurve)
 	, BlendType(Section.GetBlendType().Get())
 {
 }
 
 void FMovieSceneMarginSectionTemplate::Evaluate(const FMovieSceneEvaluationOperand& Operand, const FMovieSceneContext& Context, const FPersistentEvaluationData& PersistentData, FMovieSceneExecutionTokens& ExecutionTokens) const
 {
-	const float Time = Context.GetTime();
+	const FFrameTime Time = Context.GetTime();
 	MovieScene::TMultiChannelValue<float, 4> AnimatedData;
 
-	if (LeftCurve.HasAnyData())
+	float Value = 0.f;
+	if (LeftCurve.Evaluate(Time, Value))
 	{
-		AnimatedData.Set(0, LeftCurve.Eval(Time));
+		AnimatedData.Set(0, Value);
 	}
 
-	if (TopCurve.HasAnyData())
+	if (TopCurve.Evaluate(Time, Value))
 	{
-		AnimatedData.Set(1, TopCurve.Eval(Time));
+		AnimatedData.Set(1, Value);
 	}
 
-	if (RightCurve.HasAnyData())
+	if (RightCurve.Evaluate(Time, Value))
 	{
-		AnimatedData.Set(2, RightCurve.Eval(Time));
+		AnimatedData.Set(2, Value);
 	}
 
-	if (BottomCurve.HasAnyData())
+	if (BottomCurve.Evaluate(Time, Value))
 	{
-		AnimatedData.Set(3, BottomCurve.Eval(Time));
+		AnimatedData.Set(3, Value);
 	}
 
 	if (!AnimatedData.IsEmpty())
diff --git a/Engine/Source/Runtime/UMG/Private/Animation/MovieSceneMarginTemplate.h b/Engine/Source/Runtime/UMG/Private/Animation/MovieSceneMarginTemplate.h
--- a/Engine/Source/Runtime/UMG/Private/Animation/MovieSceneMarginTemplate.h
+++ b/Engine/Source/Runtime/UMG/Private/Animation/MovieSceneMarginTemplate.h
@@ -4,7 +4,7 @@
 
 #include "CoreMinimal.h"
 #include "UObject/ObjectMacros.h"
-#include "Curves/RichCurve.h"
+#include "Channels/MovieSceneFloatChannel.h"
 #include "Evaluation/MovieSceneEvalTemplate.h"
 #include "Evaluation/MovieScenePropertyTemplate.h"
 #include "Evaluation/Blending/MovieSceneMultiChannelBlending.h"
@@ -24,22 +24,27 @@ struct FMovieSceneMarginSectionTemplate : public FMovieScenePropertySectionTempl
 	FMovieSceneMarginSectionTemplate(){}
 	FMovieSceneMarginSectionTemplate(const UMovieSceneMarginSection& Section, const UMovieScenePropertyTrack& Track);
 
+	const FMovieSceneFloatChannel& GetTopCurve() const    { return TopCurve; }
+	const FMovieSceneFloatChannel& GetLeftCurve() const   { return LeftCurve; }
+	const FMovieSceneFloatChannel& GetRightCurve() const  { return RightCurve; }
+	const FMovieSceneFloatChannel& GetBottomCurve() const { return BottomCurve; }
+
 private:
 
 	virtual UScriptStruct& GetScriptStructImpl() const override { return *StaticStruct(); }
 	virtual void Evaluate(const FMovieSceneEvaluationOperand& Operand, const FMovieSceneContext& Context, const FPersistentEvaluationData& PersistentData, FMovieSceneExecutionTokens& ExecutionTokens) const override;
 
 	UPROPERTY()
-	FRichCurve TopCurve;
+	FMovieSceneFloatChannel TopCurve;
 
 	UPROPERTY()
-	FRichCurve LeftCurve;
+	FMovieSceneFloatChannel LeftCurve;
 
 	UPROPERTY()
-	FRichCurve RightCurve;
+	FMovieSceneFloatChannel RightCurve;
 
 	UPROPERTY()
-	FRichCurve BottomCurve;
+	FMovieSceneFloatChannel BottomCurve;
 
 	UPROPERTY()
 	EMovieSceneBlendType BlendType;
diff --git a/Engine/Source/Runtime/UMG/Private/Animation/UMGSequencePlayer.cpp b/Engine/Source/Runtime/UMG/Private/Animation/UMGSequencePlayer.cpp
--- a/Engine/Source/Runtime/UMG/Private/Animation/UMGSequencePlayer.cpp
+++ b/Engine/Source/Runtime/UMG/Private/Animation/UMGSequencePlayer.cpp
@@ -4,14 +4,14 @@
 #include "MovieScene.h"
 #include "UMGPrivate.h"
 #include "Animation/WidgetAnimation.h"
+#include "MovieSceneTimeHelpers.h"
 
 
 UUMGSequencePlayer::UUMGSequencePlayer(const FObjectInitializer& ObjectInitializer)
 	: Super(ObjectInitializer)
 {
 	PlayerStatus = EMovieScenePlayerStatus::Stopped;
-	TimeCursorPosition = 0.0f;
-	AnimationStartOffset = 0;
+	TimeCursorPosition = FFrameTime(0);
 	PlaybackSpeed = 1;
 	Animation = nullptr;
 	bIsEvaluating = false;
@@ -20,28 +20,29 @@ UUMGSequencePlayer::UUMGSequencePlayer(const FObjectInitializer& ObjectInitializ
 void UUMGSequencePlayer::InitSequencePlayer( UWidgetAnimation& InAnimation, UUserWidget& InUserWidget )
 {
 	Animation = &InAnimation;
+	UserWidget = &InUserWidget;
+
 
 	UMovieScene* MovieScene = Animation->GetMovieScene();
 
 	// Cache the time range of the sequence to determine when we stop
-	TimeRange = MovieScene->GetPlaybackRange();
-	AnimationStartOffset = TimeRange.GetLowerBoundValue();
-
-	UserWidget = &InUserWidget;
+	Duration = MovieScene::DiscreteSize(MovieScene->GetPlaybackRange());
+	AnimationResolution = MovieScene->GetFrameResolution();
+	AbsolutePlaybackStart = MovieScene::DiscreteInclusiveLower(MovieScene->GetPlaybackRange());
 }
 
 void UUMGSequencePlayer::Tick(float DeltaTime)
 {
 	if ( PlayerStatus == EMovieScenePlayerStatus::Playing )
 	{
-		const double AnimationLength = CurrentPlayRange.Size<double>();
+		FFrameTime DeltaFrameTime = (bIsPlayingForward ? DeltaTime * PlaybackSpeed : -DeltaTime * PlaybackSpeed) * AnimationResolution;
 
-		double LastTimePosition = TimeCursorPosition;
-		TimeCursorPosition += bIsPlayingForward ? DeltaTime * PlaybackSpeed : -DeltaTime * PlaybackSpeed;
+		FFrameTime LastTimePosition = TimeCursorPosition;
+		TimeCursorPosition += DeltaFrameTime;
 
 		// Check if we crossed over bounds
-		const bool bCrossedLowerBound = TimeCursorPosition < CurrentPlayRange.GetLowerBoundValue();
-		const bool bCrossedUpperBound = TimeCursorPosition > CurrentPlayRange.GetUpperBoundValue();
+		const bool bCrossedLowerBound = TimeCursorPosition < 0;
+		const bool bCrossedUpperBound = TimeCursorPosition >= FFrameTime(Duration);
 		const bool bCrossedEndTime = bIsPlayingForward
 			? LastTimePosition < EndTime && EndTime <= TimeCursorPosition
 			: LastTimePosition > EndTime && EndTime >= TimeCursorPosition;
@@ -60,38 +61,40 @@ void UUMGSequencePlayer::Tick(float DeltaTime)
 		{
 			if (bCompleted)
 			{
-				TimeCursorPosition = CurrentPlayRange.GetLowerBoundValue();
+				TimeCursorPosition = FFrameTime(0);
 			}
 			else
 			{
 				if (PlayMode == EUMGSequencePlayMode::PingPong)
 				{
 					bIsPlayingForward = !bIsPlayingForward;
-					TimeCursorPosition = FMath::Abs(TimeCursorPosition - CurrentPlayRange.GetLowerBoundValue()) + CurrentPlayRange.GetLowerBoundValue();
+					TimeCursorPosition = FMath::Abs(TimeCursorPosition);
 				}
 				else
 				{
-					TimeCursorPosition += AnimationLength;
+					TimeCursorPosition += FFrameTime(Duration);
 					LastTimePosition = TimeCursorPosition;
 				}
 			}
 		}
 		else if (bCrossedUpperBound)
 		{
+			FFrameTime LastValidFrame(Duration-1, 0.99999994f);
+
 			if (bCompleted)
 			{
-				TimeCursorPosition = CurrentPlayRange.GetUpperBoundValue();
+				TimeCursorPosition = LastValidFrame;
 			}
 			else
 			{
 				if (PlayMode == EUMGSequencePlayMode::PingPong)
 				{
 					bIsPlayingForward = !bIsPlayingForward;
-					TimeCursorPosition = CurrentPlayRange.GetUpperBoundValue() - (TimeCursorPosition - CurrentPlayRange.GetUpperBoundValue());
+					TimeCursorPosition = LastValidFrame - (TimeCursorPosition - FFrameTime(Duration));
 				}
 				else
 				{
-					TimeCursorPosition -= AnimationLength;
+					TimeCursorPosition = TimeCursorPosition - FFrameTime(Duration);
 					LastTimePosition = TimeCursorPosition;
 				}
 			}
@@ -105,11 +108,11 @@ void UUMGSequencePlayer::Tick(float DeltaTime)
 		}
 		if (RootTemplateInstance.IsValid())
 		{
+			UMovieScene* MovieScene = Animation->GetMovieScene();
+
 			bIsEvaluating = true;
 
-			const FMovieSceneContext Context(
-				FMovieSceneEvaluationRange(TimeCursorPosition + AnimationStartOffset, LastTimePosition + AnimationStartOffset),
-				PlayerStatus);
+			const FMovieSceneContext Context(FMovieSceneEvaluationRange(AbsolutePlaybackStart + TimeCursorPosition, AbsolutePlaybackStart + LastTimePosition, AnimationResolution), PlayerStatus);
 			RootTemplateInstance.Evaluate(Context, *this);
 
 			bIsEvaluating = false;
@@ -119,6 +122,11 @@ void UUMGSequencePlayer::Tick(float DeltaTime)
 
 		if ( bCompleted )
 		{
+			if (RootTemplateInstance.IsValid())
+			{
+				RootTemplateInstance.Finish(*this);
+			}
+			
 			PlayerStatus = EMovieScenePlayerStatus::Stopped;
 			OnSequenceFinishedPlayingEvent.Broadcast(*this);
 			Animation->OnAnimationFinished.Broadcast();
@@ -126,29 +134,28 @@ void UUMGSequencePlayer::Tick(float DeltaTime)
 	}
 }
 
-void UUMGSequencePlayer::PlayInternal(double StartAtTime, double EndAtTime, double SubAnimStartTime, double SubAnimEndTime, int32 InNumLoopsToPlay, EUMGSequencePlayMode::Type InPlayMode, float InPlaybackSpeed)
+void UUMGSequencePlayer::PlayInternal(double StartAtTime, double EndAtTime, int32 InNumLoopsToPlay, EUMGSequencePlayMode::Type InPlayMode, float InPlaybackSpeed)
 {
 	RootTemplateInstance.Initialize(*Animation, *this);
 
 	PlaybackSpeed = FMath::Abs(InPlaybackSpeed);
 	PlayMode = InPlayMode;
 
-	// Set the temporary range for this play of the animation
-	CurrentPlayRange = TRange<double>(SubAnimStartTime, TRangeBound<double>::Inclusive(SubAnimEndTime));
+	FFrameTime LastValidFrame(Duration-1, 0.99999994f);
 
 	if (PlayMode == EUMGSequencePlayMode::Reverse)
 	{
 		// When playing in reverse count subtract the start time from the end.
-		TimeCursorPosition = CurrentPlayRange.GetUpperBoundValue() - StartAtTime;
+		TimeCursorPosition = LastValidFrame - StartAtTime * AnimationResolution;
 	}
 	else
 	{
-		TimeCursorPosition = StartAtTime;
+		TimeCursorPosition = StartAtTime * AnimationResolution;
 	}
-	
+
 	// Clamp the start time and end time to be within the bounds
-	TimeCursorPosition = FMath::Clamp(TimeCursorPosition, CurrentPlayRange.GetLowerBoundValue(), CurrentPlayRange.GetUpperBoundValue());
-	EndTime = FMath::Clamp(EndAtTime, CurrentPlayRange.GetLowerBoundValue(), CurrentPlayRange.GetUpperBoundValue());
+	TimeCursorPosition = FMath::Clamp(TimeCursorPosition, FFrameTime(0), LastValidFrame);
+	EndTime = FMath::Clamp(EndAtTime * AnimationResolution, FFrameTime(0), LastValidFrame);
 
 	if ( PlayMode == EUMGSequencePlayMode::PingPong )
 	{
@@ -167,8 +174,8 @@ void UUMGSequencePlayer::PlayInternal(double StartAtTime, double EndAtTime, doub
 	// rendered using the first frames data, otherwise you may see a *pop* due to a widget being constructed with a default different than the
 	// first frame of the animation.
 	if (RootTemplateInstance.IsValid())
-	{			
-		const FMovieSceneContext Context(FMovieSceneEvaluationRange(TimeCursorPosition, TimeCursorPosition), PlayerStatus);
+	{
+		const FMovieSceneContext Context(FMovieSceneEvaluationRange(AbsolutePlaybackStart + TimeCursorPosition, AbsolutePlaybackStart + TimeCursorPosition, AnimationResolution), PlayerStatus);
 		RootTemplateInstance.Evaluate(Context, *this);
 	}
 
@@ -178,18 +185,12 @@ void UUMGSequencePlayer::PlayInternal(double StartAtTime, double EndAtTime, doub
 
 void UUMGSequencePlayer::Play(float StartAtTime, int32 InNumLoopsToPlay, EUMGSequencePlayMode::Type InPlayMode, float InPlaybackSpeed)
 {
-	double SubAnimStartTime = 0.0;
-	double SubAnimEndTime = TimeRange.Size<float>();
-
-	PlayInternal(StartAtTime, 0.0, SubAnimStartTime, SubAnimEndTime, InNumLoopsToPlay, InPlayMode, InPlaybackSpeed);
+	PlayInternal(StartAtTime, 0.0, InNumLoopsToPlay, InPlayMode, InPlaybackSpeed);
 }
 
 void UUMGSequencePlayer::PlayTo(float StartAtTime, float EndAtTime, int32 InNumLoopsToPlay, EUMGSequencePlayMode::Type InPlayMode, float InPlaybackSpeed)
 {
-	double SubAnimStartTime = 0.0;
-	double SubAnimEndTime = TimeRange.Size<float>();
-
-	PlayInternal(StartAtTime, EndAtTime, SubAnimStartTime, SubAnimEndTime, InNumLoopsToPlay, InPlayMode, InPlaybackSpeed);
+	PlayInternal(StartAtTime, EndAtTime, InNumLoopsToPlay, InPlayMode, InPlaybackSpeed);
 }
 
 void UUMGSequencePlayer::Pause()
@@ -204,7 +205,7 @@ void UUMGSequencePlayer::Pause()
 	PlayerStatus = EMovieScenePlayerStatus::Stopped;
 
 	// Evaluate the sequence at its current time, with a status of 'stopped' to ensure that animated state pauses correctly. (ie. audio sounds should stop/pause)
-	const FMovieSceneContext Context(FMovieSceneEvaluationRange(TimeCursorPosition, TimeCursorPosition), PlayerStatus);
+	const FMovieSceneContext Context(FMovieSceneEvaluationRange(AbsolutePlaybackStart + TimeCursorPosition, AbsolutePlaybackStart + TimeCursorPosition, AnimationResolution), PlayerStatus);
 	RootTemplateInstance.Evaluate(Context, *this);
 
 	ApplyLatentActions();
@@ -230,7 +231,7 @@ void UUMGSequencePlayer::Stop()
 
 	if (RootTemplateInstance.IsValid())
 	{
-		const FMovieSceneContext Context(FMovieSceneEvaluationRange(0), PlayerStatus);
+		const FMovieSceneContext Context(FMovieSceneEvaluationRange(AbsolutePlaybackStart, AnimationResolution), PlayerStatus);
 		RootTemplateInstance.Evaluate(Context, *this);
 		RootTemplateInstance.Finish(*this);
 	}
@@ -238,7 +239,7 @@ void UUMGSequencePlayer::Stop()
 	OnSequenceFinishedPlayingEvent.Broadcast(*this);
 	Animation->OnAnimationFinished.Broadcast();
 
-	TimeCursorPosition = 0;
+	TimeCursorPosition = FFrameTime(0);
 }
 
 void UUMGSequencePlayer::SetNumLoopsToPlay(int32 InNumLoopsToPlay)
diff --git a/Engine/Source/Runtime/UMG/Private/Animation/WidgetAnimation.cpp b/Engine/Source/Runtime/UMG/Private/Animation/WidgetAnimation.cpp
--- a/Engine/Source/Runtime/UMG/Private/Animation/WidgetAnimation.cpp
+++ b/Engine/Source/Runtime/UMG/Private/Animation/WidgetAnimation.cpp
@@ -37,6 +37,8 @@ UWidgetAnimation* UWidgetAnimation::GetNullAnimation()
 		NullAnimation->AddToRoot();
 		NullAnimation->MovieScene = NewObject<UMovieScene>(NullAnimation, FName("No Animation"));
 		NullAnimation->MovieScene->AddToRoot();
+
+		NullAnimation->MovieScene->SetPlaybackFrameRate(FFrameRate(20, 1));
 	}
 
 	return NullAnimation;
@@ -47,13 +49,13 @@ UWidgetAnimation* UWidgetAnimation::GetNullAnimation()
 
 float UWidgetAnimation::GetStartTime() const
 {
-	return MovieScene->GetPlaybackRange().GetLowerBoundValue();
+	return MovieScene->GetPlaybackRange().GetLowerBoundValue() / MovieScene->GetFrameResolution();
 }
 
 
 float UWidgetAnimation::GetEndTime() const
 {
-	return MovieScene->GetPlaybackRange().GetUpperBoundValue();
+	return MovieScene->GetPlaybackRange().GetUpperBoundValue() / MovieScene->GetFrameResolution();
 }
 
 
diff --git a/Engine/Source/Runtime/UMG/Private/UserWidget.cpp b/Engine/Source/Runtime/UMG/Private/UserWidget.cpp
--- a/Engine/Source/Runtime/UMG/Private/UserWidget.cpp
+++ b/Engine/Source/Runtime/UMG/Private/UserWidget.cpp
@@ -690,7 +690,7 @@ float UUserWidget::PauseAnimation(const UWidgetAnimation* InAnimation)
 		if ( FoundPlayer )
 		{
 			( *FoundPlayer )->Pause();
-			return (float)( *FoundPlayer )->GetTimeCursorPosition();
+			return (float)( *FoundPlayer )->GetCurrentTime().AsSeconds();
 		}
 	}
 
@@ -704,7 +704,7 @@ float UUserWidget::GetAnimationCurrentTime(const UWidgetAnimation* InAnimation)
 		const UUMGSequencePlayer*const* FoundPlayer = ActiveSequencePlayers.FindByPredicate([&](const UUMGSequencePlayer* Player) { return Player->GetAnimation() == InAnimation; });
 		if (FoundPlayer)
 		{
-			return (float)(*FoundPlayer)->GetTimeCursorPosition();
+			return (float)(*FoundPlayer)->GetCurrentTime().AsSeconds();
 		}
 	}
 
diff --git a/Engine/Source/Runtime/UMG/Public/Animation/MovieScene2DTransformSection.h b/Engine/Source/Runtime/UMG/Public/Animation/MovieScene2DTransformSection.h
--- a/Engine/Source/Runtime/UMG/Public/Animation/MovieScene2DTransformSection.h
+++ b/Engine/Source/Runtime/UMG/Public/Animation/MovieScene2DTransformSection.h
@@ -4,42 +4,10 @@
 
 #include "CoreMinimal.h"
 #include "UObject/ObjectMacros.h"
-#include "Curves/KeyHandle.h"
-#include "Curves/RichCurve.h"
+#include "Channels/MovieSceneFloatChannel.h"
 #include "MovieSceneSection.h"
-#include "Slate/WidgetTransform.h"
-#include "Sections/IKeyframeSection.h"
 #include "MovieScene2DTransformSection.generated.h"
 
-enum class EKey2DTransformChannel
-{
-	Translation,
-	Rotation,
-	Scale,
-	Shear
-};
-
-
-enum class EKey2DTransformAxis
-{
-	X,
-	Y,
-	None
-};
-
-struct F2DTransformKey
-{
-	F2DTransformKey( EKey2DTransformChannel InChannel, EKey2DTransformAxis InAxis, float InValue )
-	{
-		Channel = InChannel;
-		Axis = InAxis;
-		Value = InValue;
-	}
-	EKey2DTransformChannel Channel;
-	EKey2DTransformAxis Axis;
-	float Value;
-};
-
 
 /**
  * A transform section
@@ -47,62 +15,24 @@ struct F2DTransformKey
 UCLASS(MinimalAPI)
 class UMovieScene2DTransformSection
 	: public UMovieSceneSection
-	, public IKeyframeSection<F2DTransformKey>
 {
 	GENERATED_UCLASS_BODY()
 
 public:
 
-	// UMovieSceneSection interface
-
-	virtual void MoveSection(float DeltaPosition, TSet<FKeyHandle>& KeyHandles) override;
-	virtual void DilateSection(float DilationFactor, float Origin, TSet<FKeyHandle>& KeyHandles) override;
-	virtual void GetKeyHandles(TSet<FKeyHandle>& OutKeyHandles, TRange<float> TimeRange) const override;
-	virtual TOptional<float> GetKeyTime( FKeyHandle KeyHandle ) const override;
-	virtual void SetKeyTime( FKeyHandle KeyHandle, float Time ) override;
-
-public:
-
-	UMG_API FRichCurve& GetTranslationCurve( EAxis::Type Axis );
-	UMG_API const FRichCurve& GetTranslationCurve( EAxis::Type Axis ) const;
-
-	UMG_API FRichCurve& GetRotationCurve();
-	UMG_API const FRichCurve& GetRotationCurve() const;
-
-	UMG_API FRichCurve& GetScaleCurve( EAxis::Type Axis );
-	UMG_API const FRichCurve& GetScaleCurve( EAxis::Type Axis ) const;
-
-	DEPRECATED(4.15, "Please use GetShearCurve.")
-	UMG_API FRichCurve& GetSheerCurve( EAxis::Type Axis );
-	
-	UMG_API FRichCurve& GetShearCurve( EAxis::Type Axis );
-	UMG_API const FRichCurve& GetShearCurve( EAxis::Type Axis ) const;
-
-	DEPRECATED(4.15, "Evaluation is now the responsibility of FMovieScene2DTransformSectionTemplate")
-	FWidgetTransform Eval( float Position, const FWidgetTransform& DefaultValue ) const;
-
-	// IKeyframeSection interface.
-	virtual bool NewKeyIsNewData( float Time, const struct F2DTransformKey& TransformKey ) const override;
-	virtual bool HasKeys( const struct F2DTransformKey& TransformKey ) const override;
-	virtual void AddKey( float Time, const struct F2DTransformKey& TransformKey, EMovieSceneKeyInterpolation KeyInterpolation ) override;
-	virtual void SetDefault( const struct F2DTransformKey& TransformKey ) override;
-	virtual void ClearDefaults() override;
-
-private:
-
 	/** Translation curves*/
 	UPROPERTY()
-	FRichCurve Translation[2];
+	FMovieSceneFloatChannel Translation[2];
 	
 	/** Rotation curve */
 	UPROPERTY()
-	FRichCurve Rotation;
+	FMovieSceneFloatChannel Rotation;
 
 	/** Scale curves */
 	UPROPERTY()
-	FRichCurve Scale[2];
+	FMovieSceneFloatChannel Scale[2];
 
 	/** Shear curve */
 	UPROPERTY()
-	FRichCurve Shear[2];
+	FMovieSceneFloatChannel Shear[2];
 };
diff --git a/Engine/Source/Runtime/UMG/Public/Animation/MovieScene2DTransformTrack.h b/Engine/Source/Runtime/UMG/Public/Animation/MovieScene2DTransformTrack.h
--- a/Engine/Source/Runtime/UMG/Public/Animation/MovieScene2DTransformTrack.h
+++ b/Engine/Source/Runtime/UMG/Public/Animation/MovieScene2DTransformTrack.h
@@ -20,14 +20,8 @@ class UMovieScene2DTransformTrack : public UMovieScenePropertyTrack
 public:
 
 	UMovieScene2DTransformTrack(const FObjectInitializer& ObjectInitializer);
-	
-	// UMovieSceneTrack interface
 
+	//~ UMovieSceneTrack interface
 	virtual UMovieSceneSection* CreateNewSection() override;
 	virtual FMovieSceneEvalTemplatePtr CreateTemplateForSection(const UMovieSceneSection& InSection) const override;
-	
-public:
-
-	DEPRECATED(4.15, "Please evaluate using FMovieScene2DTransformTemplate.")
-	bool Eval( float Position, float LastPostion, FWidgetTransform& InOutTransform ) const;
 };
diff --git a/Engine/Source/Runtime/UMG/Public/Animation/MovieSceneMarginSection.h b/Engine/Source/Runtime/UMG/Public/Animation/MovieSceneMarginSection.h
--- a/Engine/Source/Runtime/UMG/Public/Animation/MovieSceneMarginSection.h
+++ b/Engine/Source/Runtime/UMG/Public/Animation/MovieSceneMarginSection.h
@@ -8,101 +8,34 @@
 #include "Curves/RichCurve.h"
 #include "MovieSceneSection.h"
 #include "Layout/Margin.h"
-#include "Sections/IKeyframeSection.h"
 #include "Evaluation/MovieSceneEvalTemplate.h"
+#include "Channels/MovieSceneFloatChannel.h"
 #include "MovieSceneMarginSection.generated.h"
 
-enum class EKeyMarginChannel
-{
-	Left,
-	Top,
-	Right,
-	Bottom
-};
-
-
-struct FMarginKey
-{
-	FMarginKey( EKeyMarginChannel InChannel, float InValue )
-	{
-		Channel = InChannel;
-		Value = InValue;
-	}
-	EKeyMarginChannel Channel;
-	float Value;
-};
-
-
 /**
  * A section in a Margin track
  */
 UCLASS(MinimalAPI)
 class UMovieSceneMarginSection 
 	: public UMovieSceneSection
-	, public IKeyframeSection<FMarginKey>
 {
 	GENERATED_UCLASS_BODY()
-public:
-	/** UMovieSceneSection interface */
-	virtual void MoveSection(float DeltaPosition, TSet<FKeyHandle>& KeyHandles) override;
-	virtual void DilateSection(float DilationFactor, float Origin, TSet<FKeyHandle>& KeyHandles) override;
-	virtual void GetKeyHandles(TSet<FKeyHandle>& OutKeyHandles, TRange<float> TimeRange) const override;
-	virtual TOptional<float> GetKeyTime( FKeyHandle KeyHandle ) const override;
-	virtual void SetKeyTime( FKeyHandle KeyHandle, float Time ) override;
-
-	// IKeyframeSection interface.
-	virtual void AddKey( float Time, const FMarginKey& MarginKey, EMovieSceneKeyInterpolation KeyInterpolation ) override;
-	virtual bool NewKeyIsNewData(float Time, const FMarginKey& Key ) const override;
-	virtual bool HasKeys(const FMarginKey& Key) const override;
-	virtual void SetDefault(const FMarginKey& Key) override;
-	virtual void ClearDefaults() override;
 
-	/**
-	 * Gets the top curve
-	 *
-	 * @return The rich curve for this margin channel
-	 */
-	FRichCurve& GetTopCurve() { return TopCurve; }
-	const FRichCurve& GetTopCurve() const { return TopCurve; }
-
-	/**
-	 * Gets the left curve
-	 *
-	 * @return The rich curve for this margin channel
-	 */
-	FRichCurve& GetLeftCurve() { return LeftCurve; }
-	const FRichCurve& GetLeftCurve() const { return LeftCurve; }
-	/**
-	 * Gets the right curve
-	 *
-	 * @return The rich curve for this margin channel
-	 */
-	FRichCurve& GetRightCurve() { return RightCurve; }
-	const FRichCurve& GetRightCurve() const { return RightCurve; }
-	
-	/**
-	 * Gets the bottom curve
-	 *
-	 * @return The rich curve for this margin channel
-	 */
-	FRichCurve& GetBottomCurve() { return BottomCurve; }
-	const FRichCurve& GetBottomCurve() const { return BottomCurve; }
-
-private:
+public:
 
 	/** Red curve data */
 	UPROPERTY()
-	FRichCurve TopCurve;
+	FMovieSceneFloatChannel TopCurve;
 
 	/** Green curve data */
 	UPROPERTY()
-	FRichCurve LeftCurve;
+	FMovieSceneFloatChannel LeftCurve;
 
 	/** Blue curve data */
 	UPROPERTY()
-	FRichCurve RightCurve;
+	FMovieSceneFloatChannel RightCurve;
 
 	/** Alpha curve data */
 	UPROPERTY()
-	FRichCurve BottomCurve;
+	FMovieSceneFloatChannel BottomCurve;
 };
diff --git a/Engine/Source/Runtime/UMG/Public/Animation/UMGSequencePlayer.h b/Engine/Source/Runtime/UMG/Public/Animation/UMGSequencePlayer.h
--- a/Engine/Source/Runtime/UMG/Public/Animation/UMGSequencePlayer.h
+++ b/Engine/Source/Runtime/UMG/Public/Animation/UMGSequencePlayer.h
@@ -8,6 +8,7 @@
 #include "Blueprint/UserWidget.h"
 #include "IMovieScenePlayer.h"
 #include "Evaluation/MovieSceneEvaluationTemplateInstance.h"
+#include "QualifiedFrameTime.h"
 #include "UMGSequencePlayer.generated.h"
 
 class UWidgetAnimation;
@@ -39,7 +40,10 @@ public:
 	void Reverse();
 
 	/** Gets the current time position in the player (in seconds). */
-	double GetTimeCursorPosition() const { return TimeCursorPosition; }
+	DEPRECATED(4.20, "Please use GetCurrentTime instead.")
+	double GetTimeCursorPosition() const { return GetCurrentTime().AsSeconds(); }
+
+	FQualifiedFrameTime GetCurrentTime() const { return FQualifiedFrameTime(TimeCursorPosition, AnimationResolution); }
 
 	/** @return The current animation being played */
 	const UWidgetAnimation* GetAnimation() const { return Animation; }
@@ -68,7 +72,7 @@ public:
 
 private:
 	/** Internal play function with a verbose parameter set */
-	void PlayInternal(double StartAtTime, double EndAtTime, double SubAnimStartTime, double SubAnimEndTime, int32 InNumLoopsToPlay, EUMGSequencePlayMode::Type InPlayMode, float InPlaybackSpeed);
+	void PlayInternal(double StartAtTime, double EndAtTime, int32 InNumLoopsToPlay, EUMGSequencePlayMode::Type InPlayMode, float InPlaybackSpeed);
 
 	/** Apply any latent actions which may have accumulated while the sequence was being evaluated */
 	void ApplyLatentActions();
@@ -82,20 +86,19 @@ private:
 
 	FMovieSceneRootEvaluationTemplateInstance RootTemplateInstance;
 
-	/** Time range of the animation */
-	TRange<float> TimeRange;
+	/** The resolution at which all FFrameNumbers are stored */
+	FFrameRate AnimationResolution;
 
-	/** The current time cursor position within the sequence (in seconds) */
-	double TimeCursorPosition;
+	FFrameNumber AbsolutePlaybackStart;
 
-	/** The time the animation should end, only effects the last loop (in seconds) */
-	double EndTime;
+	/** The current time cursor position within the sequence, between 0 and Duration */
+	FFrameTime TimeCursorPosition;
 
-	/** The offset from 0 to the start of the animation (in seconds) */
-	double AnimationStartOffset;
+	/** The duration of the sequence */
+	int32 Duration;
 
-	/** Time range of the current play of the animation */
-	TRange<double> CurrentPlayRange;
+	/** Time at which to end the animation after looping */
+	FFrameTime EndTime;
 
 	/** Status of the player (e.g play, stopped) */
 	EMovieScenePlayerStatus::Type PlayerStatus;
