UMG: Animating an FSlateBrush's material will now invalidate the cached widget's layout and volatility when setting MIDs

  - This fixes an issue where the brush's material is still referenced by widgets held within invalidation panels, causing a crash attempting to render that batch if the material is no longer referenced (and therefore subject to GC)

#tests Ran a cooked soak test of paragon and preflighted with "Test Orion PS4" which now passes.

diff --git a/Engine/Source/Runtime/UMG/Private/Animation/MovieSceneWidgetMaterialTemplate.cpp b/Engine/Source/Runtime/UMG/Private/Animation/MovieSceneWidgetMaterialTemplate.cpp
--- a/Engine/Source/Runtime/UMG/Private/Animation/MovieSceneWidgetMaterialTemplate.cpp
+++ b/Engine/Source/Runtime/UMG/Private/Animation/MovieSceneWidgetMaterialTemplate.cpp
@@ -39,7 +39,7 @@ struct FWidgetMaterialAccessor : FDefaultMaterialAccessor
 		if (UWidget* Widget = Cast<UWidget>(&Object))
 		{
 			FWidgetMaterialHandle Handle = WidgetMaterialTrackUtilities::GetMaterialHandle(Widget, BrushPropertyNamePath);
-			Handle.SetMaterial(&Material);
+			Handle.SetMaterial(&Material, Widget);
 		}
 	}
 
diff --git a/Engine/Source/Runtime/UMG/Private/Animation/WidgetMaterialTrackUtilities.cpp b/Engine/Source/Runtime/UMG/Private/Animation/WidgetMaterialTrackUtilities.cpp
--- a/Engine/Source/Runtime/UMG/Private/Animation/WidgetMaterialTrackUtilities.cpp
+++ b/Engine/Source/Runtime/UMG/Private/Animation/WidgetMaterialTrackUtilities.cpp
@@ -103,19 +103,25 @@ UMaterialInterface* FWidgetMaterialHandle::GetMaterial() const
 	}
 }
 
-void FWidgetMaterialHandle::SetMaterial(UMaterialInterface* InMaterial)
+void FWidgetMaterialHandle::SetMaterial(UMaterialInterface* InMaterial, UWidget* OwnerWidget)
 {
 	if(TypeName == TMaterialStructType<FSlateFontInfo>::GetTypeName())
 	{
-		return TMaterialStructType<FSlateFontInfo>::SetMaterial(Data, InMaterial);
+		TMaterialStructType<FSlateFontInfo>::SetMaterial(Data, InMaterial);
 	}
 	else if(TypeName == TMaterialStructType<FSlateBrush>::GetTypeName())
 	{
-		return TMaterialStructType<FSlateBrush>::SetMaterial(Data, InMaterial);
+		TMaterialStructType<FSlateBrush>::SetMaterial(Data, InMaterial);
 	}
 	else if(TypeName == TMaterialStructType<FFontOutlineSettings>::GetTypeName())
 	{
-		return TMaterialStructType<FFontOutlineSettings>::SetMaterial(Data, InMaterial);
+		TMaterialStructType<FFontOutlineSettings>::SetMaterial(Data, InMaterial);
+	}
+
+	TSharedPtr<SWidget> RawWidget = OwnerWidget->GetCachedWidget();
+	if (RawWidget.IsValid())
+	{
+		RawWidget->Invalidate(EInvalidateWidget::LayoutAndVolatility);
 	}
 }
 
diff --git a/Engine/Source/Runtime/UMG/Public/Animation/WidgetMaterialTrackUtilities.h b/Engine/Source/Runtime/UMG/Public/Animation/WidgetMaterialTrackUtilities.h
--- a/Engine/Source/Runtime/UMG/Public/Animation/WidgetMaterialTrackUtilities.h
+++ b/Engine/Source/Runtime/UMG/Public/Animation/WidgetMaterialTrackUtilities.h
@@ -30,7 +30,8 @@ public:
 	UMaterialInterface* GetMaterial() const;
 
 	/** Sets the new bound material */
-	void SetMaterial(UMaterialInterface* InMaterial);
+	void SetMaterial(UMaterialInterface* InMaterial, UWidget* OwnerWidget);
+
 private:
 	/** Struct typename for that the data is pointing to */
 	FName TypeName;
