Restore changes lost in Dev-Documentation merge
#rnx

diff --git a/Engine/Source/Runtime/UMG/Private/Components/Image.cpp b/Engine/Source/Runtime/UMG/Private/Components/Image.cpp
--- a/Engine/Source/Runtime/UMG/Private/Components/Image.cpp
+++ b/Engine/Source/Runtime/UMG/Private/Components/Image.cpp
@@ -109,6 +109,17 @@ void UImage::SetBrushSize(FVector2D DesiredSize)
 	}
 }
 
+void UImage::SetBrushTintColor(FSlateColor TintColor)
+{
+	Brush.TintColor = TintColor;
+
+	if (MyImage.IsValid())
+	{
+		MyImage->SetImage(&Brush);
+		MyImage->Invalidate(EInvalidateWidget::PaintAndVolatility);
+	}
+}
+
 void UImage::SetBrushFromAsset(USlateBrushAsset* Asset)
 {
 	Brush = Asset ? Asset->Brush : FSlateBrush();
@@ -222,6 +233,7 @@ UMaterialInstanceDynamic* UImage::GetDynamicMaterial()
 			if ( MyImage.IsValid() )
 			{
 				MyImage->SetImage(&Brush);
+				MyImage->Invalidate(EInvalidateWidget::LayoutAndVolatility);
 			}
 		}
 		
