Fix for crash when ObjectTable is null

diff --git a/Engine/Source/Runtime/UMG/Public/Slate/SObjectTableRow.h b/Engine/Source/Runtime/UMG/Public/Slate/SObjectTableRow.h
--- a/Engine/Source/Runtime/UMG/Public/Slate/SObjectTableRow.h
+++ b/Engine/Source/Runtime/UMG/Public/Slate/SObjectTableRow.h
@@ -96,11 +96,14 @@ public:
 		// Since event-based selection changes are only really needed in a game scenario, we (crudely) monitor it here to generate events.
 		// If desired, per-item selection events could be added as a longer-term todo
 		TSharedPtr<ITypedTableView<ItemType>> OwnerTable = OwnerTablePtr.Pin();
-		const ItemType& MyItem = *OwnerTable->Private_ItemFromWidget(this);
-		if (bIsAppearingSelected != OwnerTable->Private_IsItemSelected(MyItem))
+		if (OwnerTable.IsValid())
 		{
-			bIsAppearingSelected = !bIsAppearingSelected;
-			OnItemSelectionChanged(bIsAppearingSelected);
+			const ItemType& MyItem = *OwnerTable->Private_ItemFromWidget(this);
+			if (bIsAppearingSelected != OwnerTable->Private_IsItemSelected(MyItem))
+			{
+				bIsAppearingSelected = !bIsAppearingSelected;
+				OnItemSelectionChanged(bIsAppearingSelected);
+			}
 		}
 	}
 
