RL - Fix Widget Renderer using old style FDeferredCleanupInterface

diff --git a/Engine/Source/Runtime/UMG/Private/Slate/WidgetRenderer.cpp b/Engine/Source/Runtime/UMG/Private/Slate/WidgetRenderer.cpp
--- a/Engine/Source/Runtime/UMG/Private/Slate/WidgetRenderer.cpp
+++ b/Engine/Source/Runtime/UMG/Private/Slate/WidgetRenderer.cpp
@@ -195,8 +195,3 @@ void FWidgetRenderer::DrawWindow(
 	}
 #endif // !UE_SERVER
 }
-
-void FWidgetRenderer::FinishCleanup()
-{
-	delete this;
-}
diff --git a/Engine/Source/Runtime/UMG/Public/Slate/WidgetRenderer.h b/Engine/Source/Runtime/UMG/Public/Slate/WidgetRenderer.h
--- a/Engine/Source/Runtime/UMG/Public/Slate/WidgetRenderer.h
+++ b/Engine/Source/Runtime/UMG/Public/Slate/WidgetRenderer.h
@@ -82,8 +82,6 @@ public:
 
 	TArray< TSharedPtr<FSlateWindowElementList::FDeferredPaint> > DeferredPaints;
 
-	/** FDeferredCleanup interface */
-	virtual void FinishCleanup();
 private:
 	/** The slate 3D renderer used to render the user slate widget */
 	TSharedPtr<ISlate3DRenderer, ESPMode::ThreadSafe> Renderer;
