Skip UMG runtime warning about slow CreateWidget path in nativized widget blueprint classes in non-shipping builds.

Note: The fast path is not currently compatible with Blueprint nativization, so it is already being implicitly disabled at cook time if the user did not explicitly force the slow path (see FWidgetBlueprintCompiler::CanAllowTemplate).

diff --git a/Engine/Source/Editor/UMGEditor/Private/WidgetBlueprintCompiler.cpp b/Engine/Source/Editor/UMGEditor/Private/WidgetBlueprintCompiler.cpp
--- a/Engine/Source/Editor/UMGEditor/Private/WidgetBlueprintCompiler.cpp
+++ b/Engine/Source/Editor/UMGEditor/Private/WidgetBlueprintCompiler.cpp
@@ -405,7 +405,7 @@ bool FWidgetBlueprintCompiler::CanAllowTemplate(FCompilerResultsLog& MessageLog,
 	}
 
 	// For now we don't support nativization, it's going to require some extra work moving the template support
-	// during the nativization process.
+	// during the nativization process. Also see UUserWidget::CreateInstanceInternal (we skip the runtime warning for the nativized case).
 	if ( WidgetBP->NativizationFlag != EBlueprintNativizationFlag::Disabled )
 	{
 		MessageLog.Warning(*LOCTEXT("TemplatingAndNativization", "Nativization and Fast Widget Creation is not supported at this time.").ToString());
diff --git a/Engine/Source/Runtime/UMG/Private/UserWidget.cpp b/Engine/Source/Runtime/UMG/Private/UserWidget.cpp
--- a/Engine/Source/Runtime/UMG/Private/UserWidget.cpp
+++ b/Engine/Source/Runtime/UMG/Private/UserWidget.cpp
@@ -2054,7 +2054,12 @@ UUserWidget* UUserWidget::CreateInstanceInternal(UObject* Outer, TSubclassOf<UUs
 #if !WITH_EDITOR && (UE_BUILD_DEBUG || UE_BUILD_DEVELOPMENT)
 	else
 	{
-		UE_LOG(LogUMG, Warning, TEXT("Widget Class %s - Using Slow CreateWidget path because this class could not be templated."), *UserWidgetClass->GetName());
+		// Nativized widget blueprint class types (UDynamicClass) do not currently support the fast path (see FWidgetBlueprintCompiler::CanAllowTemplate), so we bypass the runtime warning in that case.
+		const bool bIsDynamicClass = Cast<UDynamicClass>(UserWidgetClass) != nullptr;
+		if (!bIsDynamicClass)
+		{
+			UE_LOG(LogUMG, Warning, TEXT("Widget Class %s - Using Slow CreateWidget path because this class could not be templated."), *UserWidgetClass->GetName());
+		}
 	}
 #endif
 
