Fix for delayed destruction of UWidgets when they are manually removed from a panel as part of tear down.
Inspired by the pull request, however I put in a more generic fix.

PR #4904: Fix late release of Slate resources managed by UMG slot widgets (Contributed by [REDACTED_PERSON])

diff --git a/Engine/Source/Runtime/UMG/Private/Components/PanelSlot.cpp b/Engine/Source/Runtime/UMG/Private/Components/PanelSlot.cpp
--- a/Engine/Source/Runtime/UMG/Private/Components/PanelSlot.cpp
+++ b/Engine/Source/Runtime/UMG/Private/Components/PanelSlot.cpp
@@ -20,7 +20,8 @@ void UPanelSlot::ReleaseSlateResources(bool bReleaseChildren)
 {
 	Super::ReleaseSlateResources(bReleaseChildren);
 
-	if ( Content )
+	// ReleaseSlateResources for Content unless the content is a UUserWidget as they are responsible for releasing their own content.
+	if (Content && !Content->IsA<UUserWidget>())
 	{
 		Content->ReleaseSlateResources(bReleaseChildren);
 	}
diff --git a/Engine/Source/Runtime/UMG/Private/Components/PanelWidget.cpp b/Engine/Source/Runtime/UMG/Private/Components/PanelWidget.cpp
--- a/Engine/Source/Runtime/UMG/Private/Components/PanelWidget.cpp
+++ b/Engine/Source/Runtime/UMG/Private/Components/PanelWidget.cpp
@@ -89,13 +89,8 @@ bool UPanelWidget::RemoveChildAt(int32 Index)
 
 	OnSlotRemoved(PanelSlot);
 
-	// If the child is a UserWidget, we should let it manage it's own slate resources instead of forcing a clear here. This fixes issues such as UE-39106
-	if (PanelSlot->Content && !PanelSlot->Content->IsA<UUserWidget>())
-	{
-		const bool bReleaseChildren = true;
-		PanelSlot->ReleaseSlateResources(bReleaseChildren);
-	}
-
+	const bool bReleaseChildren = true;
+	PanelSlot->ReleaseSlateResources(bReleaseChildren);
 	PanelSlot->Parent = nullptr;
 	PanelSlot->Content = nullptr;
 
