Fixed recursion of calls through SetDesignerFlags to avoid double-recursion with many nested panels

diff --git a/Engine/Source/Runtime/UMG/Private/Components/NamedSlotInterface.cpp b/Engine/Source/Runtime/UMG/Private/Components/NamedSlotInterface.cpp
--- a/Engine/Source/Runtime/UMG/Private/Components/NamedSlotInterface.cpp
+++ b/Engine/Source/Runtime/UMG/Private/Components/NamedSlotInterface.cpp
@@ -40,3 +40,17 @@ void INamedSlotInterface::ReleaseNamedSlotSlateResources(bool bReleaseChildren)
 		}
 	}
 }
+
+void INamedSlotInterface::SetNamedSlotDesignerFlags(EWidgetDesignFlags::Type NewFlags)
+{
+	TArray<FName> SlotNames;
+	GetSlotNames(SlotNames);
+
+	for (const FName& SlotName : SlotNames)
+	{
+		if (UWidget* Content = GetContentForSlot(SlotName))
+		{
+			Content->SetDesignerFlags(NewFlags);
+		}
+	}
+}
\ No newline at end of file
diff --git a/Engine/Source/Runtime/UMG/Private/Components/Widget.cpp b/Engine/Source/Runtime/UMG/Private/Components/Widget.cpp
--- a/Engine/Source/Runtime/UMG/Private/Components/Widget.cpp
+++ b/Engine/Source/Runtime/UMG/Private/Components/Widget.cpp
@@ -25,6 +25,7 @@
 #include "Types/ReflectionMetadata.h"
 #include "Serialization/PropertyLocalizationDataGathering.h"
 #include "HAL/LowLevelMemTracker.h"
+#include "Components/NamedSlotInterface.h"
 
 #define LOCTEXT_NAMESPACE "UMG"
 
@@ -866,6 +867,12 @@ ULocalPlayer* UWidget::GetOwningLocalPlayer() const
 void UWidget::SetDesignerFlags(EWidgetDesignFlags::Type NewFlags)
 {
 	DesignerFlags = ( EWidgetDesignFlags::Type )( DesignerFlags | NewFlags );
+
+	INamedSlotInterface* NamedSlotWidget = Cast<INamedSlotInterface>(this);
+	if (NamedSlotWidget)
+	{
+		NamedSlotWidget->SetNamedSlotDesignerFlags(NewFlags);
+	}
 }
 
 void UWidget::SetDisplayLabel(const FString& InDisplayLabel)
diff --git a/Engine/Source/Runtime/UMG/Private/UserWidget.cpp b/Engine/Source/Runtime/UMG/Private/UserWidget.cpp
--- a/Engine/Source/Runtime/UMG/Private/UserWidget.cpp
+++ b/Engine/Source/Runtime/UMG/Private/UserWidget.cpp
@@ -1282,13 +1282,11 @@ const FText UUserWidget::GetPaletteCategory()
 
 void UUserWidget::SetDesignerFlags(EWidgetDesignFlags::Type NewFlags)
 {
-	Super::SetDesignerFlags(NewFlags);
+	UWidget::SetDesignerFlags(NewFlags);
 
 	if ( ensure(WidgetTree) )
 	{
-		WidgetTree->ForEachWidget([&] (UWidget* Widget) {
-			Widget->SetDesignerFlags(NewFlags);
-		});
+		WidgetTree->RootWidget->SetDesignerFlags(NewFlags);
 	}
 }
 
diff --git a/Engine/Source/Runtime/UMG/Public/Components/NamedSlotInterface.h b/Engine/Source/Runtime/UMG/Public/Components/NamedSlotInterface.h
--- a/Engine/Source/Runtime/UMG/Public/Components/NamedSlotInterface.h
+++ b/Engine/Source/Runtime/UMG/Public/Components/NamedSlotInterface.h
@@ -5,6 +5,7 @@
 #include "CoreMinimal.h"
 #include "UObject/ObjectMacros.h"
 #include "UObject/Interface.h"
+#include "Widget.h"
 #include "NamedSlotInterface.generated.h"
 
 class UWidget;
@@ -38,4 +39,6 @@ public:
 
 	/**  */
 	void ReleaseNamedSlotSlateResources(bool bReleaseChildren);
+
+	void SetNamedSlotDesignerFlags(EWidgetDesignFlags::Type NewFlags);
 };
