Adding check for RootWidget existing

diff --git a/Engine/Source/Runtime/UMG/Private/UserWidget.cpp b/Engine/Source/Runtime/UMG/Private/UserWidget.cpp
--- a/Engine/Source/Runtime/UMG/Private/UserWidget.cpp
+++ b/Engine/Source/Runtime/UMG/Private/UserWidget.cpp
@@ -1286,7 +1286,10 @@ void UUserWidget::SetDesignerFlags(EWidgetDesignFlags::Type NewFlags)
 
 	if ( ensure(WidgetTree) )
 	{
-		WidgetTree->RootWidget->SetDesignerFlags(NewFlags);
+		if (WidgetTree->RootWidget)
+		{
+			WidgetTree->RootWidget->SetDesignerFlags(NewFlags);
+		}
 	}
 }
 
