Adding if WITH_EDITOR

diff --git a/Engine/Source/Runtime/UMG/Private/Components/NamedSlotInterface.cpp b/Engine/Source/Runtime/UMG/Private/Components/NamedSlotInterface.cpp
--- a/Engine/Source/Runtime/UMG/Private/Components/NamedSlotInterface.cpp
+++ b/Engine/Source/Runtime/UMG/Private/Components/NamedSlotInterface.cpp
@@ -40,7 +40,7 @@ void INamedSlotInterface::ReleaseNamedSlotSlateResources(bool bReleaseChildren)
 		}
 	}
 }
-
+#if WITH_EDITOR
 void INamedSlotInterface::SetNamedSlotDesignerFlags(EWidgetDesignFlags::Type NewFlags)
 {
 	TArray<FName> SlotNames;
@@ -53,4 +53,5 @@ void INamedSlotInterface::SetNamedSlotDesignerFlags(EWidgetDesignFlags::Type New
 			Content->SetDesignerFlags(NewFlags);
 		}
 	}
-}
\ No newline at end of file
+}
+#endif
\ No newline at end of file
diff --git a/Engine/Source/Runtime/UMG/Public/Components/NamedSlotInterface.h b/Engine/Source/Runtime/UMG/Public/Components/NamedSlotInterface.h
--- a/Engine/Source/Runtime/UMG/Public/Components/NamedSlotInterface.h
+++ b/Engine/Source/Runtime/UMG/Public/Components/NamedSlotInterface.h
@@ -39,6 +39,7 @@ public:
 
 	/**  */
 	void ReleaseNamedSlotSlateResources(bool bReleaseChildren);
-
+#if WITH_EDITOR
 	void SetNamedSlotDesignerFlags(EWidgetDesignFlags::Type NewFlags);
+#endif
 };
