PR #4979: UE-62795: Deprecate bAutoWrapText in TextBlock (Contributed by [REDACTED_PERSON])
Slightly modified, the base syncronize sets the autowrap value.

diff --git a/Engine/Source/Runtime/UMG/Private/Components/TextBlock.cpp b/Engine/Source/Runtime/UMG/Private/Components/TextBlock.cpp
--- a/Engine/Source/Runtime/UMG/Private/Components/TextBlock.cpp
+++ b/Engine/Source/Runtime/UMG/Private/Components/TextBlock.cpp
@@ -21,7 +21,7 @@ UTextBlock::UTextBlock(const FObjectInitializer& ObjectInitializer)
 	ShadowOffset = FVector2D(1.0f, 1.0f);
 	ColorAndOpacity = FLinearColor::White;
 	ShadowColorAndOpacity = FLinearColor::Transparent;
-	bAutoWrapText = false;
+	bAutoWrapText_DEPRECATED = false;
 
 	if (!IsRunningDedicatedServer())
 	{
@@ -30,6 +30,17 @@ UTextBlock::UTextBlock(const FObjectInitializer& ObjectInitializer)
 	}
 }
 
+void UTextBlock::PostLoad()
+{
+	Super::PostLoad();
+
+	if (bAutoWrapText_DEPRECATED)
+	{
+		AutoWrapText = bAutoWrapText_DEPRECATED;
+		bAutoWrapText_DEPRECATED = false;
+	}
+}
+
 void UTextBlock::ReleaseSlateResources(bool bReleaseChildren)
 {
 	Super::ReleaseSlateResources(bReleaseChildren);
@@ -101,7 +112,7 @@ void UTextBlock::SetMinDesiredWidth(float InMinDesiredWidth)
 
 void UTextBlock::SetAutoWrapText(bool InAutoWrapText)
 {
-	bAutoWrapText = InAutoWrapText;
+	AutoWrapText = InAutoWrapText;
 	if(MyTextBlock.IsValid())
 	{
 		MyTextBlock->SetAutoWrapText(InAutoWrapText);
@@ -213,7 +224,6 @@ void UTextBlock::SynchronizeProperties()
 		MyTextBlock->SetShadowOffset( ShadowOffset );
 		MyTextBlock->SetShadowColorAndOpacity( ShadowColorAndOpacityBinding );
 		MyTextBlock->SetMinDesiredWidth( MinDesiredWidth );
-		MyTextBlock->SetAutoWrapText(bAutoWrapText);
 
 		Super::SynchronizeTextLayoutProperties( *MyTextBlock );
 	}
diff --git a/Engine/Source/Runtime/UMG/Public/Components/TextBlock.h b/Engine/Source/Runtime/UMG/Public/Components/TextBlock.h
--- a/Engine/Source/Runtime/UMG/Public/Components/TextBlock.h
+++ b/Engine/Source/Runtime/UMG/Public/Components/TextBlock.h
@@ -139,8 +139,8 @@ public:
 	bool bWrapWithInvalidationPanel;
 
 	/** Whether the text should automatically wrap */
-	UPROPERTY(EditAnywhere, BlueprintReadOnly, Category = Appearance)
-	bool bAutoWrapText;
+	UPROPERTY(meta = (DeprecatedProperty, DeprecationMessage = "bAutoWrapText is deprecated. Please use AutoWrapText instead."))
+	bool bAutoWrapText_DEPRECATED;
 
 	///** Called when this text is double clicked */
 	//SLATE_EVENT(FOnClicked, OnDoubleClicked)
@@ -180,6 +180,10 @@ public:
 #endif
 
 protected:
+	//~ Begin UObject Interface
+	virtual void PostLoad() override;
+	//~ End UObject Interface
+
 	//~ Begin UWidget Interface
 	virtual TSharedRef<SWidget> RebuildWidget() override;
 	virtual void OnBindingChanged(const FName& Property) override;
