A new section was added to the UMG Editor Palette to add Advanced Widgets. They were previously uncategorised.

diff --git a/Engine/Source/Editor/UMGEditor/Private/Palette/SPaletteView.cpp b/Engine/Source/Editor/UMGEditor/Private/Palette/SPaletteView.cpp
--- a/Engine/Source/Editor/UMGEditor/Private/Palette/SPaletteView.cpp
+++ b/Engine/Source/Editor/UMGEditor/Private/Palette/SPaletteView.cpp
@@ -400,6 +400,15 @@ void SPaletteView::BuildWidgetList()
 
 	// Sort the view models by name
 	WidgetViewModels.Sort([] (TSharedPtr<FWidgetViewModel> L, TSharedPtr<FWidgetViewModel> R) { return R->GetName().CompareTo(L->GetName()) > 0; });
+
+	// Take the Advanced Section, and put it at the end.
+	TSharedPtr<FWidgetViewModel>* advancedSectionPtr = WidgetViewModels.FindByPredicate([](TSharedPtr<FWidgetViewModel> widget) {return widget->GetName().CompareTo(LOCTEXT("Advanced", "Advanced")) == 0; });
+	if (advancedSectionPtr)
+	{
+		TSharedPtr<FWidgetViewModel> advancedSection = *advancedSectionPtr;
+		WidgetViewModels.Remove(advancedSection);
+		WidgetViewModels.Push(advancedSection);
+	}
 }
 
 void SPaletteView::BuildClassWidgetList()
diff --git a/Engine/Source/Runtime/UMG/Private/Components/DynamicEntryBox.cpp b/Engine/Source/Runtime/UMG/Private/Components/DynamicEntryBox.cpp
--- a/Engine/Source/Runtime/UMG/Private/Components/DynamicEntryBox.cpp
+++ b/Engine/Source/Runtime/UMG/Private/Components/DynamicEntryBox.cpp
@@ -181,6 +181,12 @@ void UDynamicEntryBox::SetEntrySpacing(const FVector2D& InEntrySpacing)
 }
 
 #if WITH_EDITOR
+
+const FText UDynamicEntryBox::GetPaletteCategory()
+{
+	return LOCTEXT("Advanced", "Advanced");
+}
+
 void UDynamicEntryBox::ValidateCompiledDefaults(IWidgetCompilerLog& CompileLog) const
 {
 	if (!EntryWidgetClass)
diff --git a/Engine/Source/Runtime/UMG/Private/Components/InputKeySelector.cpp b/Engine/Source/Runtime/UMG/Private/Components/InputKeySelector.cpp
--- a/Engine/Source/Runtime/UMG/Private/Components/InputKeySelector.cpp
+++ b/Engine/Source/Runtime/UMG/Private/Components/InputKeySelector.cpp
@@ -6,6 +6,9 @@
 #include "UObject/FrameworkObjectVersion.h"
 #include "Widgets/DeclarativeSyntaxSupport.h"
 #include "Widgets/Input/SInputKeySelector.h"
+#include "Internationalization/Internationalization.h"
+
+#define LOCTEXT_NAMESPACE "UMG"
 
 UInputKeySelector::UInputKeySelector( const FObjectInitializer& ObjectInitializer )
 	: Super(ObjectInitializer)
@@ -113,6 +116,12 @@ void UInputKeySelector::SetEscapeKeys(const TArray<FKey>& InKeys)
 	}
 	EscapeKeys = InKeys;
 }
+#if WITH_EDITOR
+const FText UInputKeySelector::GetPaletteCategory()
+{
+	return LOCTEXT("Advanced", "Advanced");
+}
+#endif
 
 void UInputKeySelector::SynchronizeProperties()
 {
@@ -172,3 +181,5 @@ void UInputKeySelector::SetTextBlockVisibility(const ESlateVisibility InVisibili
 		MyInputKeySelector->SetTextBlockVisibility(SlateVisibility);
 	}
 }
+
+#undef LOCTEXT_NAMESPACE
\ No newline at end of file
diff --git a/Engine/Source/Runtime/UMG/Private/Components/WindowTitleBarArea.cpp b/Engine/Source/Runtime/UMG/Private/Components/WindowTitleBarArea.cpp
--- a/Engine/Source/Runtime/UMG/Private/Components/WindowTitleBarArea.cpp
+++ b/Engine/Source/Runtime/UMG/Private/Components/WindowTitleBarArea.cpp
@@ -139,6 +139,13 @@ void UWindowTitleBarArea::PostLoad()
 	}
 }
 
+#if WITH_EDITOR
+const FText UWindowTitleBarArea::GetPaletteCategory()
+{
+	return LOCTEXT("Advanced", "Advanced");
+}
+#endif
+
 bool UWindowTitleBarArea::HandleWindowAction(const TSharedRef<FGenericWindow>& PlatformWindow, EWindowAction::Type WindowAction)
 {
 	if (GEngine && (WindowAction == EWindowAction::Maximize || WindowAction == EWindowAction::Restore))
diff --git a/Engine/Source/Runtime/UMG/Public/Components/DynamicEntryBox.h b/Engine/Source/Runtime/UMG/Public/Components/DynamicEntryBox.h
--- a/Engine/Source/Runtime/UMG/Public/Components/DynamicEntryBox.h
+++ b/Engine/Source/Runtime/UMG/Public/Components/DynamicEntryBox.h
@@ -91,10 +91,14 @@ public:
 
 	UFUNCTION(BlueprintCallable, Category = DynamicEntryBox)
 	void SetEntrySpacing(const FVector2D& InEntrySpacing);
+	
+	//~ Begin UWidget Interface
+#if WITH_EDITOR	
+	virtual const FText GetPaletteCategory() override;
 
-#if WITH_EDITOR
 	virtual void ValidateCompiledDefaults(class IWidgetCompilerLog& CompileLog) const override;
 #endif
+	//~ End UWidget Interface
 
 #if WITH_EDITORONLY_DATA
 	UPROPERTY(EditAnywhere, Category = DynamicEntryBox, meta = (ClampMin = 0, ClampMax = 20))
diff --git a/Engine/Source/Runtime/UMG/Public/Components/InputKeySelector.h b/Engine/Source/Runtime/UMG/Public/Components/InputKeySelector.h
--- a/Engine/Source/Runtime/UMG/Public/Components/InputKeySelector.h
+++ b/Engine/Source/Runtime/UMG/Public/Components/InputKeySelector.h
@@ -115,6 +115,9 @@ public:
 	void SetEscapeKeys(const TArray<FKey>& InKeys);
 
 	//~ Begin UWidget Interface
+#if WITH_EDITOR	
+	virtual const FText GetPaletteCategory() override;
+#endif
 	virtual void SynchronizeProperties() override;
 	//~ End UWidget Interface
 
diff --git a/Engine/Source/Runtime/UMG/Public/Components/WindowTitleBarArea.h b/Engine/Source/Runtime/UMG/Public/Components/WindowTitleBarArea.h
--- a/Engine/Source/Runtime/UMG/Public/Components/WindowTitleBarArea.h
+++ b/Engine/Source/Runtime/UMG/Public/Components/WindowTitleBarArea.h
@@ -53,6 +53,12 @@ public:
 	virtual void PostLoad() override;
 	//~ End UObject Interface
 
+	//~ Begin UWidget Interface
+#if WITH_EDITOR	
+	virtual const FText GetPaletteCategory() override;
+#endif
+	//~ End UWidget Interface
+
 protected:
 	//~ Begin UWidget Interface
 	virtual TSharedRef<SWidget> RebuildWidget() override;
