Moving UListViewBase::RequestRefresh() to be public
- Turns out a valid use case does exist that doesn't involve changing the underlying list items: triggering a re-evaluation of the scrollbar and generated entries following the modification of a displayed entry's geometry (i.e. when an entry grows or shrinks appreciably)

diff --git a/Engine/Source/Runtime/UMG/Public/Components/ListViewBase.h b/Engine/Source/Runtime/UMG/Public/Components/ListViewBase.h
--- a/Engine/Source/Runtime/UMG/Public/Components/ListViewBase.h
+++ b/Engine/Source/Runtime/UMG/Public/Components/ListViewBase.h
@@ -436,6 +436,18 @@ public:
 	UFUNCTION(BlueprintCallable, Category = ListViewBase)
 	void ScrollToBottom();
 
+	/**
+	 * Sets the list to refresh on the next tick.
+	 *
+	 * Note that refreshing, from a list perspective, is limited to accounting for discrepancies between items and entries.
+	 * In other words, it will only release entries that no longer have items and generate entries for new items (or newly visible items).
+	 *
+	 * It does NOT account for changes within existing items - that is up to the item to announce and an entry to listen to as needed.
+	 * This can be onerous to set up for simple cases, so it's also reasonable (though not ideal) to call RegenerateAllEntries when changes within N list items need to be reflected.
+	 */
+	UFUNCTION(BlueprintCallable, Category = ListViewBase)
+	void RequestRefresh();
+
 	DECLARE_EVENT_OneParam(UListView, FOnListEntryGenerated, UUserWidget&);
 	FOnListEntryGenerated& OnEntryWidgetGenerated() { return OnListEntryGeneratedEvent; }
 
@@ -456,17 +468,6 @@ protected:
 	virtual void HandleListEntryHovered(UUserWidget& EntryWidget) {}
 	virtual void HandleListEntryUnhovered(UUserWidget& EntryWidget) {}
 
-	/**
-	 * Sets the list to refresh on the next tick.
-	 *
-	 * Note that refreshing, from a list perspective, is limited to accounting for discrepancies between items and entries.
-	 * In other words, it will only release entries that no longer have items and generate entries for new items (or newly visible items).
-	 *
-	 * It does NOT account for changes within existing items - that is up to the item to announce and an entry to listen to as needed.
-	 * This can be onerous to set up for simple cases, so it's also reasonable (though not ideal) to call RegenerateAllEntries when changes within N list items need to be reflected.
-	 */
-	void RequestRefresh();
-
 	template <typename WidgetEntryT = UUserWidget, typename ObjectTableRowT = SObjectTableRow<UObject*>>
 	WidgetEntryT& GenerateTypedEntry(TSubclassOf<WidgetEntryT> WidgetClass, const TSharedRef<STableViewBase>& OwnerTable)
 	{
