Making URichTextBlock::SetText virtual (matches that of UTextBlock)

diff --git a/Engine/Source/Runtime/UMG/Public/Components/RichTextBlock.h b/Engine/Source/Runtime/UMG/Public/Components/RichTextBlock.h
--- a/Engine/Source/Runtime/UMG/Public/Components/RichTextBlock.h
+++ b/Engine/Source/Runtime/UMG/Public/Components/RichTextBlock.h
@@ -130,7 +130,7 @@ public:
 	 * @param InText The text to assign to the widget
 	 */
 	UFUNCTION(BlueprintCallable, Category = "Widget")
-	void SetText(const FText& InText);
+	virtual void SetText(const FText& InText);
 
 	const FTextBlockStyle& GetDefaultTextStyle() const;
 	const FTextBlockStyle& GetCurrentDefaultTextStyle() const;
