Change dynamic entry boxes so that native classes can create different widget of different derived classes in dynamic entry boxes.

#tests compiles and works with game code. Compiled with platform specific compiler.

diff --git a/Engine/Source/Runtime/UMG/Public/Components/DynamicEntryBox.h b/Engine/Source/Runtime/UMG/Public/Components/DynamicEntryBox.h
--- a/Engine/Source/Runtime/UMG/Public/Components/DynamicEntryBox.h
+++ b/Engine/Source/Runtime/UMG/Public/Components/DynamicEntryBox.h
@@ -39,9 +39,13 @@ public:
 	const FVector2D& GetEntrySpacing() const { return EntrySpacing; }
 	
 	template <typename WidgetT = UUserWidget>
-	WidgetT* CreateEntry()
+	WidgetT* CreateEntry(const TSubclassOf<WidgetT>& ExplicitEntryClass = nullptr)
 	{
-		if (EntryWidgetClass && EntryWidgetClass->IsChildOf<WidgetT>())
+		if (ExplicitEntryClass && ExplicitEntryClass->IsChildOf(WidgetT::StaticClass()))
+		{
+			return Cast<WidgetT>(CreateEntryInternal(ExplicitEntryClass));
+		}
+		if (EntryWidgetClass && EntryWidgetClass->IsChildOf(WidgetT::StaticClass()))
 		{
 			return Cast<WidgetT>(CreateEntryInternal(EntryWidgetClass));
 		}
