Speculative fix to this crash, but taking 2 main approaches,

1) Attempt to remove the widget sooner than end play, it's possible something is happening when we unload assets that leaves live widgets that haven't been torn down yet in a bad state.
2) The remove code looks like it could suffer from re-entrant code problems where a widget component might remove and then might attempt to remove again, adding a temporary copy step to keep the memory alive, then immediately remove it, then actually remove it from the screen, so that any downstream code attempting to remove again would be thwarted.

diff --git a/Engine/Source/Runtime/UMG/Private/Components/WidgetComponent.cpp b/Engine/Source/Runtime/UMG/Private/Components/WidgetComponent.cpp
--- a/Engine/Source/Runtime/UMG/Private/Components/WidgetComponent.cpp
+++ b/Engine/Source/Runtime/UMG/Private/Components/WidgetComponent.cpp
@@ -39,14 +39,15 @@ public:
 
 	void AddComponent(UWidgetComponent* Component)
 	{
-		if ( Component )
+		if (Component)
 		{
 			Components.AddUnique(Component);
-			if ( ScreenLayer.IsValid() )
+
+			if (TSharedPtr<SWorldWidgetScreenLayer> ScreenLayer = ScreenLayerPtr.Pin())
 			{
-				if ( UUserWidget* UserWidget = Component->GetUserWidgetObject() )
+				if (UUserWidget* UserWidget = Component->GetUserWidgetObject())
 				{
-					ScreenLayer.Pin()->AddComponent(Component, UserWidget->TakeWidget());
+					ScreenLayer->AddComponent(Component, UserWidget->TakeWidget());
 				}
 				else if (Component->GetSlateWidget().IsValid())
 				{
@@ -58,26 +59,26 @@ public:
 
 	void RemoveComponent(UWidgetComponent* Component)
 	{
-		if ( Component )
+		if (Component)
 		{
 			Components.RemoveSwap(Component);
 
-			if ( ScreenLayer.IsValid() )
+			if (TSharedPtr<SWorldWidgetScreenLayer> ScreenLayer = ScreenLayerPtr.Pin())
 			{
-				ScreenLayer.Pin()->RemoveComponent(Component);
+				ScreenLayer->RemoveComponent(Component);
 			}
 		}
 	}
 
 	virtual TSharedRef<SWidget> AsWidget() override
 	{
-		if ( ScreenLayer.IsValid() )
+		if (TSharedPtr<SWorldWidgetScreenLayer> ScreenLayer = ScreenLayerPtr.Pin())
 		{
-			return ScreenLayer.Pin().ToSharedRef();
+			return ScreenLayer.ToSharedRef();
 		}
 
 		TSharedRef<SWorldWidgetScreenLayer> NewScreenLayer = SNew(SWorldWidgetScreenLayer, OwningPlayer);
-		ScreenLayer = NewScreenLayer;
+		ScreenLayerPtr = NewScreenLayer;
 
 		// Add all the pending user widgets to the surface
 		for ( TWeakObjectPtr<UWidgetComponent>& WeakComponent : Components )
@@ -100,7 +101,7 @@ public:
 
 private:
 	FLocalPlayerContext OwningPlayer;
-	TWeakPtr<SWorldWidgetScreenLayer> ScreenLayer;
+	TWeakPtr<SWorldWidgetScreenLayer> ScreenLayerPtr;
 	TArray<TWeakObjectPtr<UWidgetComponent>> Components;
 };
 
@@ -627,8 +628,20 @@ void UWidgetComponent::BeginPlay()
 
 void UWidgetComponent::EndPlay(const EEndPlayReason::Type EndPlayReason)
 {
-	ReleaseResources();
 	Super::EndPlay(EndPlayReason);
+
+	ReleaseResources();
+}
+
+void UWidgetComponent::OnLevelRemovedFromWorld(ULevel* InLevel, UWorld* InWorld)
+{
+	// If the InLevel is null, it's a signal that the entire world is about to disappear, so
+	// go ahead and remove this widget from the viewport, it could be holding onto too many
+	// dangerous actor references that won't carry over into the next world.
+	if (InLevel == nullptr && InWorld == GetWorld())
+	{
+		ReleaseResources();
+	}
 }
 
 
@@ -793,6 +806,8 @@ void UWidgetComponent::OnRegister()
 	Super::OnRegister();
 
 #if !UE_SERVER
+	FWorldDelegates::LevelRemovedFromWorld.AddUObject(this, &ThisClass::OnLevelRemovedFromWorld);
+
 	if ( !IsRunningDedicatedServer() )
 	{
 		const bool bIsGameWorld = GetWorld()->IsGameWorld();
@@ -873,6 +888,8 @@ void UWidgetComponent::UnregisterHitTesterWithViewport(TSharedPtr<SViewport> Vie
 void UWidgetComponent::OnUnregister()
 {
 #if !UE_SERVER
+	FWorldDelegates::LevelRemovedFromWorld.RemoveAll(this);
+
 	if ( GetWorld()->IsGameWorld() )
 	{
 		TSharedPtr<SViewport> GameViewportWidget = GEngine->GetGameViewportWidget();
diff --git a/Engine/Source/Runtime/UMG/Private/Slate/SWorldWidgetScreenLayer.cpp b/Engine/Source/Runtime/UMG/Private/Slate/SWorldWidgetScreenLayer.cpp
--- a/Engine/Source/Runtime/UMG/Private/Slate/SWorldWidgetScreenLayer.cpp
+++ b/Engine/Source/Runtime/UMG/Private/Slate/SWorldWidgetScreenLayer.cpp
@@ -56,14 +56,16 @@ void SWorldWidgetScreenLayer::RemoveComponent(USceneComponent* Component)
 {
 	if ( Component )
 	{
-		if ( FComponentEntry* Entry = ComponentMap.Find(Component) )
+		if ( FComponentEntry* EntryPtr = ComponentMap.Find(Component) )
 		{
-			if ( Entry->Widget.IsValid() )
+			// First remove the component entry, in case someone tries to double delete.
+			FComponentEntry Entry = *EntryPtr;
+			ComponentMap.Remove(Component);
+
+			if (TSharedPtr<SWidget> ContainerWidget = Entry.ContainerWidget)
 			{
-				Canvas->RemoveSlot(Entry->ContainerWidget.ToSharedRef());
+				Canvas->RemoveSlot(ContainerWidget.ToSharedRef());
 			}
-
-			ComponentMap.Remove(Component);
 		}
 	}
 }
diff --git a/Engine/Source/Runtime/UMG/Public/Components/WidgetComponent.h b/Engine/Source/Runtime/UMG/Public/Components/WidgetComponent.h
--- a/Engine/Source/Runtime/UMG/Public/Components/WidgetComponent.h
+++ b/Engine/Source/Runtime/UMG/Public/Components/WidgetComponent.h
@@ -350,6 +350,8 @@ public:
 	};
 
 protected:
+	void OnLevelRemovedFromWorld(ULevel* InLevel, UWorld* InWorld);
+
 	/** Just because the user attempts to receive hardware input does not mean it's possible. */
 	bool CanReceiveHardwareInput() const;
 
diff --git a/Engine/Source/Runtime/UMG/Public/Slate/SWorldWidgetScreenLayer.h b/Engine/Source/Runtime/UMG/Public/Slate/SWorldWidgetScreenLayer.h
--- a/Engine/Source/Runtime/UMG/Public/Slate/SWorldWidgetScreenLayer.h
+++ b/Engine/Source/Runtime/UMG/Public/Slate/SWorldWidgetScreenLayer.h
@@ -56,6 +56,6 @@ private:
 		SConstraintCanvas::FSlot* Slot;
 	};
 
-	TMap<USceneComponent*, FComponentEntry> ComponentMap;
+	TMap<TWeakObjectPtr<USceneComponent>, FComponentEntry> ComponentMap;
 	TSharedPtr<SConstraintCanvas> Canvas;
 };
