Fixing main compile issue.

diff --git a/Engine/Source/Runtime/UMG/Private/Components/WidgetComponent.cpp b/Engine/Source/Runtime/UMG/Private/Components/WidgetComponent.cpp
--- a/Engine/Source/Runtime/UMG/Private/Components/WidgetComponent.cpp
+++ b/Engine/Source/Runtime/UMG/Private/Components/WidgetComponent.cpp
@@ -51,7 +51,7 @@ public:
 				}
 				else if (Component->GetSlateWidget().IsValid())
 				{
-					ScreenLayer.Pin()->AddComponent(Component, Component->GetSlateWidget().ToSharedRef());
+					ScreenLayer->AddComponent(Component, Component->GetSlateWidget().ToSharedRef());
 				}
 			}
 		}
