Defaulting URichTextBlock's visibility to SelfHitTestInvisible

diff --git a/Engine/Source/Runtime/UMG/Private/Components/RichTextBlock.cpp b/Engine/Source/Runtime/UMG/Private/Components/RichTextBlock.cpp
--- a/Engine/Source/Runtime/UMG/Private/Components/RichTextBlock.cpp
+++ b/Engine/Source/Runtime/UMG/Private/Components/RichTextBlock.cpp
@@ -20,6 +20,7 @@
 URichTextBlock::URichTextBlock(const FObjectInitializer& ObjectInitializer)
 	: Super(ObjectInitializer)
 {
+	Visibility = ESlateVisibility::SelfHitTestInvisible;
 }
 
 void URichTextBlock::ReleaseSlateResources(bool bReleaseChildren)
