Further revisions to the SWorldWidgetScreenLayer, now using FObjectKey to lookup objects and remove them.

diff --git a/Engine/Source/Runtime/UMG/Private/Slate/SWorldWidgetScreenLayer.cpp b/Engine/Source/Runtime/UMG/Private/Slate/SWorldWidgetScreenLayer.cpp
--- a/Engine/Source/Runtime/UMG/Private/Slate/SWorldWidgetScreenLayer.cpp
+++ b/Engine/Source/Runtime/UMG/Private/Slate/SWorldWidgetScreenLayer.cpp
@@ -8,6 +8,17 @@
 #include "Widgets/SViewport.h"
 #include "Slate/SGameLayerManager.h"
 
+SWorldWidgetScreenLayer::FComponentEntry::FComponentEntry()
+	: Slot(nullptr)
+{
+}
+
+SWorldWidgetScreenLayer::FComponentEntry::~FComponentEntry()
+{
+	Widget.Reset();
+	ContainerWidget.Reset();
+}
+
 void SWorldWidgetScreenLayer::Construct(const FArguments& InArgs, const FLocalPlayerContext& InPlayerContext)
 {
 	PlayerContext = InPlayerContext;
@@ -32,11 +43,11 @@ void SWorldWidgetScreenLayer::SetWidgetPivot(FVector2D InPivot)
 	Pivot = InPivot;
 }
 
-void SWorldWidgetScreenLayer::AddComponent(USceneComponent* Component, TSharedPtr<SWidget> Widget)
+void SWorldWidgetScreenLayer::AddComponent(USceneComponent* Component, TSharedRef<SWidget> Widget)
 {
-	if ( Component && Widget.IsValid() )
+	if ( Component )
 	{
-		FComponentEntry& Entry = ComponentMap.FindOrAdd(Component);
+		FComponentEntry& Entry = ComponentMap.FindOrAdd(FObjectKey(Component));
 		Entry.Component = Component;
 		Entry.WidgetComponent = Cast<UWidgetComponent>(Component);
 		Entry.Widget = Widget;
@@ -46,7 +57,7 @@ void SWorldWidgetScreenLayer::AddComponent(USceneComponent* Component, TSharedPt
 		[
 			SAssignNew(Entry.ContainerWidget, SBox)
 			[
-				Widget.ToSharedRef()
+				Widget
 			]
 		];
 	}
@@ -54,17 +65,14 @@ void SWorldWidgetScreenLayer::AddComponent(USceneComponent* Component, TSharedPt
 
 void SWorldWidgetScreenLayer::RemoveComponent(USceneComponent* Component)
 {
-	if ( Component )
+	if (ensure(Component))
 	{
-		if ( FComponentEntry* EntryPtr = ComponentMap.Find(Component) )
+		if (FComponentEntry* EntryPtr = ComponentMap.Find(Component))
 		{
-			// First remove the component entry, in case someone tries to double delete.
-			FComponentEntry Entry = *EntryPtr;
-			ComponentMap.Remove(Component);
-
-			if (TSharedPtr<SWidget> ContainerWidget = Entry.ContainerWidget)
+			if (!EntryPtr->bRemoving)
 			{
-				Canvas->RemoveSlot(ContainerWidget.ToSharedRef());
+				RemoveEntryFromCanvas(*EntryPtr);
+				ComponentMap.Remove(Component);
 			}
 		}
 	}
@@ -74,8 +82,6 @@ void SWorldWidgetScreenLayer::Tick(const FGeometry& AllottedGeometry, const doub
 {
 	QUICK_SCOPE_CYCLE_COUNTER(SWorldWidgetScreenLayer_Tick);
 
-	TArray<USceneComponent*, TInlineAllocator<1>> DeadComponents;
-
 	if ( APlayerController* PlayerController = PlayerContext.GetPlayerController() )
 	{
 		if ( UGameViewportClient* ViewportClient = PlayerController->GetWorld()->GetGameViewport() )
@@ -90,6 +96,8 @@ void SWorldWidgetScreenLayer::Tick(const FGeometry& AllottedGeometry, const doub
 				{
 					FVector WorldLocation = SceneComponent->GetComponentLocation();
 
+					//TODO NDarnell Perf This can be improved, if we get the projection matrix and do all of them in one go, instead of calling the library.
+
 					FVector ViewportPosition;
 					const bool bProjected = UWidgetLayoutLibrary::ProjectWorldLocationToWidgetPositionWithDistance(PlayerController, WorldLocation, ViewportPosition);
 
@@ -130,16 +138,22 @@ void SWorldWidgetScreenLayer::Tick(const FGeometry& AllottedGeometry, const doub
 				}
 				else
 				{
-					DeadComponents.Add(SceneComponent);
+					RemoveEntryFromCanvas(Entry);
+					It.RemoveCurrent();
+					continue;
 				}
 			}
 		}
 	}
+}
 
-	// Normally components should be removed by someone calling remove component, but just in case it was 
-	// deleted in a way where they didn't happen, this is our backup solution to ensure we remove stale widgets.
-	for ( int32 Index = 0; Index < DeadComponents.Num(); Index++ )
+void SWorldWidgetScreenLayer::RemoveEntryFromCanvas(SWorldWidgetScreenLayer::FComponentEntry& Entry)
+{
+	// Mark the component was being removed, so we ignore any other remove requests for this component.
+	Entry.bRemoving = true;
+
+	if (TSharedPtr<SWidget> ContainerWidget = Entry.ContainerWidget)
 	{
-		RemoveComponent(DeadComponents[Index]);
+		Canvas->RemoveSlot(ContainerWidget.ToSharedRef());
 	}
-}
+}
\ No newline at end of file
diff --git a/Engine/Source/Runtime/UMG/Public/Slate/SWorldWidgetScreenLayer.h b/Engine/Source/Runtime/UMG/Public/Slate/SWorldWidgetScreenLayer.h
--- a/Engine/Source/Runtime/UMG/Public/Slate/SWorldWidgetScreenLayer.h
+++ b/Engine/Source/Runtime/UMG/Public/Slate/SWorldWidgetScreenLayer.h
@@ -8,6 +8,7 @@
 #include "Engine/LocalPlayer.h"
 #include "Widgets/SCompoundWidget.h"
 #include "Widgets/Layout/SConstraintCanvas.h"
+#include "UObject/ObjectKey.h"
 
 class USceneComponent;
 
@@ -28,7 +29,7 @@ public:
 
 	void SetWidgetPivot(FVector2D Pivot);
 
-	void AddComponent(USceneComponent* Component, TSharedPtr<SWidget> Widget);
+	void AddComponent(USceneComponent* Component, TSharedRef<SWidget> Widget);
 
 	void RemoveComponent(USceneComponent* Component);
 
@@ -41,13 +42,12 @@ private:
 	class FComponentEntry
 	{
 	public:
-		FComponentEntry()
-			: Slot(nullptr)
-		{
-		}
+		FComponentEntry();
+		~FComponentEntry();
 
 	public:
 
+		bool bRemoving = false;
 		TWeakObjectPtr<USceneComponent> Component;
 		class UWidgetComponent* WidgetComponent;
 
@@ -56,6 +56,8 @@ private:
 		SConstraintCanvas::FSlot* Slot;
 	};
 
-	TMap<TWeakObjectPtr<USceneComponent>, FComponentEntry> ComponentMap;
+	void RemoveEntryFromCanvas(SWorldWidgetScreenLayer::FComponentEntry& Entry);
+
+	TMap<FObjectKey, FComponentEntry> ComponentMap;
 	TSharedPtr<SConstraintCanvas> Canvas;
 };
