Fix for a regression that caused nativized widget Blueprint class property delegates to no longer be invoked at runtime.

Also addresses a non-regression issue with nativized child widgets that utilize the WidgetTree instance from a parent class.

diff --git a/Engine/Source/Developer/BlueprintCompilerCppBackend/Private/BlueprintCompilerCppBackendUMG.cpp b/Engine/Source/Developer/BlueprintCompilerCppBackend/Private/BlueprintCompilerCppBackendUMG.cpp
--- a/Engine/Source/Developer/BlueprintCompilerCppBackend/Private/BlueprintCompilerCppBackendUMG.cpp
+++ b/Engine/Source/Developer/BlueprintCompilerCppBackend/Private/BlueprintCompilerCppBackendUMG.cpp
@@ -34,29 +34,41 @@ void FBackendHelperUMG::AdditionalHeaderIncludeForWidget(FEmitterLocalContext& C
 
 void FBackendHelperUMG::CreateClassSubobjects(FEmitterLocalContext& Context, bool bCreate, bool bInitialize)
 {
-	if (auto WidgetClass = Cast<UWidgetBlueprintGeneratedClass>(Context.GetCurrentlyGeneratedClass()))
+	if (UWidgetBlueprintGeneratedClass* WidgetClass = Cast<UWidgetBlueprintGeneratedClass>(Context.GetCurrentlyGeneratedClass()))
 	{
-		if (WidgetClass->WidgetTree)
-		{
-			ensure(WidgetClass->WidgetTree->GetOuter() == Context.GetCurrentlyGeneratedClass());
-			FEmitDefaultValueHelper::HandleClassSubobject(Context, WidgetClass->WidgetTree, FEmitterLocalContext::EClassSubobjectList::MiscConvertedSubobjects, bCreate, bInitialize);
-		}
-		for (auto Anim : WidgetClass->Animations)
+		// Currently nativization does not support widget templates. This method will need to be revised if that changes.
+		check(!WidgetClass->HasTemplate());
+		
+		// Child widgets may actually use the widget tree from a parent class.
+		// - @see UUserWidget::Initialize()
+		WidgetClass = WidgetClass->FindWidgetTreeOwningClass();
+
+		// Initialize the WidgetTree only if it's owned by the current widget class.
+		if (WidgetClass == Context.GetCurrentlyGeneratedClass())
 		{
-			ensure(Anim->GetOuter() == Context.GetCurrentlyGeneratedClass());
+			if (WidgetClass->WidgetTree)
+			{
+				ensure(WidgetClass->WidgetTree->GetOuter() == Context.GetCurrentlyGeneratedClass());
+				FEmitDefaultValueHelper::HandleClassSubobject(Context, WidgetClass->WidgetTree, FEmitterLocalContext::EClassSubobjectList::MiscConvertedSubobjects, bCreate, bInitialize);
+			}
 
-			// We need the same regeneration like for cooking. See UMovieSceneSequence::Serialize
-			FMovieSceneSequencePrecompiledTemplateStore Store;
-			FMovieSceneCompiler::Compile(*Anim, Store);
+			for (UWidgetAnimation* Anim : WidgetClass->Animations)
+			{
+				ensure(Anim->GetOuter() == Context.GetCurrentlyGeneratedClass());
 
-			FEmitDefaultValueHelper::HandleClassSubobject(Context, Anim, FEmitterLocalContext::EClassSubobjectList::MiscConvertedSubobjects, bCreate, bInitialize);
+				// We need the same regeneration like for cooking. See UMovieSceneSequence::Serialize
+				FMovieSceneSequencePrecompiledTemplateStore Store;
+				FMovieSceneCompiler::Compile(*Anim, Store);
+
+				FEmitDefaultValueHelper::HandleClassSubobject(Context, Anim, FEmitterLocalContext::EClassSubobjectList::MiscConvertedSubobjects, bCreate, bInitialize);
+			}
 		}
 	}
 }
 
 void FBackendHelperUMG::EmitWidgetInitializationFunctions(FEmitterLocalContext& Context)
 {
-	if (auto WidgetClass = Cast<UWidgetBlueprintGeneratedClass>(Context.GetCurrentlyGeneratedClass()))
+	if (UWidgetBlueprintGeneratedClass* WidgetClass = Cast<UWidgetBlueprintGeneratedClass>(Context.GetCurrentlyGeneratedClass()))
 	{
 		Context.ResetPropertiesForInaccessibleStructs();
 
@@ -92,18 +104,56 @@ void FBackendHelperUMG::EmitWidgetInitializationFunctions(FEmitterLocalContext&
 			Context.AddLine(TEXT("{"));
 			Context.IncreaseIndent();
 
-			const FString WidgetTreeStr = Context.FindGloballyMappedObject(WidgetClass->WidgetTree, UWidgetTree::StaticClass(), true);
-			ensure(!WidgetTreeStr.IsEmpty());
-			const FString AnimationsArrayNativeName = GenerateLocalProperty(Context, FindFieldChecked<UArrayProperty>(UWidgetBlueprintGeneratedClass::StaticClass(), TEXT("Animations")), reinterpret_cast<const uint8*>(&WidgetClass->Animations));
-			const FString BindingsArrayNativeName = GenerateLocalProperty(Context, FindFieldChecked<UArrayProperty>(UWidgetBlueprintGeneratedClass::StaticClass(), TEXT("Bindings")), reinterpret_cast<const uint8*>(&WidgetClass->Bindings));
-					
-			Context.AddLine(FString::Printf(TEXT("UWidgetBlueprintGeneratedClass::%s(this, GetClass(), %s, %s, %s, %s, %s);")
-				, GET_FUNCTION_NAME_STRING_CHECKED(UWidgetBlueprintGeneratedClass, InitializeWidgetStatic)
-				, WidgetClass->HasTemplate() ? TEXT("true") : TEXT("false")
-				, WidgetClass->bAllowDynamicCreation ? TEXT("true") : TEXT("false")
-				, *WidgetTreeStr
-				, *AnimationsArrayNativeName
-				, *BindingsArrayNativeName));
+			// Child widgets may actually use the widget tree from a parent class.
+			// - @see UUserWidget::Initialize()
+			UWidgetBlueprintGeneratedClass* WidgetTreeOwningClass = WidgetClass->FindWidgetTreeOwningClass();
+
+			// If we have a valid WidgetTree instance, emit code to initialize the widget using the owning class.
+			if (WidgetTreeOwningClass->WidgetTree != nullptr)
+			{
+				FString WidgetClassStr;
+				FString WidgetTreeStr;
+
+				if (WidgetClass == WidgetTreeOwningClass)
+				{
+					// Simple case - WidgetTree instance is owned by the current class.
+					WidgetClassStr = TEXT("GetClass()");
+
+					// This object was already created as a class-owned subobject and mapped to the 'WidgetTree' value.
+					// - @see CreateClassSubobjects()
+					WidgetTreeStr = Context.FindGloballyMappedObject(WidgetClass->WidgetTree, UWidgetTree::StaticClass());
+				}
+				else
+				{
+					// Emit code to assign the owning class to a local variable.
+					WidgetClassStr = Context.GenerateUniqueLocalName();
+					Context.AddLine(FString::Printf(TEXT("UClass* %s = %s;"),
+						*WidgetClassStr,
+						*Context.FindGloballyMappedObject(WidgetTreeOwningClass, UClass::StaticClass(), true)));
+
+					// Emit code to locate and assign the owning class's WidgetTree instance to a local variable. This will have been created as part of the owning class's ctor, but note
+					// that we have to look it up by name/outer because the converted class is a UDynamicClass and not a UWidgetBlueprintGeneratedClass, so there is no 'WidgetTree' member.
+					WidgetTreeStr = Context.GenerateUniqueLocalName();
+					Context.AddLine(FString::Printf(TEXT("UWidgetTree* %s = CastChecked<UWidgetTree>(StaticFindObjectFast(UWidgetTree::StaticClass(), %s, TEXT(\"WidgetTree\")));"),
+						*WidgetTreeStr,
+						*WidgetClassStr));
+				}
+
+				ensure(!WidgetTreeStr.IsEmpty());
+				ensure(!WidgetClassStr.IsEmpty());
+
+				const FString AnimationsArrayNativeName = GenerateLocalProperty(Context, FindFieldChecked<UArrayProperty>(UWidgetBlueprintGeneratedClass::StaticClass(), TEXT("Animations")), reinterpret_cast<const uint8*>(&WidgetTreeOwningClass->Animations));
+				const FString BindingsArrayNativeName = GenerateLocalProperty(Context, FindFieldChecked<UArrayProperty>(UWidgetBlueprintGeneratedClass::StaticClass(), TEXT("Bindings")), reinterpret_cast<const uint8*>(&WidgetTreeOwningClass->Bindings));
+
+				Context.AddLine(FString::Printf(TEXT("UWidgetBlueprintGeneratedClass::%s(this, %s, %s, %s, %s, %s, %s);")
+					, GET_FUNCTION_NAME_STRING_CHECKED(UWidgetBlueprintGeneratedClass, InitializeWidgetStatic)
+					, *WidgetClassStr
+					, WidgetTreeOwningClass->HasTemplate() ? TEXT("true") : TEXT("false")
+					, WidgetTreeOwningClass->bAllowDynamicCreation ? TEXT("true") : TEXT("false")
+					, *WidgetTreeStr
+					, *AnimationsArrayNativeName
+					, *BindingsArrayNativeName));
+			}
 
 			Context.DecreaseIndent();
 			Context.AddLine(TEXT("}"));
diff --git a/Engine/Source/Runtime/UMG/Private/UserWidget.cpp b/Engine/Source/Runtime/UMG/Private/UserWidget.cpp
--- a/Engine/Source/Runtime/UMG/Private/UserWidget.cpp
+++ b/Engine/Source/Runtime/UMG/Private/UserWidget.cpp
@@ -90,36 +90,13 @@ UUserWidget::UUserWidget(const FObjectInitializer& ObjectInitializer)
 
 UWidgetBlueprintGeneratedClass* UUserWidget::GetWidgetTreeOwningClass()
 {
-	UWidgetBlueprintGeneratedClass* RootBGClass = Cast<UWidgetBlueprintGeneratedClass>(GetClass());
-	UWidgetBlueprintGeneratedClass* BGClass = RootBGClass;
-
-	while ( BGClass )
+	UWidgetBlueprintGeneratedClass* WidgetClass = Cast<UWidgetBlueprintGeneratedClass>(GetClass());
+	if (WidgetClass != nullptr)
 	{
-		//TODO NickD: This conditional post load shouldn't be needed any more once the Fast Widget creation path is the only path!
-		// Force post load on the generated class so all subobjects are done (specifically the widget tree).
-		BGClass->ConditionalPostLoad();
-
-		const bool bNoRootWidget = ( nullptr == BGClass->WidgetTree ) || ( nullptr == BGClass->WidgetTree->RootWidget );
-
-		if ( bNoRootWidget )
-		{
-			UWidgetBlueprintGeneratedClass* SuperBGClass = Cast<UWidgetBlueprintGeneratedClass>(BGClass->GetSuperClass());
-			if ( SuperBGClass )
-			{
-				BGClass = SuperBGClass;
-				continue;
-			}
-			else
-			{
-				// If we reach a super class that isn't a UWidgetBlueprintGeneratedClass, return the root class.
-				return RootBGClass;
-			}
-		}
-
-		return BGClass;
+		WidgetClass = WidgetClass->FindWidgetTreeOwningClass();
 	}
 
-	return nullptr;
+	return WidgetClass;
 }
 
 void UUserWidget::TemplateInit()
diff --git a/Engine/Source/Runtime/UMG/Private/WidgetBlueprintGeneratedClass.cpp b/Engine/Source/Runtime/UMG/Private/WidgetBlueprintGeneratedClass.cpp
--- a/Engine/Source/Runtime/UMG/Private/WidgetBlueprintGeneratedClass.cpp
+++ b/Engine/Source/Runtime/UMG/Private/WidgetBlueprintGeneratedClass.cpp
@@ -133,53 +133,51 @@ void UWidgetBlueprintGeneratedClass::InitializeBindingsStatic(UUserWidget* UserW
 {
 	ensure(!UserWidget->HasAnyFlags(RF_ArchetypeObject));
 
-	// Iterate all generated classes in the widget's class hierarchy and initialize bindings found on each one.
-	UClass* CurrentClass = UserWidget->GetClass();
-	while (UWidgetBlueprintGeneratedClass* WBPGC = Cast<UWidgetBlueprintGeneratedClass>(CurrentClass))
+	// Note: It's not safe to assume here that the UserWidget class type is a UWidgetBlueprintGeneratedClass!
+	// - @see InitializeWidgetStatic()
+
+	// For each property binding that we're given, find the corresponding field, and setup the delegate binding on the widget.
+	for (const FDelegateRuntimeBinding& Binding : InBindings)
 	{
-		// For each property binding that we're given, find the corresponding field, and setup the delegate binding on the widget.
-		for (const FDelegateRuntimeBinding& Binding : WBPGC->Bindings)
+		// If the binding came from a parent class, this will still find it - FindField() searches the super class hierarchy by default.
+		UObjectProperty* WidgetProperty = FindField<UObjectProperty>(UserWidget->GetClass(), *Binding.ObjectName);
+		if (WidgetProperty == nullptr)
+		{
+			continue;
+		}
+
+		UWidget* Widget = Cast<UWidget>(WidgetProperty->GetObjectPropertyValue_InContainer(UserWidget));
+
+		if (Widget)
 		{
-			UObjectProperty* WidgetProperty = FindField<UObjectProperty>(WBPGC, *Binding.ObjectName);
-			if (WidgetProperty == nullptr)
+			UDelegateProperty* DelegateProperty = FindField<UDelegateProperty>(Widget->GetClass(), FName(*(Binding.PropertyName.ToString() + TEXT("Delegate"))));
+			if (!DelegateProperty)
 			{
-				continue;
+				DelegateProperty = FindField<UDelegateProperty>(Widget->GetClass(), Binding.PropertyName);
 			}
 
-			UWidget* Widget = Cast<UWidget>(WidgetProperty->GetObjectPropertyValue_InContainer(UserWidget));
-
-			if (Widget)
+			if (DelegateProperty)
 			{
-				UDelegateProperty* DelegateProperty = FindField<UDelegateProperty>(Widget->GetClass(), FName(*(Binding.PropertyName.ToString() + TEXT("Delegate"))));
-				if (!DelegateProperty)
+				bool bSourcePathBound = false;
+
+				if (Binding.SourcePath.IsValid())
 				{
-					DelegateProperty = FindField<UDelegateProperty>(Widget->GetClass(), Binding.PropertyName);
+					bSourcePathBound = Widget->AddBinding(DelegateProperty, UserWidget, Binding.SourcePath);
 				}
 
-				if (DelegateProperty)
+				// If no native binder is found then the only possibility is that the binding is for
+				// a delegate that doesn't match the known native binders available and so we
+				// fallback to just attempting to bind to the function directly.
+				if (bSourcePathBound == false)
 				{
-					bool bSourcePathBound = false;
-
-					if (Binding.SourcePath.IsValid())
+					FScriptDelegate* ScriptDelegate = DelegateProperty->GetPropertyValuePtr_InContainer(Widget);
+					if (ScriptDelegate)
 					{
-						bSourcePathBound = Widget->AddBinding(DelegateProperty, UserWidget, Binding.SourcePath);
-					}
-
-					// If no native binder is found then the only possibility is that the binding is for
-					// a delegate that doesn't match the known native binders available and so we
-					// fallback to just attempting to bind to the function directly.
-					if (bSourcePathBound == false)
-					{
-						FScriptDelegate* ScriptDelegate = DelegateProperty->GetPropertyValuePtr_InContainer(Widget);
-						if (ScriptDelegate)
-						{
-							ScriptDelegate->BindUFunction(UserWidget, Binding.FunctionName);
-						}
+						ScriptDelegate->BindUFunction(UserWidget, Binding.FunctionName);
 					}
 				}
 			}
 		}
-		CurrentClass = CurrentClass->GetSuperClass();
 	}
 }
 
@@ -193,6 +191,10 @@ void UWidgetBlueprintGeneratedClass::InitializeWidgetStatic(UUserWidget* UserWid
 {
 	check(InClass);
 
+	// Note: It's not safe to assume here that the UserWidget class type is a UWidgetBlueprintGeneratedClass! In the case of a nativized widget
+	// blueprint class, it will be a UDynamicClass instead, and this API will be invoked by the blueprint's C++ code that's generated at cook time.
+	// - @see FBackendHelperUMG::EmitWidgetInitializationFunctions()
+
 	if ( UserWidget->HasAllFlags(RF_ArchetypeObject) )
 	{
 		UE_LOG(LogUMG, Error, TEXT("Widget Class %s - Running Initialize On Archetype, %s."), *InClass->GetName(), *UserWidget->GetName());
@@ -316,25 +318,20 @@ void UWidgetBlueprintGeneratedClass::InitializeWidgetStatic(UUserWidget* UserWid
 
 void UWidgetBlueprintGeneratedClass::BindAnimations(UUserWidget* Instance, const TArray< UWidgetAnimation* >& InAnimations)
 {
-	// Iterate all generated classes in the widget's class hierarchy and initialize each animation property with the animation in the class.
+	// Note: It's not safe to assume here that the UserWidget class type is a UWidgetBlueprintGeneratedClass!
+	// - @see InitializeWidgetStatic()
 
-	// @todo: this will not work when nativizing widget blueprints if support for them is ever added - a more complete solution will be needed in that case
-	UClass* CurrentClass = Instance->GetClass();
-	while (UWidgetBlueprintGeneratedClass* WBPGC = Cast<UWidgetBlueprintGeneratedClass>(CurrentClass))
+	for (UWidgetAnimation* Animation : InAnimations)
 	{
-		for (UWidgetAnimation* Animation : WBPGC->Animations)
+		if (Animation->GetMovieScene())
 		{
-			if (Animation->GetMovieScene())
+			// Find property with the same name as the animation and assign the animation to it.
+			UObjectPropertyBase* Prop = FindField<UObjectPropertyBase>(Instance->GetClass(), Animation->GetMovieScene()->GetFName());
+			if (Prop)
 			{
-				// Find property with the same name as the animation and assign the animation to it.
-				UObjectPropertyBase* Prop = FindField<UObjectPropertyBase>(WBPGC, Animation->GetMovieScene()->GetFName());
-				if (Prop)
-				{
-					Prop->SetObjectPropertyValue_InContainer(Instance, Animation);
-				}
+				Prop->SetObjectPropertyValue_InContainer(Instance, Animation);
 			}
 		}
-		CurrentClass = CurrentClass->GetSuperClass();
 	}
 }
 
@@ -347,7 +344,26 @@ void UWidgetBlueprintGeneratedClass::SetClassRequiresNativeTick(bool InClassRequ
 
 void UWidgetBlueprintGeneratedClass::InitializeWidget(UUserWidget* UserWidget) const
 {
-	InitializeWidgetStatic(UserWidget, this, HasTemplate(), bAllowDynamicCreation, WidgetTree, Animations, Bindings);
+	TArray<UWidgetAnimation*> AllAnims;
+	TArray<FDelegateRuntimeBinding> AllBindings;
+
+	// Include current class animations.
+	AllAnims.Append(Animations);
+
+	// Include current class bindings.
+	AllBindings.Append(Bindings);
+
+	// Iterate all generated classes in the widget's parent class hierarchy and include animations and bindings found on each one.
+	UClass* SuperClass = GetSuperClass();
+	while (UWidgetBlueprintGeneratedClass* WBPGC = Cast<UWidgetBlueprintGeneratedClass>(SuperClass))
+	{
+		AllAnims.Append(WBPGC->Animations);
+		AllBindings.Append(WBPGC->Bindings);
+
+		SuperClass = SuperClass->GetSuperClass();
+	}
+
+	InitializeWidgetStatic(UserWidget, this, HasTemplate(), bAllowDynamicCreation, WidgetTree, AllAnims, AllBindings);
 }
 
 void UWidgetBlueprintGeneratedClass::PostLoad()
@@ -626,4 +642,38 @@ void UWidgetBlueprintGeneratedClass::InitializeTemplate(const ITargetPlatform* T
 #endif
 }
 
+UWidgetBlueprintGeneratedClass* UWidgetBlueprintGeneratedClass::FindWidgetTreeOwningClass()
+{
+	UWidgetBlueprintGeneratedClass* RootBGClass = this;
+	UWidgetBlueprintGeneratedClass* BGClass = RootBGClass;
+
+	while (BGClass)
+	{
+		//TODO NickD: This conditional post load shouldn't be needed any more once the Fast Widget creation path is the only path!
+		// Force post load on the generated class so all subobjects are done (specifically the widget tree).
+		BGClass->ConditionalPostLoad();
+
+		const bool bNoRootWidget = (nullptr == BGClass->WidgetTree) || (nullptr == BGClass->WidgetTree->RootWidget);
+
+		if (bNoRootWidget)
+		{
+			UWidgetBlueprintGeneratedClass* SuperBGClass = Cast<UWidgetBlueprintGeneratedClass>(BGClass->GetSuperClass());
+			if (SuperBGClass)
+			{
+				BGClass = SuperBGClass;
+				continue;
+			}
+			else
+			{
+				// If we reach a super class that isn't a UWidgetBlueprintGeneratedClass, return the root class.
+				return RootBGClass;
+			}
+		}
+
+		return BGClass;
+	}
+
+	return nullptr;
+}
+
 #undef LOCTEXT_NAMESPACE
diff --git a/Engine/Source/Runtime/UMG/Public/Blueprint/WidgetBlueprintGeneratedClass.h b/Engine/Source/Runtime/UMG/Public/Blueprint/WidgetBlueprintGeneratedClass.h
--- a/Engine/Source/Runtime/UMG/Public/Blueprint/WidgetBlueprintGeneratedClass.h
+++ b/Engine/Source/Runtime/UMG/Public/Blueprint/WidgetBlueprintGeneratedClass.h
@@ -120,6 +120,9 @@ public:
 	void SetTemplate(UUserWidget* InTemplate);
 	UUserWidget* GetTemplate();
 
+	// Walks up the hierarchy looking for a valid widget tree.
+	UWidgetBlueprintGeneratedClass* FindWidgetTreeOwningClass();
+
 	// UObject interface
 	virtual void PreSave(const class ITargetPlatform* TargetPlatform) override;
 	virtual void Serialize(FArchive& Ar) override;
