Removing all instances of ::FArguments from UWidget constructors.  They're unsafe to use in conditions where we could be on the async loading thread initializing widgets with the EDL, since they may increment the SNullContentWidget ptr if the FArguments has a content slot, also was a bit wasteful.

diff --git a/Engine/Source/Runtime/UMG/Private/Components/Button.cpp b/Engine/Source/Runtime/UMG/Private/Components/Button.cpp
--- a/Engine/Source/Runtime/UMG/Private/Components/Button.cpp
+++ b/Engine/Source/Runtime/UMG/Private/Components/Button.cpp
@@ -14,8 +14,8 @@
 UButton::UButton(const FObjectInitializer& ObjectInitializer)
 	: Super(ObjectInitializer)
 {
-	SButton::FArguments ButtonDefaults;
-	WidgetStyle = *ButtonDefaults._ButtonStyle;
+	// HACK: THIS SHOULD NOT COME FROM CORESTYLE AND SHOULD INSTEAD BY DEFINED BY ENGINE TEXTURES/PROJECT SETTINGS
+	WidgetStyle = FCoreStyle::Get().GetWidgetStyle< FButtonStyle >("Button");
 
 	ColorAndOpacity = FLinearColor::White;
 	BackgroundColor = FLinearColor::White;
diff --git a/Engine/Source/Runtime/UMG/Private/Components/CanvasPanel.cpp b/Engine/Source/Runtime/UMG/Private/Components/CanvasPanel.cpp
--- a/Engine/Source/Runtime/UMG/Private/Components/CanvasPanel.cpp
+++ b/Engine/Source/Runtime/UMG/Private/Components/CanvasPanel.cpp
@@ -13,9 +13,7 @@ UCanvasPanel::UCanvasPanel(const FObjectInitializer& ObjectInitializer)
 	: Super(ObjectInitializer)
 {
 	bIsVariable = false;
-
-	SConstraintCanvas::FArguments Defaults;
-	Visibility = UWidget::ConvertRuntimeToSerializedVisibility(Defaults._Visibility.Get());
+	Visibility = ESlateVisibility::SelfHitTestInvisible;
 }
 
 void UCanvasPanel::ReleaseSlateResources(bool bReleaseChildren)
diff --git a/Engine/Source/Runtime/UMG/Private/Components/CheckBox.cpp b/Engine/Source/Runtime/UMG/Private/Components/CheckBox.cpp
--- a/Engine/Source/Runtime/UMG/Private/Components/CheckBox.cpp
+++ b/Engine/Source/Runtime/UMG/Private/Components/CheckBox.cpp
@@ -14,18 +14,17 @@
 UCheckBox::UCheckBox(const FObjectInitializer& ObjectInitializer)
 	: Super(ObjectInitializer)
 {
-	SCheckBox::FArguments SlateDefaults;
-	WidgetStyle = *SlateDefaults._Style;
+	// HACK: THIS SHOULD NOT COME FROM CORESTYLE AND SHOULD INSTEAD BY DEFINED BY ENGINE TEXTURES/PROJECT SETTINGS
+	WidgetStyle = FCoreStyle::Get().GetWidgetStyle< FCheckBoxStyle >("Checkbox");
 
 	CheckedState = ECheckBoxState::Unchecked;
 
-	HorizontalAlignment = SlateDefaults._HAlign;
-	Padding_DEPRECATED = SlateDefaults._Padding.Get();
+	HorizontalAlignment = HAlign_Fill;
+	Padding_DEPRECATED = FMargin(0, 0, 0, 0);
 
 	BorderBackgroundColor_DEPRECATED = FLinearColor::White;
 
 	IsFocusable = true;
-
 }
 
 void UCheckBox::ReleaseSlateResources(bool bReleaseChildren)
diff --git a/Engine/Source/Runtime/UMG/Private/Components/CircularThrobber.cpp b/Engine/Source/Runtime/UMG/Private/Components/CircularThrobber.cpp
--- a/Engine/Source/Runtime/UMG/Private/Components/CircularThrobber.cpp
+++ b/Engine/Source/Runtime/UMG/Private/Components/CircularThrobber.cpp
@@ -12,15 +12,15 @@
 // UCircularThrobber
 
 UCircularThrobber::UCircularThrobber(const FObjectInitializer& ObjectInitializer)
-	: Super(ObjectInitializer),
-	  bEnableRadius(true)
+	: Super(ObjectInitializer)
+	, bEnableRadius(true)
 {
-	SCircularThrobber::FArguments DefaultArgs;
-	Image = *DefaultArgs._PieceImage;
+	// HACK: THIS SHOULD NOT COME FROM CORESTYLE AND SHOULD INSTEAD BY DEFINED BY ENGINE TEXTURES/PROJECT SETTINGS
+	Image = *FCoreStyle::Get().GetBrush("Throbber.CircleChunk");
 
-	NumberOfPieces = DefaultArgs._NumPieces;
-	Period = DefaultArgs._Period;
-	Radius = DefaultArgs._Radius;
+	NumberOfPieces = 6;
+	Period = 0.75f;
+	Radius = 16.f;
 }
 
 void UCircularThrobber::ReleaseSlateResources(bool bReleaseChildren)
@@ -32,8 +32,6 @@ void UCircularThrobber::ReleaseSlateResources(bool bReleaseChildren)
 
 TSharedRef<SWidget> UCircularThrobber::RebuildWidget()
 {
-	SCircularThrobber::FArguments DefaultArgs;
-
 	MyCircularThrobber = SNew(SCircularThrobber)
 		.PieceImage(&Image)
 		.NumPieces(FMath::Clamp(NumberOfPieces, 1, 25))
diff --git a/Engine/Source/Runtime/UMG/Private/Components/ComboBoxString.cpp b/Engine/Source/Runtime/UMG/Private/Components/ComboBoxString.cpp
--- a/Engine/Source/Runtime/UMG/Private/Components/ComboBoxString.cpp
+++ b/Engine/Source/Runtime/UMG/Private/Components/ComboBoxString.cpp
@@ -13,9 +13,9 @@
 UComboBoxString::UComboBoxString(const FObjectInitializer& ObjectInitializer)
 	: Super(ObjectInitializer)
 {
-	SComboBox< TSharedPtr<FString> >::FArguments SlateDefaults;
-	WidgetStyle = *SlateDefaults._ComboBoxStyle;
-	ItemStyle = *SlateDefaults._ItemStyle;
+	// HACK: THIS SHOULD NOT COME FROM CORESTYLE AND SHOULD INSTEAD BY DEFINED BY ENGINE TEXTURES/PROJECT SETTINGS
+	WidgetStyle = FCoreStyle::Get().GetWidgetStyle< FComboBoxStyle >("ComboBox");
+	ItemStyle = FCoreStyle::Get().GetWidgetStyle< FTableRowStyle >("TableView.Row");
 	ItemStyle.SelectorFocusedBrush.TintColor = ItemStyle.SelectorFocusedBrush.TintColor.GetSpecifiedColor();
 	ItemStyle.ActiveHoveredBrush.TintColor = ItemStyle.ActiveHoveredBrush.TintColor.GetSpecifiedColor();
 	ItemStyle.ActiveBrush.TintColor = ItemStyle.ActiveBrush.TintColor.GetSpecifiedColor();
diff --git a/Engine/Source/Runtime/UMG/Private/Components/EditableText.cpp b/Engine/Source/Runtime/UMG/Private/Components/EditableText.cpp
--- a/Engine/Source/Runtime/UMG/Private/Components/EditableText.cpp
+++ b/Engine/Source/Runtime/UMG/Private/Components/EditableText.cpp
@@ -15,8 +15,8 @@
 UEditableText::UEditableText(const FObjectInitializer& ObjectInitializer)
 	: Super(ObjectInitializer)
 {
-	SEditableText::FArguments Defaults;
-	WidgetStyle = *Defaults._Style;
+	// HACK: THIS SHOULD NOT COME FROM CORESTYLE AND SHOULD INSTEAD BY DEFINED BY ENGINE TEXTURES/PROJECT SETTINGS
+	WidgetStyle = FCoreStyle::Get().GetWidgetStyle< FEditableTextStyle >("NormalEditableText");
 
 	ColorAndOpacity_DEPRECATED = FLinearColor::Black;
 
@@ -26,18 +26,17 @@ UEditableText::UEditableText(const FObjectInitializer& ObjectInitializer)
 		Font_DEPRECATED = FSlateFontInfo(RobotoFontObj.Object, 12, FName("Bold"));
 	}
 
-	// Grab other defaults from slate arguments.
-	IsReadOnly = Defaults._IsReadOnly.Get();
-	IsPassword = Defaults._IsPassword.Get();
-	MinimumDesiredWidth = Defaults._MinDesiredWidth.Get();
-	IsCaretMovedWhenGainFocus = Defaults._IsCaretMovedWhenGainFocus.Get();
-	SelectAllTextWhenFocused = Defaults._SelectAllTextWhenFocused.Get();
-	RevertTextOnEscape = Defaults._RevertTextOnEscape.Get();
-	ClearKeyboardFocusOnCommit = Defaults._ClearKeyboardFocusOnCommit.Get();
-	SelectAllTextOnCommit = Defaults._SelectAllTextOnCommit.Get();
-	AllowContextMenu = Defaults._AllowContextMenu.Get();
-	VirtualKeyboardDismissAction = Defaults._VirtualKeyboardDismissAction.Get();
-	Clipping = Defaults._Clipping;
+	IsReadOnly = false;
+	IsPassword = false;
+	MinimumDesiredWidth = 0.0f;
+	IsCaretMovedWhenGainFocus = true;
+	SelectAllTextWhenFocused = false;
+	RevertTextOnEscape = false;
+	ClearKeyboardFocusOnCommit = true;
+	SelectAllTextOnCommit = false;
+	AllowContextMenu = true;
+	VirtualKeyboardDismissAction = EVirtualKeyboardDismissAction::TextChangeOnDismiss;
+	Clipping = EWidgetClipping::ClipToBounds;
 }
 
 void UEditableText::ReleaseSlateResources(bool bReleaseChildren)
diff --git a/Engine/Source/Runtime/UMG/Private/Components/EditableTextBox.cpp b/Engine/Source/Runtime/UMG/Private/Components/EditableTextBox.cpp
--- a/Engine/Source/Runtime/UMG/Private/Components/EditableTextBox.cpp
+++ b/Engine/Source/Runtime/UMG/Private/Components/EditableTextBox.cpp
@@ -24,21 +24,20 @@ UEditableTextBox::UEditableTextBox(const FObjectInitializer& ObjectInitializer)
 		Font_DEPRECATED = FSlateFontInfo(RobotoFontObj.Object, 12, FName("Bold"));
 	}
 
-	// Grab other defaults from slate arguments.
-	SEditableTextBox::FArguments Defaults;
-	IsReadOnly = Defaults._IsReadOnly.Get();
-	IsPassword = Defaults._IsPassword.Get();
-	MinimumDesiredWidth = Defaults._MinDesiredWidth.Get();
-	Padding_DEPRECATED = Defaults._Padding.Get();
-	IsCaretMovedWhenGainFocus = Defaults._IsCaretMovedWhenGainFocus.Get();
-	SelectAllTextWhenFocused = Defaults._SelectAllTextWhenFocused.Get();
-	RevertTextOnEscape = Defaults._RevertTextOnEscape.Get();
-	ClearKeyboardFocusOnCommit = Defaults._ClearKeyboardFocusOnCommit.Get();
-	SelectAllTextOnCommit = Defaults._SelectAllTextOnCommit.Get();
-	AllowContextMenu = Defaults._AllowContextMenu.Get();
-	VirtualKeyboardDismissAction = Defaults._VirtualKeyboardDismissAction.Get();
-
-	WidgetStyle = *Defaults._Style;
+	IsReadOnly = false;
+	IsPassword = false;
+	MinimumDesiredWidth = 0.0f;
+	Padding_DEPRECATED = FMargin(0, 0, 0, 0);
+	IsCaretMovedWhenGainFocus = true;
+	SelectAllTextWhenFocused = false;
+	RevertTextOnEscape = false;
+	ClearKeyboardFocusOnCommit = true;
+	SelectAllTextOnCommit = false;
+	AllowContextMenu = true;
+	VirtualKeyboardDismissAction = EVirtualKeyboardDismissAction::TextChangeOnDismiss;
+
+	// HACK: THIS SHOULD NOT COME FROM CORESTYLE AND SHOULD INSTEAD BY DEFINED BY ENGINE TEXTURES/PROJECT SETTINGS
+	WidgetStyle = FCoreStyle::Get().GetWidgetStyle< FEditableTextBoxStyle >("NormalEditableTextBox");
 }
 
 void UEditableTextBox::ReleaseSlateResources(bool bReleaseChildren)
diff --git a/Engine/Source/Runtime/UMG/Private/Components/ExpandableArea.cpp b/Engine/Source/Runtime/UMG/Private/Components/ExpandableArea.cpp
--- a/Engine/Source/Runtime/UMG/Private/Components/ExpandableArea.cpp
+++ b/Engine/Source/Runtime/UMG/Private/Components/ExpandableArea.cpp
@@ -20,12 +20,12 @@ UExpandableArea::UExpandableArea(const FObjectInitializer& ObjectInitializer)
 {
 	bIsVariable = true;
 
-	SExpandableArea::FArguments ExpandableDefaults;
-	Style       = *ExpandableDefaults._Style;
-	BorderColor = ExpandableDefaults._BorderBackgroundColor.Get( FLinearColor::White );
-	BorderBrush = *ExpandableDefaults._BorderImage.Get( FStyleDefaults::GetNoBrush() );
-	AreaPadding = ExpandableDefaults._Padding.Get();
-	HeaderPadding = ExpandableDefaults._HeaderPadding.Get();
+	// HACK: THIS SHOULD NOT COME FROM CORESTYLE AND SHOULD INSTEAD BY DEFINED BY ENGINE TEXTURES/PROJECT SETTINGS
+	Style = FCoreStyle::Get().GetWidgetStyle<FExpandableAreaStyle>("ExpandableArea");
+	BorderColor = FLinearColor::White;
+	BorderBrush = *FCoreStyle::Get().GetBrush("ExpandableArea.Border");
+	AreaPadding = FMargin(1);
+	HeaderPadding = FMargin(4.0f, 2.0f);
 }
 
 bool UExpandableArea::GetIsExpanded() const
diff --git a/Engine/Source/Runtime/UMG/Private/Components/GridPanel.cpp b/Engine/Source/Runtime/UMG/Private/Components/GridPanel.cpp
--- a/Engine/Source/Runtime/UMG/Private/Components/GridPanel.cpp
+++ b/Engine/Source/Runtime/UMG/Private/Components/GridPanel.cpp
@@ -14,9 +14,7 @@ UGridPanel::UGridPanel(const FObjectInitializer& ObjectInitializer)
 	: Super(ObjectInitializer)
 {
 	bIsVariable = false;
-
-	SGridPanel::FArguments Defaults;
-	Visibility = UWidget::ConvertRuntimeToSerializedVisibility(Defaults._Visibility.Get());
+	Visibility = ESlateVisibility::SelfHitTestInvisible;
 }
 
 void UGridPanel::ReleaseSlateResources(bool bReleaseChildren)
diff --git a/Engine/Source/Runtime/UMG/Private/Components/HorizontalBox.cpp b/Engine/Source/Runtime/UMG/Private/Components/HorizontalBox.cpp
--- a/Engine/Source/Runtime/UMG/Private/Components/HorizontalBox.cpp
+++ b/Engine/Source/Runtime/UMG/Private/Components/HorizontalBox.cpp
@@ -12,9 +12,7 @@ UHorizontalBox::UHorizontalBox(const FObjectInitializer& ObjectInitializer)
 	: Super(ObjectInitializer)
 {
 	bIsVariable = false;
-
-	SHorizontalBox::FArguments Defaults;
-	Visibility = UWidget::ConvertRuntimeToSerializedVisibility(Defaults._Visibility.Get());
+	Visibility = ESlateVisibility::SelfHitTestInvisible;
 }
 
 void UHorizontalBox::ReleaseSlateResources(bool bReleaseChildren)
diff --git a/Engine/Source/Runtime/UMG/Private/Components/InputKeySelector.cpp b/Engine/Source/Runtime/UMG/Private/Components/InputKeySelector.cpp
--- a/Engine/Source/Runtime/UMG/Private/Components/InputKeySelector.cpp
+++ b/Engine/Source/Runtime/UMG/Private/Components/InputKeySelector.cpp
@@ -13,14 +13,15 @@
 UInputKeySelector::UInputKeySelector( const FObjectInitializer& ObjectInitializer )
 	: Super(ObjectInitializer)
 {
-	SInputKeySelector::FArguments InputKeySelectorDefaults;
-	WidgetStyle = *InputKeySelectorDefaults._ButtonStyle;
-	TextStyle = *InputKeySelectorDefaults._TextStyle;
-	KeySelectionText = InputKeySelectorDefaults._KeySelectionText;
-	NoKeySpecifiedText = InputKeySelectorDefaults._NoKeySpecifiedText;
-	SelectedKey = InputKeySelectorDefaults._SelectedKey.Get();
-	bAllowModifierKeys = InputKeySelectorDefaults._AllowModifierKeys;
-	bAllowGamepadKeys = InputKeySelectorDefaults._AllowGamepadKeys;
+	// HACK: THIS SHOULD NOT COME FROM CORESTYLE AND SHOULD INSTEAD BY DEFINED BY ENGINE TEXTURES/PROJECT SETTINGS
+	WidgetStyle = FCoreStyle::Get().GetWidgetStyle<FButtonStyle>("Button");
+	TextStyle = FCoreStyle::Get().GetWidgetStyle< FTextBlockStyle >("NormalText");
+
+	KeySelectionText = NSLOCTEXT("InputKeySelector", "DefaultKeySelectionText", "...");
+	NoKeySpecifiedText = NSLOCTEXT("InputKeySelector", "DefaultEmptyText", "Empty");
+	SelectedKey = FInputChord(EKeys::Invalid);
+	bAllowModifierKeys = true;
+	bAllowGamepadKeys = false;
 
 	EscapeKeys.AddUnique(EKeys::Gamepad_Special_Right); // In most (if not all) cases this is going to be the menu button
 
diff --git a/Engine/Source/Runtime/UMG/Private/Components/MultiLineEditableText.cpp b/Engine/Source/Runtime/UMG/Private/Components/MultiLineEditableText.cpp
--- a/Engine/Source/Runtime/UMG/Private/Components/MultiLineEditableText.cpp
+++ b/Engine/Source/Runtime/UMG/Private/Components/MultiLineEditableText.cpp
@@ -14,18 +14,19 @@
 UMultiLineEditableText::UMultiLineEditableText(const FObjectInitializer& ObjectInitializer)
 	: Super(ObjectInitializer)
 {
-	SMultiLineEditableText::FArguments Defaults;
-	WidgetStyle = *Defaults._TextStyle;
-	bIsReadOnly = Defaults._IsReadOnly.Get();
-	SelectAllTextWhenFocused = Defaults._SelectAllTextWhenFocused.Get();
-	ClearTextSelectionOnFocusLoss = Defaults._ClearTextSelectionOnFocusLoss.Get();
-	RevertTextOnEscape = Defaults._RevertTextOnEscape.Get();
-	ClearKeyboardFocusOnCommit = Defaults._ClearKeyboardFocusOnCommit.Get();
-	AllowContextMenu = Defaults._AllowContextMenu.Get();
-	Clipping = Defaults._Clipping;
-	VirtualKeyboardDismissAction = Defaults._VirtualKeyboardDismissAction.Get();
-	AutoWrapText = true;
+	// HACK: THIS SHOULD NOT COME FROM CORESTYLE AND SHOULD INSTEAD BY DEFINED BY ENGINE TEXTURES/PROJECT SETTINGS
+	WidgetStyle = FCoreStyle::Get().GetWidgetStyle<FTextBlockStyle>("NormalText");
 	
+	bIsReadOnly = false;
+	SelectAllTextWhenFocused = false;
+	ClearTextSelectionOnFocusLoss = true;
+	RevertTextOnEscape = false;
+	ClearKeyboardFocusOnCommit = true;
+	AllowContextMenu = true;
+	Clipping = EWidgetClipping::ClipToBounds;
+	VirtualKeyboardDismissAction = EVirtualKeyboardDismissAction::TextChangeOnDismiss;
+	AutoWrapText = true;
+
 	if (!IsRunningDedicatedServer())
 	{
 		static ConstructorHelpers::FObjectFinder<UFont> RobotoFontObj(*UWidget::GetDefaultFontName());
diff --git a/Engine/Source/Runtime/UMG/Private/Components/MultiLineEditableTextBox.cpp b/Engine/Source/Runtime/UMG/Private/Components/MultiLineEditableTextBox.cpp
--- a/Engine/Source/Runtime/UMG/Private/Components/MultiLineEditableTextBox.cpp
+++ b/Engine/Source/Runtime/UMG/Private/Components/MultiLineEditableTextBox.cpp
@@ -18,12 +18,13 @@ UMultiLineEditableTextBox::UMultiLineEditableTextBox(const FObjectInitializer& O
 	BackgroundColor_DEPRECATED = FLinearColor::White;
 	ReadOnlyForegroundColor_DEPRECATED = FLinearColor::Black;
 
-	SMultiLineEditableTextBox::FArguments Defaults;
-	bIsReadOnly = Defaults._IsReadOnly.Get();
-	WidgetStyle = *Defaults._Style;
-	TextStyle = *Defaults._TextStyle;
-	AllowContextMenu = Defaults._AllowContextMenu.Get();
-	VirtualKeyboardDismissAction = Defaults._VirtualKeyboardDismissAction.Get();
+	// HACK: THIS SHOULD NOT COME FROM CORESTYLE AND SHOULD INSTEAD BY DEFINED BY ENGINE TEXTURES/PROJECT SETTINGS
+	WidgetStyle = FCoreStyle::Get().GetWidgetStyle<FEditableTextBoxStyle>("NormalEditableTextBox");
+	TextStyle = FCoreStyle::Get().GetWidgetStyle<FTextBlockStyle>("NormalText");
+
+	bIsReadOnly = false;
+	AllowContextMenu = true;
+	VirtualKeyboardDismissAction = EVirtualKeyboardDismissAction::TextChangeOnDismiss;
 	AutoWrapText = true;
 
 	if (!IsRunningDedicatedServer())
diff --git a/Engine/Source/Runtime/UMG/Private/Components/Overlay.cpp b/Engine/Source/Runtime/UMG/Private/Components/Overlay.cpp
--- a/Engine/Source/Runtime/UMG/Private/Components/Overlay.cpp
+++ b/Engine/Source/Runtime/UMG/Private/Components/Overlay.cpp
@@ -12,9 +12,7 @@ UOverlay::UOverlay(const FObjectInitializer& ObjectInitializer)
 	: Super(ObjectInitializer)
 {
 	bIsVariable = false;
-
-	SOverlay::FArguments Defaults;
-	Visibility = UWidget::ConvertRuntimeToSerializedVisibility(Defaults._Visibility.Get());
+	Visibility = ESlateVisibility::SelfHitTestInvisible;
 }
 
 void UOverlay::ReleaseSlateResources(bool bReleaseChildren)
diff --git a/Engine/Source/Runtime/UMG/Private/Components/ProgressBar.cpp b/Engine/Source/Runtime/UMG/Private/Components/ProgressBar.cpp
--- a/Engine/Source/Runtime/UMG/Private/Components/ProgressBar.cpp
+++ b/Engine/Source/Runtime/UMG/Private/Components/ProgressBar.cpp
@@ -11,8 +11,8 @@
 UProgressBar::UProgressBar(const FObjectInitializer& ObjectInitializer)
 	: Super(ObjectInitializer)
 {
-	SProgressBar::FArguments SlateDefaults;
-	WidgetStyle = *SlateDefaults._Style;
+	// HACK: THIS SHOULD NOT COME FROM CORESTYLE AND SHOULD INSTEAD BY DEFINED BY ENGINE TEXTURES/PROJECT SETTINGS
+	WidgetStyle = FCoreStyle::Get().GetWidgetStyle<FProgressBarStyle>("ProgressBar");
 	WidgetStyle.FillImage.TintColor = FLinearColor::White;
 
 	BarFillType = EProgressBarFillType::LeftToRight;
diff --git a/Engine/Source/Runtime/UMG/Private/Components/ScrollBar.cpp b/Engine/Source/Runtime/UMG/Private/Components/ScrollBar.cpp
--- a/Engine/Source/Runtime/UMG/Private/Components/ScrollBar.cpp
+++ b/Engine/Source/Runtime/UMG/Private/Components/ScrollBar.cpp
@@ -17,8 +17,8 @@ UScrollBar::UScrollBar(const FObjectInitializer& ObjectInitializer)
 	Orientation = Orient_Vertical;
 	Thickness = FVector2D(12.0f, 12.0f);
 
-	SScrollBar::FArguments Defaults;
-	WidgetStyle = *Defaults._Style;
+	// HACK: THIS SHOULD NOT COME FROM CORESTYLE AND SHOULD INSTEAD BY DEFINED BY ENGINE TEXTURES/PROJECT SETTINGS
+	WidgetStyle = FCoreStyle::Get().GetWidgetStyle<FScrollBarStyle>("Scrollbar");
 }
 
 void UScrollBar::ReleaseSlateResources(bool bReleaseChildren)
diff --git a/Engine/Source/Runtime/UMG/Private/Components/ScrollBox.cpp b/Engine/Source/Runtime/UMG/Private/Components/ScrollBox.cpp
--- a/Engine/Source/Runtime/UMG/Private/Components/ScrollBox.cpp
+++ b/Engine/Source/Runtime/UMG/Private/Components/ScrollBox.cpp
@@ -23,12 +23,12 @@ UScrollBox::UScrollBox(const FObjectInitializer& ObjectInitializer)
 {
 	bIsVariable = false;
 
-	SScrollBox::FArguments Defaults;
-	Visibility = UWidget::ConvertRuntimeToSerializedVisibility(Defaults._Visibility.Get());
+	Visibility = ESlateVisibility::Visible;
 	Clipping = EWidgetClipping::ClipToBounds;
 
-	WidgetStyle = *Defaults._Style;
-	WidgetBarStyle = *Defaults._ScrollBarStyle;
+	// HACK: THIS SHOULD NOT COME FROM CORESTYLE AND SHOULD INSTEAD BY DEFINED BY ENGINE TEXTURES/PROJECT SETTINGS
+	WidgetStyle = FCoreStyle::Get().GetWidgetStyle<FScrollBoxStyle>("ScrollBox");
+	WidgetBarStyle = FCoreStyle::Get().GetWidgetStyle<FScrollBarStyle>("ScrollBar");
 	bAllowRightClickDragScrolling = true;
 }
 
diff --git a/Engine/Source/Runtime/UMG/Private/Components/SizeBoxSlot.cpp b/Engine/Source/Runtime/UMG/Private/Components/SizeBoxSlot.cpp
--- a/Engine/Source/Runtime/UMG/Private/Components/SizeBoxSlot.cpp
+++ b/Engine/Source/Runtime/UMG/Private/Components/SizeBoxSlot.cpp
@@ -15,8 +15,6 @@ USizeBoxSlot::USizeBoxSlot(const FObjectInitializer& ObjectInitializer)
 
 	HorizontalAlignment = HAlign_Fill;
 	VerticalAlignment = VAlign_Fill;
-
-	SBox::FArguments SizeBoxDefaults;
 }
 
 void USizeBoxSlot::ReleaseSlateResources(bool bReleaseChildren)
diff --git a/Engine/Source/Runtime/UMG/Private/Components/Slider.cpp b/Engine/Source/Runtime/UMG/Private/Components/Slider.cpp
--- a/Engine/Source/Runtime/UMG/Private/Components/Slider.cpp
+++ b/Engine/Source/Runtime/UMG/Private/Components/Slider.cpp
@@ -16,11 +16,12 @@ USlider::USlider(const FObjectInitializer& ObjectInitializer)
 	SliderBarColor = FLinearColor::White;
 	SliderHandleColor = FLinearColor::White;
 	StepSize = 0.01f;
-	SSlider::FArguments Defaults;
-	WidgetStyle = *Defaults._Style;
 	IsFocusable = true;
 	MouseUsesStep = false;
 	RequiresControllerLock = true;
+
+	// HACK: THIS SHOULD NOT COME FROM CORESTYLE AND SHOULD INSTEAD BY DEFINED BY ENGINE TEXTURES/PROJECT SETTINGS
+	WidgetStyle = FCoreStyle::Get().GetWidgetStyle<FSliderStyle>("Slider");
 }
 
 TSharedRef<SWidget> USlider::RebuildWidget()
diff --git a/Engine/Source/Runtime/UMG/Private/Components/SpinBox.cpp b/Engine/Source/Runtime/UMG/Private/Components/SpinBox.cpp
--- a/Engine/Source/Runtime/UMG/Private/Components/SpinBox.cpp
+++ b/Engine/Source/Runtime/UMG/Private/Components/SpinBox.cpp
@@ -18,22 +18,20 @@ USpinBox::USpinBox(const FObjectInitializer& ObjectInitializer)
 		Font = FSlateFontInfo(RobotoFontObj.Object, 12, FName("Bold"));
 	}
 
-	// Grab other defaults from slate arguments.
-	SSpinBox<float>::FArguments Defaults;
-
-	Value = Defaults._Value.Get();
-	MinValue = Defaults._MinValue.Get().Get(0.0f);
-	MaxValue = Defaults._MaxValue.Get().Get(0.0f);
-	MinSliderValue = Defaults._MinSliderValue.Get().Get(0.0f);
-	MaxSliderValue = Defaults._MaxSliderValue.Get().Get(0.0f);
-	Delta = Defaults._Delta.Get();
-	SliderExponent = Defaults._SliderExponent.Get();
-	MinDesiredWidth = Defaults._MinDesiredWidth.Get();
-	ClearKeyboardFocusOnCommit = Defaults._ClearKeyboardFocusOnCommit.Get();
-	SelectAllTextOnCommit = Defaults._SelectAllTextOnCommit.Get();
-
-	WidgetStyle = *Defaults._Style;
+	Value = 0;
+	MinValue = 0;
+	MaxValue = 10;
+	MinSliderValue = 0;
+	MaxSliderValue = 0;
+	Delta = 0;
+	SliderExponent = 1;
+	MinDesiredWidth = 0;
+	ClearKeyboardFocusOnCommit = false;
+	SelectAllTextOnCommit = true;
 	ForegroundColor = FSlateColor(FLinearColor::Black);
+
+	// HACK: THIS SHOULD NOT COME FROM CORESTYLE AND SHOULD INSTEAD BY DEFINED BY ENGINE TEXTURES/PROJECT SETTINGS
+	WidgetStyle = FCoreStyle::Get().GetWidgetStyle<FSpinBoxStyle>("SpinBox");
 }
 
 void USpinBox::ReleaseSlateResources(bool bReleaseChildren)
diff --git a/Engine/Source/Runtime/UMG/Private/Components/Throbber.cpp b/Engine/Source/Runtime/UMG/Private/Components/Throbber.cpp
--- a/Engine/Source/Runtime/UMG/Private/Components/Throbber.cpp
+++ b/Engine/Source/Runtime/UMG/Private/Components/Throbber.cpp
@@ -12,14 +12,14 @@
 UThrobber::UThrobber(const FObjectInitializer& ObjectInitializer)
 	: Super(ObjectInitializer)
 {
-	SThrobber::FArguments DefaultArgs;
-	Image = *DefaultArgs._PieceImage;
+	NumberOfPieces = 3;
 
-	NumberOfPieces = DefaultArgs._NumPieces;
+	bAnimateVertically = true;
+	bAnimateHorizontally = true;
+	bAnimateOpacity = true;
 
-	bAnimateVertically = (DefaultArgs._Animate & SThrobber::Vertical) != 0;
-	bAnimateHorizontally = (DefaultArgs._Animate & SThrobber::Horizontal) != 0;
-	bAnimateOpacity = (DefaultArgs._Animate & SThrobber::Opacity) != 0;
+	// HACK: THIS SHOULD NOT COME FROM CORESTYLE AND SHOULD INSTEAD BY DEFINED BY ENGINE TEXTURES/PROJECT SETTINGS
+	Image = *FCoreStyle::Get().GetBrush("Throbber.Chunk");
 }
 
 void UThrobber::ReleaseSlateResources(bool bReleaseChildren)
@@ -31,8 +31,6 @@ void UThrobber::ReleaseSlateResources(bool bReleaseChildren)
 
 TSharedRef<SWidget> UThrobber::RebuildWidget()
 {
-	SThrobber::FArguments DefaultArgs;
-
 	MyThrobber = SNew(SThrobber)
 		.PieceImage(&Image)
 		.NumPieces(FMath::Clamp(NumberOfPieces, 1, 25))
diff --git a/Engine/Source/Runtime/UMG/Private/Components/UniformGridPanel.cpp b/Engine/Source/Runtime/UMG/Private/Components/UniformGridPanel.cpp
--- a/Engine/Source/Runtime/UMG/Private/Components/UniformGridPanel.cpp
+++ b/Engine/Source/Runtime/UMG/Private/Components/UniformGridPanel.cpp
@@ -14,9 +14,7 @@ UUniformGridPanel::UUniformGridPanel(const FObjectInitializer& ObjectInitializer
 	: Super(ObjectInitializer)
 {
 	bIsVariable = false;
-
-	SUniformGridPanel::FArguments Defaults;
-	Visibility = UWidget::ConvertRuntimeToSerializedVisibility(Defaults._Visibility.Get());
+	Visibility = ESlateVisibility::SelfHitTestInvisible;
 }
 
 void UUniformGridPanel::ReleaseSlateResources(bool bReleaseChildren)
diff --git a/Engine/Source/Runtime/UMG/Private/Components/VerticalBox.cpp b/Engine/Source/Runtime/UMG/Private/Components/VerticalBox.cpp
--- a/Engine/Source/Runtime/UMG/Private/Components/VerticalBox.cpp
+++ b/Engine/Source/Runtime/UMG/Private/Components/VerticalBox.cpp
@@ -12,9 +12,7 @@ UVerticalBox::UVerticalBox(const FObjectInitializer& ObjectInitializer)
 	: Super(ObjectInitializer)
 {
 	bIsVariable = false;
-
-	SVerticalBox::FArguments Defaults;
-	Visibility = UWidget::ConvertRuntimeToSerializedVisibility(Defaults._Visibility.Get());
+	Visibility = ESlateVisibility::SelfHitTestInvisible;
 }
 
 void UVerticalBox::ReleaseSlateResources(bool bReleaseChildren)
diff --git a/Engine/Source/Runtime/UMG/Private/Components/WidgetSwitcher.cpp b/Engine/Source/Runtime/UMG/Private/Components/WidgetSwitcher.cpp
--- a/Engine/Source/Runtime/UMG/Private/Components/WidgetSwitcher.cpp
+++ b/Engine/Source/Runtime/UMG/Private/Components/WidgetSwitcher.cpp
@@ -13,9 +13,7 @@ UWidgetSwitcher::UWidgetSwitcher(const FObjectInitializer& ObjectInitializer)
 	: Super(ObjectInitializer)
 {
 	bIsVariable = true;
-
-	SWidgetSwitcher::FArguments Defaults;
-	Visibility = UWidget::ConvertRuntimeToSerializedVisibility(Defaults._Visibility.Get());
+	Visibility = ESlateVisibility::SelfHitTestInvisible;
 }
 
 void UWidgetSwitcher::ReleaseSlateResources(bool bReleaseChildren)
diff --git a/Engine/Source/Runtime/UMG/Private/Components/WindowTitleBarArea.cpp b/Engine/Source/Runtime/UMG/Private/Components/WindowTitleBarArea.cpp
--- a/Engine/Source/Runtime/UMG/Private/Components/WindowTitleBarArea.cpp
+++ b/Engine/Source/Runtime/UMG/Private/Components/WindowTitleBarArea.cpp
@@ -19,9 +19,7 @@ UWindowTitleBarArea::UWindowTitleBarArea(const FObjectInitializer& ObjectInitial
 	: Super(ObjectInitializer)
 {
 	bIsVariable = false;
-
-	SWindowTitleBarArea::FArguments Defaults;
-	Visibility = UWidget::ConvertRuntimeToSerializedVisibility(Defaults._Visibility.Get());
+	Visibility = ESlateVisibility::Visible;
 
 	bDoubleClickTogglesFullscreen = false;
 }
diff --git a/Engine/Source/Runtime/UMG/Private/Components/WrapBox.cpp b/Engine/Source/Runtime/UMG/Private/Components/WrapBox.cpp
--- a/Engine/Source/Runtime/UMG/Private/Components/WrapBox.cpp
+++ b/Engine/Source/Runtime/UMG/Private/Components/WrapBox.cpp
@@ -12,10 +12,7 @@ UWrapBox::UWrapBox(const FObjectInitializer& ObjectInitializer)
 	: Super(ObjectInitializer)
 {
 	bIsVariable = false;
-
-	SWrapBox::FArguments Defaults;
-	Visibility = UWidget::ConvertRuntimeToSerializedVisibility(Defaults._Visibility.Get());
-
+	Visibility = ESlateVisibility::SelfHitTestInvisible;
 	WrapWidth = 500;
 	bExplicitWrapWidth = false;
 }
