Split DynamicEntryBox out into Base and Derived. Base has all addition/construction Protected. Derived exposes these.
Allows for the cration of DynamicEntryBox widgets which control their own entries, and don't allow other widgets/classes to expand them.

diff --git a/Engine/Source/Editor/UMGEditor/Private/Customizations/DynamicEntryBoxDetails.cpp b/Engine/Source/Editor/UMGEditor/Private/Customizations/DynamicEntryBoxDetails.cpp
--- a/Engine/Source/Editor/UMGEditor/Private/Customizations/DynamicEntryBoxDetails.cpp
+++ b/Engine/Source/Editor/UMGEditor/Private/Customizations/DynamicEntryBoxDetails.cpp
@@ -7,12 +7,16 @@
 #include "Components/DynamicEntryBox.h"
 #include "PropertyCustomizationHelpers.h"
 
-TSharedRef<IDetailCustomization> FDynamicEntryBoxDetails::MakeInstance()
+//////////////////////////////////////////////////////////////////////////
+// FDynamicEntryBoxBaseDetails
+//////////////////////////////////////////////////////////////////////////
+
+TSharedRef<IDetailCustomization> FDynamicEntryBoxBaseDetails::MakeInstance()
 {
-	return MakeShareable(new FDynamicEntryBoxDetails());
+	return MakeShareable(new FDynamicEntryBoxBaseDetails());
 }
 
-void FDynamicEntryBoxDetails::CustomizeDetails(IDetailLayoutBuilder& DetailLayout)
+void FDynamicEntryBoxBaseDetails::CustomizeDetails(IDetailLayoutBuilder& DetailLayout)
 {
 	TArray<TWeakObjectPtr<UObject>> Objects;
 	DetailLayout.GetObjectsBeingCustomized(Objects);
@@ -29,38 +33,56 @@ void FDynamicEntryBoxDetails::CustomizeDetails(IDetailLayoutBuilder& DetailLayou
 	IDetailCategoryBuilder& EntryLayoutCategory = DetailLayout.EditCategory(TEXT("EntryLayout"));
 
 	const TAttribute<bool> CanEditAignmentAttribute(this, &FDynamicEntryBoxDetails::CanEditAlignment);
-	EntryLayoutCategory.AddProperty(DetailLayout.GetProperty(GET_MEMBER_NAME_CHECKED(UDynamicEntryBox, EntryHorizontalAlignment)))
+	EntryLayoutCategory.AddProperty(DetailLayout.GetProperty(GET_MEMBER_NAME_CHECKED(UDynamicEntryBoxBase, EntryHorizontalAlignment)))
 		.IsEnabled(CanEditAignmentAttribute);
-	EntryLayoutCategory.AddProperty(DetailLayout.GetProperty(GET_MEMBER_NAME_CHECKED(UDynamicEntryBox, EntryVerticalAlignment)))
+	EntryLayoutCategory.AddProperty(DetailLayout.GetProperty(GET_MEMBER_NAME_CHECKED(UDynamicEntryBoxBase, EntryVerticalAlignment)))
 		.IsEnabled(CanEditAignmentAttribute);
 
-	EntryLayoutCategory.AddProperty(DetailLayout.GetProperty(GET_MEMBER_NAME_CHECKED(UDynamicEntryBox, MaxElementSize)))
+	EntryLayoutCategory.AddProperty(DetailLayout.GetProperty(GET_MEMBER_NAME_CHECKED(UDynamicEntryBoxBase, MaxElementSize)))
 		.IsEnabled(TAttribute<bool>(this, &FDynamicEntryBoxDetails::CanEditMaxElementSize));
-	EntryLayoutCategory.AddProperty(DetailLayout.GetProperty(GET_MEMBER_NAME_CHECKED(UDynamicEntryBox, EntrySpacing)))
+	EntryLayoutCategory.AddProperty(DetailLayout.GetProperty(GET_MEMBER_NAME_CHECKED(UDynamicEntryBoxBase, EntrySpacing)))
 		.IsEnabled(TAttribute<bool>(this, &FDynamicEntryBoxDetails::CanEditEntrySpacing));
-	EntryLayoutCategory.AddProperty(DetailLayout.GetProperty(GET_MEMBER_NAME_CHECKED(UDynamicEntryBox, SpacingPattern)))
+	EntryLayoutCategory.AddProperty(DetailLayout.GetProperty(GET_MEMBER_NAME_CHECKED(UDynamicEntryBoxBase, SpacingPattern)))
 		.IsEnabled(TAttribute<bool>(this, &FDynamicEntryBoxDetails::CanEditSpacingPattern));
-
-	AddEntryClassPicker(*EntryBox, EntryLayoutCategory, DetailLayout.GetProperty(GET_MEMBER_NAME_CHECKED(UDynamicEntryBox, EntryWidgetClass)));
 }
 
-bool FDynamicEntryBoxDetails::CanEditSpacingPattern() const
+bool FDynamicEntryBoxBaseDetails::CanEditSpacingPattern() const
 {
 	return EntryBox->GetBoxType() == EDynamicBoxType::Overlay;
 }
 
-bool FDynamicEntryBoxDetails::CanEditEntrySpacing() const
+bool FDynamicEntryBoxBaseDetails::CanEditEntrySpacing() const
 {
 	return EntryBox->SpacingPattern.Num() == 0;
 }
 
-bool FDynamicEntryBoxDetails::CanEditAlignment() const
+bool FDynamicEntryBoxBaseDetails::CanEditAlignment() const
 {
 	return EntryBox->GetBoxType() != EDynamicBoxType::Overlay || CanEditEntrySpacing();
 }
 
-bool FDynamicEntryBoxDetails::CanEditMaxElementSize() const
+bool FDynamicEntryBoxBaseDetails::CanEditMaxElementSize() const
 {
 	const EDynamicBoxType BoxType = EntryBox->GetBoxType();
 	return BoxType == EDynamicBoxType::Horizontal || BoxType == EDynamicBoxType::Vertical;
 }
+
+//////////////////////////////////////////////////////////////////////////
+// FDynamicEntryBoxDetails
+//////////////////////////////////////////////////////////////////////////
+
+TSharedRef<IDetailCustomization> FDynamicEntryBoxDetails::MakeInstance()
+{
+	return MakeShareable(new FDynamicEntryBoxDetails());
+}
+
+void FDynamicEntryBoxDetails::CustomizeDetails(IDetailLayoutBuilder& DetailLayout)
+{
+	FDynamicEntryBoxBaseDetails::CustomizeDetails(DetailLayout);
+
+	if (EntryBox.IsValid())
+	{
+		IDetailCategoryBuilder& EntryLayoutCategory = DetailLayout.EditCategory(TEXT("EntryLayout"));
+		AddEntryClassPicker(*EntryBox, EntryLayoutCategory, DetailLayout.GetProperty(GET_MEMBER_NAME_CHECKED(UDynamicEntryBox, EntryWidgetClass)));
+	}
+}
diff --git a/Engine/Source/Editor/UMGEditor/Private/Customizations/DynamicEntryBoxDetails.h b/Engine/Source/Editor/UMGEditor/Private/Customizations/DynamicEntryBoxDetails.h
--- a/Engine/Source/Editor/UMGEditor/Private/Customizations/DynamicEntryBoxDetails.h
+++ b/Engine/Source/Editor/UMGEditor/Private/Customizations/DynamicEntryBoxDetails.h
@@ -4,11 +4,14 @@
 
 #include "DynamicEntryWidgetDetailsBase.h"
 
-
 class IPropertyHandle;
-class UDynamicEntryBox;
+class UDynamicEntryBoxBase;
+
+//////////////////////////////////////////////////////////////////////////
+// FDynamicEntryBoxBaseDetails
+//////////////////////////////////////////////////////////////////////////
 
-class FDynamicEntryBoxDetails : public FDynamicEntryWidgetDetailsBase
+class FDynamicEntryBoxBaseDetails : public FDynamicEntryWidgetDetailsBase
 {
 public:
 	static TSharedRef<IDetailCustomization> MakeInstance();
@@ -16,11 +19,25 @@ public:
 	/* Main customization of details */
 	virtual void CustomizeDetails(IDetailLayoutBuilder& DetailLayout) override;
 
+protected:
+	TWeakObjectPtr<UDynamicEntryBoxBase> EntryBox;
+
 private:
 	bool CanEditSpacingPattern() const;
 	bool CanEditEntrySpacing() const;
 	bool CanEditAlignment() const;
 	bool CanEditMaxElementSize() const;
+};
+
+//////////////////////////////////////////////////////////////////////////
+// FDynamicEntryBoxDetails
+//////////////////////////////////////////////////////////////////////////
+
+class FDynamicEntryBoxDetails : public FDynamicEntryBoxBaseDetails
+{
+public:
+	static TSharedRef<IDetailCustomization> MakeInstance();
 
-	TWeakObjectPtr<UDynamicEntryBox> EntryBox;
+	/* Main customization of details */
+	virtual void CustomizeDetails(IDetailLayoutBuilder& DetailLayout) override;
 };
\ No newline at end of file
diff --git a/Engine/Source/Editor/UMGEditor/Private/UMGEditorModule.cpp b/Engine/Source/Editor/UMGEditor/Private/UMGEditorModule.cpp
--- a/Engine/Source/Editor/UMGEditor/Private/UMGEditorModule.cpp
+++ b/Engine/Source/Editor/UMGEditor/Private/UMGEditorModule.cpp
@@ -83,6 +83,7 @@ public:
 
 		// Class detail customizations
 		FPropertyEditorModule& PropertyModule = FModuleManager::GetModuleChecked<FPropertyEditorModule>(TEXT("PropertyEditor"));
+		PropertyModule.RegisterCustomClassLayout(TEXT("DynamicEntryBoxBase"), FOnGetDetailCustomizationInstance::CreateStatic(&FDynamicEntryBoxBaseDetails::MakeInstance));
 		PropertyModule.RegisterCustomClassLayout(TEXT("DynamicEntryBox"), FOnGetDetailCustomizationInstance::CreateStatic(&FDynamicEntryBoxDetails::MakeInstance));
 		PropertyModule.RegisterCustomClassLayout(TEXT("ListViewBase"), FOnGetDetailCustomizationInstance::CreateStatic(&FListViewBaseDetails::MakeInstance));
 	}
diff --git a/Engine/Source/Runtime/UMG/Private/Blueprint/UserWidgetPool.cpp b/Engine/Source/Runtime/UMG/Private/Blueprint/UserWidgetPool.cpp
--- a/Engine/Source/Runtime/UMG/Private/Blueprint/UserWidgetPool.cpp
+++ b/Engine/Source/Runtime/UMG/Private/Blueprint/UserWidgetPool.cpp
@@ -68,3 +68,8 @@ void FUserWidgetPool::ResetPool()
 	ActiveWidgets.Reset();
 	CachedSlateByWidgetObject.Reset();
 }
+
+void FUserWidgetPool::ReleaseSlateResources()
+{
+	CachedSlateByWidgetObject.Reset();
+}
diff --git a/Engine/Source/Runtime/UMG/Private/Components/DynamicEntryBox.cpp b/Engine/Source/Runtime/UMG/Private/Components/DynamicEntryBox.cpp
--- a/Engine/Source/Runtime/UMG/Private/Components/DynamicEntryBox.cpp
+++ b/Engine/Source/Runtime/UMG/Private/Components/DynamicEntryBox.cpp
@@ -12,183 +12,21 @@
 // UDynamicEntryBox
 //////////////////////////////////////////////////////////////////////////
 
-UDynamicEntryBox::UDynamicEntryBox(const FObjectInitializer& Initializer)
-	: Super(Initializer)
-	, EntryWidgetPool(*this)
-{
-	bIsVariable = true;
-
-	Visibility = ESlateVisibility::SelfHitTestInvisible;
-	EntrySizeRule.SizeRule = ESlateSizeRule::Automatic;
-}
-
-void UDynamicEntryBox::ReleaseSlateResources(bool bReleaseChildren)
-{
-	Super::ReleaseSlateResources(bReleaseChildren);
-
-	EntryWidgetPool.ResetPool();
-	MyPanelWidget.Reset();
-}
-
 void UDynamicEntryBox::Reset(bool bDeleteWidgets)
 {
-	EntryWidgetPool.ReleaseAll(bDeleteWidgets);
-
-	if (MyPanelWidget.IsValid())
-	{
-		switch (EntryBoxType)
-		{
-		case EDynamicBoxType::Horizontal:
-		case EDynamicBoxType::Vertical:
-			StaticCastSharedPtr<SBoxPanel>(MyPanelWidget)->ClearChildren();
-			break;
-		case EDynamicBoxType::Wrap:
-			StaticCastSharedPtr<SWrapBox>(MyPanelWidget)->ClearChildren();
-			break;
-		case EDynamicBoxType::Overlay:
-			StaticCastSharedPtr<SOverlay>(MyPanelWidget)->ClearChildren();
-			break;
-		}
-	}
-}
-
-const TArray<UUserWidget*>& UDynamicEntryBox::GetAllEntries() const
-{
-	return EntryWidgetPool.GetActiveWidgets();
-}
-
-int32 UDynamicEntryBox::GetNumEntries() const
-{
-	return MyPanelWidget.IsValid() ? MyPanelWidget->GetChildren()->Num() : 0;
+	ResetInternal(bDeleteWidgets);
 }
 
 void UDynamicEntryBox::RemoveEntry(UUserWidget* EntryWidget)
 {
-	if (EntryWidget)
-	{
-		if (MyPanelWidget.IsValid())
-		{
-			TSharedPtr<SWidget> CachedEntryWidget = EntryWidget->GetCachedWidget();
-			if (CachedEntryWidget.IsValid())
-			{
-				switch (EntryBoxType)
-				{
-				case EDynamicBoxType::Horizontal:
-				case EDynamicBoxType::Vertical:
-					StaticCastSharedPtr<SBoxPanel>(MyPanelWidget)->RemoveSlot(CachedEntryWidget.ToSharedRef());
-					break;
-				case EDynamicBoxType::Wrap:
-					StaticCastSharedPtr<SWrapBox>(MyPanelWidget)->RemoveSlot(CachedEntryWidget.ToSharedRef());
-					break;
-				case EDynamicBoxType::Overlay:
-					StaticCastSharedPtr<SOverlay>(MyPanelWidget)->RemoveSlot(CachedEntryWidget.ToSharedRef());
-					break;
-				}
-			}
-		}
-		EntryWidgetPool.Release(EntryWidget);
-	}
-}
-
-void UDynamicEntryBox::SetEntrySpacing(const FVector2D& InEntrySpacing)
-{
-	EntrySpacing = InEntrySpacing;
-
-	if (MyPanelWidget.IsValid())
-	{
-		if (EntryBoxType == EDynamicBoxType::Wrap)
-		{
-			// Wrap boxes can change their widget spacing on the fly
-			StaticCastSharedPtr<SWrapBox>(MyPanelWidget)->SetInnerSlotPadding(EntrySpacing);
-		}
-		else if (EntryBoxType == EDynamicBoxType::Overlay)
-		{
-			TPanelChildren<SOverlay::FOverlaySlot>* OverlayChildren = static_cast<TPanelChildren<SOverlay::FOverlaySlot>*>(MyPanelWidget->GetChildren());
-			for (int32 ChildIdx = 0; ChildIdx < OverlayChildren->Num(); ++ChildIdx)
-			{
-				FMargin Padding;
-				if (SpacingPattern.Num() > 0)
-				{
-					FVector2D Spacing(0.f, 0.f);
-
-					// First establish the starting location
-					for (int32 CountIdx = 0; CountIdx < ChildIdx; ++CountIdx)
-					{
-						int32 PatternIdx = CountIdx % SpacingPattern.Num();
-						Spacing += SpacingPattern[PatternIdx];
-					}
-					
-					// Negative padding is no good, so negative spacing is expressed as positive spacing on the opposite side
-					if (Spacing.X >= 0.f)
-					{
-						Padding.Left = Spacing.X;
-					}
-					else
-					{
-						Padding.Right = -Spacing.X;
-					}
-					if (Spacing.Y >= 0.f)
-					{
-						Padding.Top = Spacing.Y;
-					}
-					else
-					{
-						Padding.Bottom = -Spacing.Y;
-					}
-				}
-				else
-				{
-					if (EntrySpacing.X >= 0.f)
-					{
-						Padding.Left = ChildIdx * EntrySpacing.X;
-					}
-					else
-					{
-						Padding.Right = ChildIdx * -EntrySpacing.X;
-					}
-
-					if (EntrySpacing.Y >= 0.f)
-					{
-						Padding.Top = ChildIdx * EntrySpacing.Y;
-					}
-					else
-					{
-						Padding.Bottom = ChildIdx * -EntrySpacing.Y;
-					}
-				}
-				SOverlay::FOverlaySlot& OverlaySlot = (*OverlayChildren)[ChildIdx];
-				OverlaySlot.SlotPadding = Padding;
-			}
-		}
-		else
-		{
-			// Vertical & Horizontal have to manually update the padding on each slot
-			const bool bIsHBox = EntryBoxType == EDynamicBoxType::Horizontal;
-			TPanelChildren<SBoxPanel::FSlot>* BoxChildren = static_cast<TPanelChildren<SBoxPanel::FSlot>*>(MyPanelWidget->GetChildren());
-			for (int32 ChildIdx = 0; ChildIdx < BoxChildren->Num(); ++ChildIdx)
-			{
-				const bool bIsFirstChild = ChildIdx == 0;
-
-				FMargin Padding;
-				Padding.Top = bIsHBox || bIsFirstChild ? 0.f : EntrySpacing.Y;
-				Padding.Left = bIsHBox && !bIsFirstChild ? EntrySpacing.X : 0.f;
-
-				SBoxPanel::FSlot& BoxSlot = (*BoxChildren)[ChildIdx];
-				BoxSlot.SlotPadding = Padding;
-			}
-		}
-	}
+	RemoveEntryInternal(EntryWidget);
 }
 
 #if WITH_EDITOR
-
-const FText UDynamicEntryBox::GetPaletteCategory()
-{
-	return LOCTEXT("Advanced", "Advanced");
-}
-
 void UDynamicEntryBox::ValidateCompiledDefaults(IWidgetCompilerLog& CompileLog) const
 {
+	Super::ValidateCompiledDefaults(CompileLog);
+	
 	if (!EntryWidgetClass)
 	{
 		CompileLog.Error(FText::Format(LOCTEXT("Error_DynamicEntryBox_MissingEntryClass", "{0} has no EntryWidgetClass specified - required for any Dynamic Entry Box to function."), FText::FromString(GetName())));
@@ -196,43 +34,6 @@ void UDynamicEntryBox::ValidateCompiledDefaults(IWidgetCompilerLog& CompileLog)
 }
 #endif
 
-TSharedRef<SWidget> UDynamicEntryBox::RebuildWidget()
-{
-	TSharedPtr<SWidget> EntryBoxWidget;
-	switch (EntryBoxType)
-	{
-	case EDynamicBoxType::Horizontal:
-		EntryBoxWidget = SAssignNew(MyPanelWidget, SHorizontalBox);
-		break;
-	case EDynamicBoxType::Vertical:
-		EntryBoxWidget = SAssignNew(MyPanelWidget, SVerticalBox);
-		break;
-	case EDynamicBoxType::Wrap:
-		EntryBoxWidget = SAssignNew(MyPanelWidget, SWrapBox)
-			.UseAllottedWidth(true)
-			.InnerSlotPadding(EntrySpacing);
-		break;
-	case EDynamicBoxType::Overlay:
-		EntryBoxWidget = SAssignNew(MyPanelWidget, SOverlay)
-			.Clipping(EWidgetClipping::ClipToBounds);
-		break;
-	}
-
-	return EntryBoxWidget.ToSharedRef();
-}
-
-#if WITH_EDITOR
-void UDynamicEntryBox::PostEditChangeProperty(FPropertyChangedEvent& PropertyChangedEvent)
-{
-	if (MyPanelWidget.IsValid() && PropertyChangedEvent.GetPropertyName() == TEXT("EntryBoxType"))
-	{
-		MyPanelWidget.Reset();
-	}
-
-	Super::PostEditChangeProperty(PropertyChangedEvent);
-}
-#endif
-
 void UDynamicEntryBox::SynchronizeProperties()
 {
 	Super::SynchronizeProperties();
@@ -261,13 +62,6 @@ void UDynamicEntryBox::SynchronizeProperties()
 				StartingNumber++;
 			}
 		}
-		else
-		{
-			// If we don't need to rebuild, update existing entries
-			SetEntrySpacing(EntrySpacing);
-			
-			//@todo DanH: update alignment, spacing pattern, clipping, size rule, max element size
-		}
 	}
 #endif
 }
@@ -283,109 +77,7 @@ UUserWidget* UDynamicEntryBox::BP_CreateEntryOfClass(TSubclassOf<UUserWidget> En
 	{
 		return CreateEntryInternal(EntryClass);
 	}
-
 	return nullptr;
 }
 
-UUserWidget* UDynamicEntryBox::CreateEntryInternal(TSubclassOf<UUserWidget> InEntryClass)
-{
-	if (MyPanelWidget.IsValid())
-	{
-		UUserWidget* NewEntryWidget = EntryWidgetPool.GetOrCreateInstance(InEntryClass);
-		AddEntryChild(*NewEntryWidget);
-		return NewEntryWidget;
-	}
-
-	UE_LOG(LogUMG, Warning, TEXT("UDynamicEntryBox::CreateEntryInternal(): Failed to create an entry."));
-	return nullptr;
-}
-
-FMargin UDynamicEntryBox::BuildEntryPadding(const FVector2D& DesiredSpacing)
-{
-	FMargin EntryPadding;
-	if (DesiredSpacing.X >= 0.f)
-	{
-		EntryPadding.Left = DesiredSpacing.X;
-	}
-	else
-	{
-		EntryPadding.Right = -DesiredSpacing.X;
-	}
-
-	if (DesiredSpacing.Y >= 0.f)
-	{
-		EntryPadding.Top = DesiredSpacing.Y;
-	}
-	else
-	{
-		EntryPadding.Bottom = -DesiredSpacing.Y;
-	}
-
-	return EntryPadding;
-}
-
-void UDynamicEntryBox::AddEntryChild(UUserWidget& ChildWidget)
-{
-	FSlotBase* NewSlot = nullptr;
-	if (EntryBoxType == EDynamicBoxType::Wrap)
-	{
-		NewSlot = &StaticCastSharedPtr<SWrapBox>(MyPanelWidget)->AddSlot()
-			.FillEmptySpace(false)
-			.HAlign(EntryHorizontalAlignment)
-			.VAlign(EntryVerticalAlignment);
-	}
-	else if (EntryBoxType == EDynamicBoxType::Overlay)
-	{
-		const int32 ChildIdx = MyPanelWidget->GetChildren()->Num();
-		SOverlay::FOverlaySlot& OverlaySlot = (SOverlay::FOverlaySlot&)StaticCastSharedPtr<SOverlay>(MyPanelWidget)->AddSlot();
-
-		EHorizontalAlignment HAlign = EntryHorizontalAlignment;
-		EVerticalAlignment VAlign = EntryVerticalAlignment;
-
-		FVector2D TargetSpacing = FVector2D::ZeroVector;
-		if (SpacingPattern.Num() > 0)
-		{
-			for (int32 CountIdx = 0; CountIdx < ChildIdx; ++CountIdx)
-			{
-				const int32 PatternIdx = CountIdx % SpacingPattern.Num();
-				TargetSpacing += SpacingPattern[PatternIdx];
-			}
-		}
-		else
-		{
-			TargetSpacing = EntrySpacing * ChildIdx;
-			HAlign = EntrySpacing.X >= 0.f ? EHorizontalAlignment::HAlign_Left : EHorizontalAlignment::HAlign_Right;
-			VAlign = EntrySpacing.Y >= 0.f ? EVerticalAlignment::VAlign_Top : EVerticalAlignment::VAlign_Bottom;
-		}
-		
-		OverlaySlot.HAlignment = HAlign;
-		OverlaySlot.VAlignment = VAlign;
-		OverlaySlot.SlotPadding = BuildEntryPadding(TargetSpacing);
-
-		NewSlot = &OverlaySlot;
-	}
-	else
-	{
-		const bool bIsHBox = EntryBoxType == EDynamicBoxType::Horizontal;
-		const bool bIsFirstChild = MyPanelWidget->GetChildren()->Num() == 0;
-
-		SBoxPanel::FSlot& BoxPanelSlot = bIsHBox ? (SBoxPanel::FSlot&)StaticCastSharedPtr<SHorizontalBox>(MyPanelWidget)->AddSlot().MaxWidth(MaxElementSize) : (SBoxPanel::FSlot&)StaticCastSharedPtr<SVerticalBox>(MyPanelWidget)->AddSlot().MaxHeight(MaxElementSize);
-		BoxPanelSlot.HAlignment = EntryHorizontalAlignment;
-		BoxPanelSlot.VAlignment = EntryVerticalAlignment;
-		BoxPanelSlot.SizeParam = UWidget::ConvertSerializedSizeParamToRuntime(EntrySizeRule);
-
-		FMargin Padding;
-		Padding.Top = bIsHBox || bIsFirstChild ? 0.f : EntrySpacing.Y;
-		Padding.Left = bIsHBox && !bIsFirstChild ? EntrySpacing.X : 0.f;
-		BoxPanelSlot.SlotPadding = Padding;
-
-		NewSlot = &BoxPanelSlot;
-	}
-
-	if (ensure(NewSlot))
-	{
-		NewSlot->AttachWidget(ChildWidget.TakeWidget());
-	}
-}
-
 #undef LOCTEXT_NAMESPACE
\ No newline at end of file
diff --git a/Engine/Source/Runtime/UMG/Private/Components/DynamicEntryBoxBase.cpp b/Engine/Source/Runtime/UMG/Private/Components/DynamicEntryBoxBase.cpp
new file mode 100644
--- /dev/null
+++ b/Engine/Source/Runtime/UMG/Private/Components/DynamicEntryBoxBase.cpp
@@ -0,0 +1,342 @@
+// Copyright 1998-2018 Epic Games, Inc. All Rights Reserved.
+
+#include "Components/DynamicEntryBoxBase.h"
+#include "UMGPrivate.h"
+#include "Widgets/Layout/SWrapBox.h"
+#include "Widgets/SBoxPanel.h"
+#include "Editor/WidgetCompilerLog.h"
+
+#define LOCTEXT_NAMESPACE "UMG"
+
+UDynamicEntryBoxBase::UDynamicEntryBoxBase(const FObjectInitializer& Initializer)
+	: Super(Initializer)
+	, EntryWidgetPool(*this)
+{
+	Visibility = ESlateVisibility::SelfHitTestInvisible;
+	EntrySizeRule.SizeRule = ESlateSizeRule::Automatic;
+}
+
+void UDynamicEntryBoxBase::ReleaseSlateResources(bool bReleaseChildren)
+{
+	Super::ReleaseSlateResources(bReleaseChildren);
+
+	EntryWidgetPool.ReleaseSlateResources();
+	MyPanelWidget.Reset();
+}
+
+void UDynamicEntryBoxBase::ResetInternal(bool bDeleteWidgets)
+{
+	EntryWidgetPool.ReleaseAll(bDeleteWidgets);
+
+	if (MyPanelWidget.IsValid())
+	{
+		switch (EntryBoxType)
+		{
+		case EDynamicBoxType::Horizontal:
+		case EDynamicBoxType::Vertical:
+			StaticCastSharedPtr<SBoxPanel>(MyPanelWidget)->ClearChildren();
+			break;
+		case EDynamicBoxType::Wrap:
+			StaticCastSharedPtr<SWrapBox>(MyPanelWidget)->ClearChildren();
+			break;
+		case EDynamicBoxType::Overlay:
+			StaticCastSharedPtr<SOverlay>(MyPanelWidget)->ClearChildren();
+			break;
+		}
+	}
+}
+
+const TArray<UUserWidget*>& UDynamicEntryBoxBase::GetAllEntries() const
+{
+	return EntryWidgetPool.GetActiveWidgets();
+}
+
+int32 UDynamicEntryBoxBase::GetNumEntries() const
+{
+	return EntryWidgetPool.GetActiveWidgets().Num();
+}
+
+void UDynamicEntryBoxBase::RemoveEntryInternal(UUserWidget* EntryWidget)
+{
+	if (EntryWidget)
+	{
+		if (MyPanelWidget.IsValid())
+		{
+			TSharedPtr<SWidget> CachedEntryWidget = EntryWidget->GetCachedWidget();
+			if (CachedEntryWidget.IsValid())
+			{
+				switch (EntryBoxType)
+				{
+				case EDynamicBoxType::Horizontal:
+				case EDynamicBoxType::Vertical:
+					StaticCastSharedPtr<SBoxPanel>(MyPanelWidget)->RemoveSlot(CachedEntryWidget.ToSharedRef());
+					break;
+				case EDynamicBoxType::Wrap:
+					StaticCastSharedPtr<SWrapBox>(MyPanelWidget)->RemoveSlot(CachedEntryWidget.ToSharedRef());
+					break;
+				case EDynamicBoxType::Overlay:
+					StaticCastSharedPtr<SOverlay>(MyPanelWidget)->RemoveSlot(CachedEntryWidget.ToSharedRef());
+					break;
+				}
+			}
+		}
+		EntryWidgetPool.Release(EntryWidget);
+	}
+}
+
+void UDynamicEntryBoxBase::SetEntrySpacing(const FVector2D& InEntrySpacing)
+{
+	EntrySpacing = InEntrySpacing;
+
+	if (MyPanelWidget.IsValid())
+	{
+		if (EntryBoxType == EDynamicBoxType::Wrap)
+		{
+			// Wrap boxes can change their widget spacing on the fly
+			StaticCastSharedPtr<SWrapBox>(MyPanelWidget)->SetInnerSlotPadding(EntrySpacing);
+		}
+		else if (EntryBoxType == EDynamicBoxType::Overlay)
+		{
+			TPanelChildren<SOverlay::FOverlaySlot>* OverlayChildren = static_cast<TPanelChildren<SOverlay::FOverlaySlot>*>(MyPanelWidget->GetChildren());
+			for (int32 ChildIdx = 0; ChildIdx < OverlayChildren->Num(); ++ChildIdx)
+			{
+				FMargin Padding;
+				if (SpacingPattern.Num() > 0)
+				{
+					FVector2D Spacing(0.f, 0.f);
+
+					// First establish the starting location
+					for (int32 CountIdx = 0; CountIdx < ChildIdx; ++CountIdx)
+					{
+						int32 PatternIdx = CountIdx % SpacingPattern.Num();
+						Spacing += SpacingPattern[PatternIdx];
+					}
+					
+					// Negative padding is no good, so negative spacing is expressed as positive spacing on the opposite side
+					if (Spacing.X >= 0.f)
+					{
+						Padding.Left = Spacing.X;
+					}
+					else
+					{
+						Padding.Right = -Spacing.X;
+					}
+					if (Spacing.Y >= 0.f)
+					{
+						Padding.Top = Spacing.Y;
+					}
+					else
+					{
+						Padding.Bottom = -Spacing.Y;
+					}
+				}
+				else
+				{
+					if (EntrySpacing.X >= 0.f)
+					{
+						Padding.Left = ChildIdx * EntrySpacing.X;
+					}
+					else
+					{
+						Padding.Right = ChildIdx * -EntrySpacing.X;
+					}
+
+					if (EntrySpacing.Y >= 0.f)
+					{
+						Padding.Top = ChildIdx * EntrySpacing.Y;
+					}
+					else
+					{
+						Padding.Bottom = ChildIdx * -EntrySpacing.Y;
+					}
+				}
+				SOverlay::FOverlaySlot& OverlaySlot = (*OverlayChildren)[ChildIdx];
+				OverlaySlot.SlotPadding = Padding;
+			}
+		}
+		else
+		{
+			// Vertical & Horizontal have to manually update the padding on each slot
+			const bool bIsHBox = EntryBoxType == EDynamicBoxType::Horizontal;
+			TPanelChildren<SBoxPanel::FSlot>* BoxChildren = static_cast<TPanelChildren<SBoxPanel::FSlot>*>(MyPanelWidget->GetChildren());
+			for (int32 ChildIdx = 0; ChildIdx < BoxChildren->Num(); ++ChildIdx)
+			{
+				const bool bIsFirstChild = ChildIdx == 0;
+
+				FMargin Padding;
+				Padding.Top = bIsHBox || bIsFirstChild ? 0.f : EntrySpacing.Y;
+				Padding.Left = bIsHBox && !bIsFirstChild ? EntrySpacing.X : 0.f;
+
+				SBoxPanel::FSlot& BoxSlot = (*BoxChildren)[ChildIdx];
+				BoxSlot.SlotPadding = Padding;
+			}
+		}
+	}
+}
+
+#if WITH_EDITOR
+
+const FText UDynamicEntryBoxBase::GetPaletteCategory()
+{
+	return LOCTEXT("Advanced", "Advanced");
+}
+#endif
+
+TSharedRef<SWidget> UDynamicEntryBoxBase::RebuildWidget()
+{
+	TSharedPtr<SWidget> EntryBoxWidget;
+	switch (EntryBoxType)
+	{
+	case EDynamicBoxType::Horizontal:
+		EntryBoxWidget = SAssignNew(MyPanelWidget, SHorizontalBox);
+		break;
+	case EDynamicBoxType::Vertical:
+		EntryBoxWidget = SAssignNew(MyPanelWidget, SVerticalBox);
+		break;
+	case EDynamicBoxType::Wrap:
+		EntryBoxWidget = SAssignNew(MyPanelWidget, SWrapBox)
+			.UseAllottedWidth(true)
+			.InnerSlotPadding(EntrySpacing);
+		break;
+	case EDynamicBoxType::Overlay:
+		EntryBoxWidget = SAssignNew(MyPanelWidget, SOverlay)
+			.Clipping(EWidgetClipping::ClipToBounds);
+		break;
+	}
+
+	if (!IsDesignTime())
+	{
+		// Populate now with all the entries that have been created so far
+		for (UUserWidget* ActiveWidget : EntryWidgetPool.GetActiveWidgets())
+		{
+			AddEntryChild(*ActiveWidget);
+		}
+	}
+
+	return EntryBoxWidget.ToSharedRef();
+}
+
+#if WITH_EDITOR
+void UDynamicEntryBoxBase::PostEditChangeProperty(FPropertyChangedEvent& PropertyChangedEvent)
+{
+	if (MyPanelWidget.IsValid() && PropertyChangedEvent.GetPropertyName() == TEXT("EntryBoxType"))
+	{
+		MyPanelWidget.Reset();
+	}
+
+	Super::PostEditChangeProperty(PropertyChangedEvent);
+}
+#endif
+
+void UDynamicEntryBoxBase::SynchronizeProperties()
+{
+	Super::SynchronizeProperties();
+
+#if WITH_EDITORONLY_DATA
+	if (IsDesignTime())
+	{
+		SetEntrySpacing(EntrySpacing);
+	}
+#endif
+}
+
+UUserWidget* UDynamicEntryBoxBase::CreateEntryInternal(TSubclassOf<UUserWidget> InEntryClass)
+{
+	UUserWidget* NewEntryWidget = EntryWidgetPool.GetOrCreateInstance(InEntryClass);
+	if (MyPanelWidget.IsValid())
+	{
+		// If we've already been constructed, immediately add the child to our panel widget
+		AddEntryChild(*NewEntryWidget);
+	}
+	return NewEntryWidget;
+}
+
+FMargin UDynamicEntryBoxBase::BuildEntryPadding(const FVector2D& DesiredSpacing)
+{
+	FMargin EntryPadding;
+	if (DesiredSpacing.X >= 0.f)
+	{
+		EntryPadding.Left = DesiredSpacing.X;
+	}
+	else
+	{
+		EntryPadding.Right = -DesiredSpacing.X;
+	}
+
+	if (DesiredSpacing.Y >= 0.f)
+	{
+		EntryPadding.Top = DesiredSpacing.Y;
+	}
+	else
+	{
+		EntryPadding.Bottom = -DesiredSpacing.Y;
+	}
+
+	return EntryPadding;
+}
+
+void UDynamicEntryBoxBase::AddEntryChild(UUserWidget& ChildWidget)
+{
+	FSlotBase* NewSlot = nullptr;
+	if (EntryBoxType == EDynamicBoxType::Wrap)
+	{
+		NewSlot = &StaticCastSharedPtr<SWrapBox>(MyPanelWidget)->AddSlot()
+			.FillEmptySpace(false)
+			.HAlign(EntryHorizontalAlignment)
+			.VAlign(EntryVerticalAlignment);
+	}
+	else if (EntryBoxType == EDynamicBoxType::Overlay)
+	{
+		const int32 ChildIdx = MyPanelWidget->GetChildren()->Num();
+		SOverlay::FOverlaySlot& OverlaySlot = (SOverlay::FOverlaySlot&)StaticCastSharedPtr<SOverlay>(MyPanelWidget)->AddSlot();
+
+		EHorizontalAlignment HAlign = EntryHorizontalAlignment;
+		EVerticalAlignment VAlign = EntryVerticalAlignment;
+
+		FVector2D TargetSpacing = FVector2D::ZeroVector;
+		if (SpacingPattern.Num() > 0)
+		{
+			for (int32 CountIdx = 0; CountIdx < ChildIdx; ++CountIdx)
+			{
+				const int32 PatternIdx = CountIdx % SpacingPattern.Num();
+				TargetSpacing += SpacingPattern[PatternIdx];
+			}
+		}
+		else
+		{
+			TargetSpacing = EntrySpacing * ChildIdx;
+			HAlign = EntrySpacing.X >= 0.f ? EHorizontalAlignment::HAlign_Left : EHorizontalAlignment::HAlign_Right;
+			VAlign = EntrySpacing.Y >= 0.f ? EVerticalAlignment::VAlign_Top : EVerticalAlignment::VAlign_Bottom;
+		}
+		
+		OverlaySlot.HAlignment = HAlign;
+		OverlaySlot.VAlignment = VAlign;
+		OverlaySlot.SlotPadding = BuildEntryPadding(TargetSpacing);
+
+		NewSlot = &OverlaySlot;
+	}
+	else
+	{
+		const bool bIsHBox = EntryBoxType == EDynamicBoxType::Horizontal;
+		const bool bIsFirstChild = MyPanelWidget->GetChildren()->Num() == 0;
+
+		SBoxPanel::FSlot& BoxPanelSlot = bIsHBox ? (SBoxPanel::FSlot&)StaticCastSharedPtr<SHorizontalBox>(MyPanelWidget)->AddSlot().MaxWidth(MaxElementSize) : (SBoxPanel::FSlot&)StaticCastSharedPtr<SVerticalBox>(MyPanelWidget)->AddSlot().MaxHeight(MaxElementSize);
+		BoxPanelSlot.HAlignment = EntryHorizontalAlignment;
+		BoxPanelSlot.VAlignment = EntryVerticalAlignment;
+		BoxPanelSlot.SizeParam = UWidget::ConvertSerializedSizeParamToRuntime(EntrySizeRule);
+
+		FMargin Padding;
+		Padding.Top = bIsHBox || bIsFirstChild ? 0.f : EntrySpacing.Y;
+		Padding.Left = bIsHBox && !bIsFirstChild ? EntrySpacing.X : 0.f;
+		BoxPanelSlot.SlotPadding = Padding;
+
+		NewSlot = &BoxPanelSlot;
+	}
+
+	if (ensure(NewSlot))
+	{
+		NewSlot->AttachWidget(ChildWidget.TakeWidget());
+	}
+}
+
+#undef LOCTEXT_NAMESPACE
\ No newline at end of file
diff --git a/Engine/Source/Runtime/UMG/Public/Blueprint/UserWidgetPool.h b/Engine/Source/Runtime/UMG/Public/Blueprint/UserWidgetPool.h
--- a/Engine/Source/Runtime/UMG/Public/Blueprint/UserWidgetPool.h
+++ b/Engine/Source/Runtime/UMG/Public/Blueprint/UserWidgetPool.h
@@ -66,6 +66,9 @@ public:
 	/** Full reset of all created widget objects (and any cached underlying slate) */
 	void ResetPool();
 
+	/** Reset of all cached underlying Slate widgets, but not the active UUserWidget objects */
+	void ReleaseSlateResources();
+
 private:
 	template <typename UserWidgetT = UUserWidget>
 	UserWidgetT* AddActiveWidgetInternal(TSubclassOf<UserWidgetT> WidgetClass, WidgetConstructFunc ConstructWidgetFunc)
diff --git a/Engine/Source/Runtime/UMG/Public/Components/DynamicEntryBox.h b/Engine/Source/Runtime/UMG/Public/Components/DynamicEntryBox.h
--- a/Engine/Source/Runtime/UMG/Public/Components/DynamicEntryBox.h
+++ b/Engine/Source/Runtime/UMG/Public/Components/DynamicEntryBox.h
@@ -2,22 +2,10 @@
 
 #pragma once
 
-#include "Components/Widget.h"
-#include "Blueprint/UserWidgetPool.h"
+#include "Components/DynamicEntryBoxBase.h"
 
 #include "DynamicEntryBox.generated.h"
 
-class UUserWidget;
-
-UENUM(BlueprintType)
-enum class EDynamicBoxType : uint8
-{
-	Horizontal,
-	Vertical,
-	Wrap,
-	Overlay
-};
-
 /**
  * A special box panel that auto-generates its entries at both design-time and runtime.
  * Useful for cases where you can have a varying number of entries, but it isn't worth the effort or conceptual overhead to set up a list/tile view.
@@ -26,28 +14,20 @@ enum class EDynamicBoxType : uint8
  * No children can be manually added in the designer - all are auto-generated based on the given entry class.
  */
 UCLASS()
-class UMG_API UDynamicEntryBox : public UWidget
+class UMG_API UDynamicEntryBox : public UDynamicEntryBoxBase
 {
 	GENERATED_BODY()
 
 public:
-	UDynamicEntryBox(const FObjectInitializer& Initializer);
-	virtual void ReleaseSlateResources(bool bReleaseChildren) override;
-
-	EDynamicBoxType GetBoxType() const { return EntryBoxType; }
 	TSubclassOf<UUserWidget> GetEntryWidgetClass() const { return EntryWidgetClass; }
-	const FVector2D& GetEntrySpacing() const { return EntrySpacing; }
 	
 	template <typename WidgetT = UUserWidget>
 	WidgetT* CreateEntry(const TSubclassOf<WidgetT>& ExplicitEntryClass = nullptr)
 	{
-		if (ExplicitEntryClass && ExplicitEntryClass->IsChildOf(WidgetT::StaticClass()))
-		{
-			return Cast<WidgetT>(CreateEntryInternal(ExplicitEntryClass));
-		}
-		if (EntryWidgetClass && EntryWidgetClass->IsChildOf(WidgetT::StaticClass()))
+		TSubclassOf<UUserWidget> EntryClass = ExplicitEntryClass ? ExplicitEntryClass : EntryWidgetClass;
+		if (EntryClass && EntryClass->IsChildOf(WidgetT::StaticClass()))
 		{
-			return Cast<WidgetT>(CreateEntryInternal(EntryWidgetClass));
+			return Cast<WidgetT>(CreateEntryInternal(EntryClass));
 		}
 		return nullptr;
 	}
@@ -55,7 +35,7 @@ public:
 	template <typename WidgetT = UUserWidget>
 	void Reset(TFunctionRef<void(WidgetT&)> ResetEntryFunc, bool bDeleteWidgets = false)
 	{
-		for (UUserWidget* EntryWidget : EntryWidgetPool.GetActiveWidgets())
+		for (UUserWidget* EntryWidget : GetAllEntries())
 		{
 			if (WidgetT* TypedEntry = Cast<WidgetT>(EntryWidget))
 			{
@@ -70,36 +50,11 @@ public:
 	UFUNCTION(BlueprintCallable, Category = DynamicEntryBox)
 	void Reset(bool bDeleteWidgets = false);
 
-	UFUNCTION(BlueprintCallable, Category = DynamicEntryBox)
-	const TArray<UUserWidget*>& GetAllEntries() const;
-
-	template <typename EntryWidgetT = UUserWidget>
-	TArray<EntryWidgetT*> GetTypedEntries() const
-	{
-		TArray<EntryWidgetT*> TypedEntries;
-		for (UUserWidget* Entry : GetAllEntries())
-		{
-			if (EntryWidgetT* TypedEntry = Cast<EntryWidgetT>(Entry))
-			{
-				TypedEntries.Add(TypedEntry);
-			}
-		}
-		return TypedEntries;
-	}
-
-	UFUNCTION(BlueprintCallable, Category = DynamicEntryBox)
-	int32 GetNumEntries() const;
-
 	UFUNCTION(BlueprintCallable, Category = DynamicEntryBox)
 	void RemoveEntry(UUserWidget* EntryWidget);
 
-	UFUNCTION(BlueprintCallable, Category = DynamicEntryBox)
-	void SetEntrySpacing(const FVector2D& InEntrySpacing);
-	
 	//~ Begin UWidget Interface
 #if WITH_EDITOR	
-	virtual const FText GetPaletteCategory() override;
-
 	virtual void ValidateCompiledDefaults(class IWidgetCompilerLog& CompileLog) const override;
 #endif
 	//~ End UWidget Interface
@@ -116,53 +71,8 @@ public:
 #endif
 
 protected:
-	virtual TSharedRef<SWidget> RebuildWidget() override;
-#if WITH_EDITOR
-	virtual void PostEditChangeProperty(struct FPropertyChangedEvent& PropertyChangedEvent) override;
-#endif
 	virtual void SynchronizeProperties() override;
 
-	virtual void AddEntryChild(UUserWidget& ChildWidget);
-
-protected:
-	UUserWidget* CreateEntryInternal(TSubclassOf<UUserWidget> InEntryClass);
-
-	/** The type of box panel into which created entries are added. Some differences in functionality exist between each type. */
-	UPROPERTY(EditAnywhere, Category = DynamicEntryBox, meta = (DesignerRebuild))
-	EDynamicBoxType EntryBoxType;
-
-	/** 
-	 * The padding to apply between entries in the box.
-	 * Note horizontal boxes only use the X and vertical boxes only use Y. Value is also ignored for the first entry in the box.
-	 * Wrap and Overlay types use both X and Y for spacing.
-	 */
-	UPROPERTY(EditAnywhere, BlueprintReadOnly, Category = EntryLayout)
-	FVector2D EntrySpacing;
-
-	//@todo DanH EntryBox: Consider giving a callback option as well/instead. Then this thing could actually create circular or pinwheel layouts...
-	/** The looping sequence of entry paddings to apply as entries are created. Overlay boxes only. Ignores EntrySpacing if not empty. */
-	UPROPERTY(EditAnywhere, BlueprintReadOnly, Category = EntryLayout)
-	TArray<FVector2D> SpacingPattern;
-
-	/** Sizing rule to apply to generated entries. Horizontal/Vertical boxes only. */
-	UPROPERTY(EditAnywhere, BlueprintReadOnly, Category = EntryLayout)
-	FSlateChildSize EntrySizeRule;
-
-	/** Horizontal alignment of generated entries. Horizontal/Vertical/Wrap boxes only. */
-	UPROPERTY(EditAnywhere, BlueprintReadOnly, Category = EntryLayout)
-	TEnumAsByte<EHorizontalAlignment> EntryHorizontalAlignment;
-
-	/** Vertical alignment of generated entries. Horizontal/Vertical/Wrap boxes only. */
-	UPROPERTY(EditAnywhere, BlueprintReadOnly, Category = EntryLayout)
-	TEnumAsByte<EVerticalAlignment> EntryVerticalAlignment;
-
-	/** The maximum size of each entry in the dominant axis of the box. Vertical/Horizontal boxes only. */
-	UPROPERTY(EditAnywhere, BlueprintReadOnly, Category = EntryLayout)
-	int32 MaxElementSize = 0;
-
-	// Can be a horizontal, vertical, wrap box, or overlay
-	TSharedPtr<SPanel> MyPanelWidget;
-
 private:
 	/** Creates and establishes a new dynamic entry in the box */
 	UFUNCTION(BlueprintCallable, Category = DynamicEntryBox, meta = (DisplayName = "Create Entry", AllowPrivateAccess = true))
@@ -172,16 +82,13 @@ private:
 	UFUNCTION(BlueprintCallable, Category = DynamicEntryBox, meta = (DisplayName = "Create Entry of Class", AllowPrivateAccess = true))
 	UUserWidget* BP_CreateEntryOfClass(TSubclassOf<UUserWidget> EntryClass);
 
-	FMargin BuildEntryPadding(const FVector2D& DesiredSpacing);
-	
 	/**
 	 * The class of widget to create entries of.
 	 * If natively binding this widget, use the EntryClass UPROPERTY metadata to specify the desired entry widget base class.
 	 */
 	UPROPERTY(EditAnywhere, BlueprintReadOnly, Category = EntryLayout, meta = (AllowPrivateAccess = true))
 	TSubclassOf<UUserWidget> EntryWidgetClass;
-	FUserWidgetPool EntryWidgetPool;
-
+	
 	// Let the details customization manipulate us directly
 	friend class FDynamicEntryBoxDetails;
 };
\ No newline at end of file
diff --git a/Engine/Source/Runtime/UMG/Public/Components/DynamicEntryBoxBase.h b/Engine/Source/Runtime/UMG/Public/Components/DynamicEntryBoxBase.h
new file mode 100644
--- /dev/null
+++ b/Engine/Source/Runtime/UMG/Public/Components/DynamicEntryBoxBase.h
@@ -0,0 +1,137 @@
+// Copyright 1998-2018 Epic Games, Inc. All Rights Reserved.
+
+#pragma once
+
+#include "Components/Widget.h"
+#include "Blueprint/UserWidgetPool.h"
+
+#include "DynamicEntryBoxBase.generated.h"
+
+class UUserWidget;
+
+UENUM(BlueprintType)
+enum class EDynamicBoxType : uint8
+{
+	Horizontal,
+	Vertical,
+	Wrap,
+	Overlay
+};
+
+/**
+ * Base for widgets that support a dynamic number of auto-generated entries at both design- and run-time.
+ * Contains all functionality needed to create, construct, and cache an arbitrary number of entry widgets, but exposes no means of entry creation or removal
+ * It's up to child classes to decide how they want to perform the population (some may do so entirely on their own without exposing a thing)
+ *
+ * @see UDynamicEntryBox for a ready-to-use version
+ */
+UCLASS(Abstract)
+class UMG_API UDynamicEntryBoxBase : public UWidget
+{
+	GENERATED_BODY()
+
+public:
+	UDynamicEntryBoxBase(const FObjectInitializer& Initializer);
+	virtual void ReleaseSlateResources(bool bReleaseChildren) override;
+
+	EDynamicBoxType GetBoxType() const { return EntryBoxType; }
+	const FVector2D& GetEntrySpacing() const { return EntrySpacing; }
+
+	UFUNCTION(BlueprintCallable, Category = DynamicEntryBox)
+	const TArray<UUserWidget*>& GetAllEntries() const;
+
+	template <typename EntryWidgetT = UUserWidget>
+	TArray<EntryWidgetT*> GetTypedEntries() const
+	{
+		TArray<EntryWidgetT*> TypedEntries;
+		for (UUserWidget* Entry : GetAllEntries())
+		{
+			if (EntryWidgetT* TypedEntry = Cast<EntryWidgetT>(Entry))
+			{
+				TypedEntries.Add(TypedEntry);
+			}
+		}
+		return TypedEntries;
+	}
+
+	UFUNCTION(BlueprintCallable, Category = DynamicEntryBox)
+	int32 GetNumEntries() const;
+
+	UFUNCTION(BlueprintCallable, Category = DynamicEntryBox)
+	void SetEntrySpacing(const FVector2D& InEntrySpacing);
+	
+protected:
+	virtual TSharedRef<SWidget> RebuildWidget() override;
+	virtual void SynchronizeProperties() override;
+	virtual void AddEntryChild(UUserWidget& ChildWidget);
+
+	UUserWidget* CreateEntryInternal(TSubclassOf<UUserWidget> InEntryClass);
+	void RemoveEntryInternal(UUserWidget* EntryWidget);
+	FMargin BuildEntryPadding(const FVector2D& DesiredSpacing);
+
+	/** Clear out the box entries, optionally deleting the underlying Slate widgets entirely as well. */
+	void ResetInternal(bool bDeleteWidgets = false);
+
+	/** Clear out the box entries, executing the provided reset function for each and optionally deleting the underlying Slate widgets entirely as well. */
+	template <typename WidgetT = UUserWidget>
+	void ResetInternal(TFunctionRef<void(WidgetT&)> ResetEntryFunc, bool bDeleteWidgets = false)
+	{
+		for (UUserWidget* EntryWidget : EntryWidgetPool.GetActiveWidgets())
+		{
+			if (WidgetT* TypedEntry = Cast<WidgetT>(EntryWidget))
+			{
+				ResetEntryFunc(*TypedEntry);
+			}
+		}
+
+		ResetInternal(bDeleteWidgets);
+	}
+
+#if WITH_EDITOR
+	virtual const FText GetPaletteCategory() override;
+	virtual void PostEditChangeProperty(struct FPropertyChangedEvent& PropertyChangedEvent) override;
+#endif
+
+protected:
+	/** The type of box panel into which created entries are added. Some differences in functionality exist between each type. */
+	UPROPERTY(EditAnywhere, Category = DynamicEntryBox, meta = (DesignerRebuild))
+	EDynamicBoxType EntryBoxType;
+
+	/** 
+	 * The padding to apply between entries in the box.
+	 * Note horizontal boxes only use the X and vertical boxes only use Y. Value is also ignored for the first entry in the box.
+	 * Wrap and Overlay types use both X and Y for spacing.
+	 */
+	UPROPERTY(EditAnywhere, BlueprintReadOnly, Category = EntryLayout)
+	FVector2D EntrySpacing;
+
+	//@todo DanH EntryBox: Consider giving a callback option as well/instead. Then this thing could actually create circular or pinwheel layouts...
+	/** The looping sequence of entry paddings to apply as entries are created. Overlay boxes only. Ignores EntrySpacing if not empty. */
+	UPROPERTY(EditAnywhere, BlueprintReadOnly, Category = EntryLayout)
+	TArray<FVector2D> SpacingPattern;
+
+	/** Sizing rule to apply to generated entries. Horizontal/Vertical boxes only. */
+	UPROPERTY(EditAnywhere, BlueprintReadOnly, Category = EntryLayout)
+	FSlateChildSize EntrySizeRule;
+
+	/** Horizontal alignment of generated entries. Horizontal/Vertical/Wrap boxes only. */
+	UPROPERTY(EditAnywhere, BlueprintReadOnly, Category = EntryLayout)
+	TEnumAsByte<EHorizontalAlignment> EntryHorizontalAlignment;
+
+	/** Vertical alignment of generated entries. Horizontal/Vertical/Wrap boxes only. */
+	UPROPERTY(EditAnywhere, BlueprintReadOnly, Category = EntryLayout)
+	TEnumAsByte<EVerticalAlignment> EntryVerticalAlignment;
+
+	/** The maximum size of each entry in the dominant axis of the box. Vertical/Horizontal boxes only. */
+	UPROPERTY(EditAnywhere, BlueprintReadOnly, Category = EntryLayout)
+	int32 MaxElementSize = 0;
+
+	// Can be a horizontal, vertical, wrap box, or overlay
+	TSharedPtr<SPanel> MyPanelWidget;
+
+private:
+	FUserWidgetPool EntryWidgetPool;
+
+	// Let the details customization manipulate us directly
+	friend class FDynamicEntryBoxBaseDetails;
+};
\ No newline at end of file
