Moving widget style/brush gathering from the core style into static local variables, to reduce rare crashes found from accessing the widget style lookup cross threads.

diff --git a/Engine/Source/Runtime/UMG/Private/Components/Button.cpp b/Engine/Source/Runtime/UMG/Private/Components/Button.cpp
--- a/Engine/Source/Runtime/UMG/Private/Components/Button.cpp
+++ b/Engine/Source/Runtime/UMG/Private/Components/Button.cpp
@@ -15,7 +15,8 @@ UButton::UButton(const FObjectInitializer& ObjectInitializer)
 	: Super(ObjectInitializer)
 {
 	// HACK: THIS SHOULD NOT COME FROM CORESTYLE AND SHOULD INSTEAD BY DEFINED BY ENGINE TEXTURES/PROJECT SETTINGS
-	WidgetStyle = FCoreStyle::Get().GetWidgetStyle< FButtonStyle >("Button");
+	static const FButtonStyle StaticButtonStyle = FCoreStyle::Get().GetWidgetStyle< FButtonStyle >("Button");
+	WidgetStyle = StaticButtonStyle;
 
 	ColorAndOpacity = FLinearColor::White;
 	BackgroundColor = FLinearColor::White;
diff --git a/Engine/Source/Runtime/UMG/Private/Components/CheckBox.cpp b/Engine/Source/Runtime/UMG/Private/Components/CheckBox.cpp
--- a/Engine/Source/Runtime/UMG/Private/Components/CheckBox.cpp
+++ b/Engine/Source/Runtime/UMG/Private/Components/CheckBox.cpp
@@ -15,7 +15,8 @@ UCheckBox::UCheckBox(const FObjectInitializer& ObjectInitializer)
 	: Super(ObjectInitializer)
 {
 	// HACK: THIS SHOULD NOT COME FROM CORESTYLE AND SHOULD INSTEAD BY DEFINED BY ENGINE TEXTURES/PROJECT SETTINGS
-	WidgetStyle = FCoreStyle::Get().GetWidgetStyle< FCheckBoxStyle >("Checkbox");
+	static const FCheckBoxStyle StaticCheckboxStyle = FCoreStyle::Get().GetWidgetStyle< FCheckBoxStyle >("Checkbox");
+	WidgetStyle = StaticCheckboxStyle;
 
 	CheckedState = ECheckBoxState::Unchecked;
 
diff --git a/Engine/Source/Runtime/UMG/Private/Components/CircularThrobber.cpp b/Engine/Source/Runtime/UMG/Private/Components/CircularThrobber.cpp
--- a/Engine/Source/Runtime/UMG/Private/Components/CircularThrobber.cpp
+++ b/Engine/Source/Runtime/UMG/Private/Components/CircularThrobber.cpp
@@ -16,7 +16,8 @@ UCircularThrobber::UCircularThrobber(const FObjectInitializer& ObjectInitializer
 	, bEnableRadius(true)
 {
 	// HACK: THIS SHOULD NOT COME FROM CORESTYLE AND SHOULD INSTEAD BY DEFINED BY ENGINE TEXTURES/PROJECT SETTINGS
-	Image = *FCoreStyle::Get().GetBrush("Throbber.CircleChunk");
+	static const FSlateBrush StaticBrushStyle = *FCoreStyle::Get().GetBrush("Throbber.CircleChunk");
+	Image = StaticBrushStyle;
 
 	NumberOfPieces = 6;
 	Period = 0.75f;
diff --git a/Engine/Source/Runtime/UMG/Private/Components/ComboBoxString.cpp b/Engine/Source/Runtime/UMG/Private/Components/ComboBoxString.cpp
--- a/Engine/Source/Runtime/UMG/Private/Components/ComboBoxString.cpp
+++ b/Engine/Source/Runtime/UMG/Private/Components/ComboBoxString.cpp
@@ -14,8 +14,11 @@ UComboBoxString::UComboBoxString(const FObjectInitializer& ObjectInitializer)
 	: Super(ObjectInitializer)
 {
 	// HACK: THIS SHOULD NOT COME FROM CORESTYLE AND SHOULD INSTEAD BY DEFINED BY ENGINE TEXTURES/PROJECT SETTINGS
-	WidgetStyle = FCoreStyle::Get().GetWidgetStyle< FComboBoxStyle >("ComboBox");
-	ItemStyle = FCoreStyle::Get().GetWidgetStyle< FTableRowStyle >("TableView.Row");
+	static const FComboBoxStyle StaticComboboxStyle = FCoreStyle::Get().GetWidgetStyle< FComboBoxStyle >("ComboBox");
+	static const FTableRowStyle StaticRowStyle = FCoreStyle::Get().GetWidgetStyle< FTableRowStyle >("TableView.Row");
+
+	WidgetStyle = StaticComboboxStyle;
+	ItemStyle = StaticRowStyle;
 	ItemStyle.SelectorFocusedBrush.TintColor = ItemStyle.SelectorFocusedBrush.TintColor.GetSpecifiedColor();
 	ItemStyle.ActiveHoveredBrush.TintColor = ItemStyle.ActiveHoveredBrush.TintColor.GetSpecifiedColor();
 	ItemStyle.ActiveBrush.TintColor = ItemStyle.ActiveBrush.TintColor.GetSpecifiedColor();
diff --git a/Engine/Source/Runtime/UMG/Private/Components/EditableText.cpp b/Engine/Source/Runtime/UMG/Private/Components/EditableText.cpp
--- a/Engine/Source/Runtime/UMG/Private/Components/EditableText.cpp
+++ b/Engine/Source/Runtime/UMG/Private/Components/EditableText.cpp
@@ -16,7 +16,8 @@ UEditableText::UEditableText(const FObjectInitializer& ObjectInitializer)
 	: Super(ObjectInitializer)
 {
 	// HACK: THIS SHOULD NOT COME FROM CORESTYLE AND SHOULD INSTEAD BY DEFINED BY ENGINE TEXTURES/PROJECT SETTINGS
-	WidgetStyle = FCoreStyle::Get().GetWidgetStyle< FEditableTextStyle >("NormalEditableText");
+	static const FEditableTextStyle StaticEditableTextStyle = FCoreStyle::Get().GetWidgetStyle< FEditableTextStyle >("NormalEditableText");
+	WidgetStyle = StaticEditableTextStyle;
 
 	ColorAndOpacity_DEPRECATED = FLinearColor::Black;
 
diff --git a/Engine/Source/Runtime/UMG/Private/Components/EditableTextBox.cpp b/Engine/Source/Runtime/UMG/Private/Components/EditableTextBox.cpp
--- a/Engine/Source/Runtime/UMG/Private/Components/EditableTextBox.cpp
+++ b/Engine/Source/Runtime/UMG/Private/Components/EditableTextBox.cpp
@@ -37,7 +37,8 @@ UEditableTextBox::UEditableTextBox(const FObjectInitializer& ObjectInitializer)
 	VirtualKeyboardDismissAction = EVirtualKeyboardDismissAction::TextChangeOnDismiss;
 
 	// HACK: THIS SHOULD NOT COME FROM CORESTYLE AND SHOULD INSTEAD BY DEFINED BY ENGINE TEXTURES/PROJECT SETTINGS
-	WidgetStyle = FCoreStyle::Get().GetWidgetStyle< FEditableTextBoxStyle >("NormalEditableTextBox");
+	static const FEditableTextBoxStyle StaticNormalEditableTextBox = FCoreStyle::Get().GetWidgetStyle< FEditableTextBoxStyle >("NormalEditableTextBox");
+	WidgetStyle = StaticNormalEditableTextBox;
 }
 
 void UEditableTextBox::ReleaseSlateResources(bool bReleaseChildren)
diff --git a/Engine/Source/Runtime/UMG/Private/Components/ExpandableArea.cpp b/Engine/Source/Runtime/UMG/Private/Components/ExpandableArea.cpp
--- a/Engine/Source/Runtime/UMG/Private/Components/ExpandableArea.cpp
+++ b/Engine/Source/Runtime/UMG/Private/Components/ExpandableArea.cpp
@@ -21,9 +21,13 @@ UExpandableArea::UExpandableArea(const FObjectInitializer& ObjectInitializer)
 	bIsVariable = true;
 
 	// HACK: THIS SHOULD NOT COME FROM CORESTYLE AND SHOULD INSTEAD BY DEFINED BY ENGINE TEXTURES/PROJECT SETTINGS
-	Style = FCoreStyle::Get().GetWidgetStyle<FExpandableAreaStyle>("ExpandableArea");
+	static const FExpandableAreaStyle StaticExpandableArea = FCoreStyle::Get().GetWidgetStyle<FExpandableAreaStyle>("ExpandableArea");
+	static const FSlateBrush StaticBorderBrush = *FCoreStyle::Get().GetBrush("ExpandableArea.Border");
+
+	Style = StaticExpandableArea;
+	BorderBrush = StaticBorderBrush;
+
 	BorderColor = FLinearColor::White;
-	BorderBrush = *FCoreStyle::Get().GetBrush("ExpandableArea.Border");
 	AreaPadding = FMargin(1);
 	HeaderPadding = FMargin(4.0f, 2.0f);
 }
diff --git a/Engine/Source/Runtime/UMG/Private/Components/InputKeySelector.cpp b/Engine/Source/Runtime/UMG/Private/Components/InputKeySelector.cpp
--- a/Engine/Source/Runtime/UMG/Private/Components/InputKeySelector.cpp
+++ b/Engine/Source/Runtime/UMG/Private/Components/InputKeySelector.cpp
@@ -14,8 +14,10 @@ UInputKeySelector::UInputKeySelector( const FObjectInitializer& ObjectInitialize
 	: Super(ObjectInitializer)
 {
 	// HACK: THIS SHOULD NOT COME FROM CORESTYLE AND SHOULD INSTEAD BY DEFINED BY ENGINE TEXTURES/PROJECT SETTINGS
-	WidgetStyle = FCoreStyle::Get().GetWidgetStyle<FButtonStyle>("Button");
-	TextStyle = FCoreStyle::Get().GetWidgetStyle< FTextBlockStyle >("NormalText");
+	static const FButtonStyle StaticButtonStyle = FCoreStyle::Get().GetWidgetStyle<FButtonStyle>("Button");
+	static const FTextBlockStyle StaticNormalTextStyle = FCoreStyle::Get().GetWidgetStyle< FTextBlockStyle >("NormalText");
+	WidgetStyle = StaticButtonStyle;
+	TextStyle = StaticNormalTextStyle;
 
 	KeySelectionText = NSLOCTEXT("InputKeySelector", "DefaultKeySelectionText", "...");
 	NoKeySpecifiedText = NSLOCTEXT("InputKeySelector", "DefaultEmptyText", "Empty");
diff --git a/Engine/Source/Runtime/UMG/Private/Components/MultiLineEditableText.cpp b/Engine/Source/Runtime/UMG/Private/Components/MultiLineEditableText.cpp
--- a/Engine/Source/Runtime/UMG/Private/Components/MultiLineEditableText.cpp
+++ b/Engine/Source/Runtime/UMG/Private/Components/MultiLineEditableText.cpp
@@ -15,7 +15,8 @@ UMultiLineEditableText::UMultiLineEditableText(const FObjectInitializer& ObjectI
 	: Super(ObjectInitializer)
 {
 	// HACK: THIS SHOULD NOT COME FROM CORESTYLE AND SHOULD INSTEAD BY DEFINED BY ENGINE TEXTURES/PROJECT SETTINGS
-	WidgetStyle = FCoreStyle::Get().GetWidgetStyle<FTextBlockStyle>("NormalText");
+	static const FTextBlockStyle StaticNormalText = FCoreStyle::Get().GetWidgetStyle<FTextBlockStyle>("NormalText");
+	WidgetStyle = StaticNormalText;
 	
 	bIsReadOnly = false;
 	SelectAllTextWhenFocused = false;
diff --git a/Engine/Source/Runtime/UMG/Private/Components/MultiLineEditableTextBox.cpp b/Engine/Source/Runtime/UMG/Private/Components/MultiLineEditableTextBox.cpp
--- a/Engine/Source/Runtime/UMG/Private/Components/MultiLineEditableTextBox.cpp
+++ b/Engine/Source/Runtime/UMG/Private/Components/MultiLineEditableTextBox.cpp
@@ -19,8 +19,10 @@ UMultiLineEditableTextBox::UMultiLineEditableTextBox(const FObjectInitializer& O
 	ReadOnlyForegroundColor_DEPRECATED = FLinearColor::Black;
 
 	// HACK: THIS SHOULD NOT COME FROM CORESTYLE AND SHOULD INSTEAD BY DEFINED BY ENGINE TEXTURES/PROJECT SETTINGS
-	WidgetStyle = FCoreStyle::Get().GetWidgetStyle<FEditableTextBoxStyle>("NormalEditableTextBox");
-	TextStyle = FCoreStyle::Get().GetWidgetStyle<FTextBlockStyle>("NormalText");
+	static const FEditableTextBoxStyle StaticNormalEditableTextBox = FCoreStyle::Get().GetWidgetStyle<FEditableTextBoxStyle>("NormalEditableTextBox");
+	static const FTextBlockStyle StaticNormalText = FCoreStyle::Get().GetWidgetStyle<FTextBlockStyle>("NormalText");
+	WidgetStyle = StaticNormalEditableTextBox;
+	TextStyle = StaticNormalText;
 
 	bIsReadOnly = false;
 	AllowContextMenu = true;
diff --git a/Engine/Source/Runtime/UMG/Private/Components/ProgressBar.cpp b/Engine/Source/Runtime/UMG/Private/Components/ProgressBar.cpp
--- a/Engine/Source/Runtime/UMG/Private/Components/ProgressBar.cpp
+++ b/Engine/Source/Runtime/UMG/Private/Components/ProgressBar.cpp
@@ -12,7 +12,9 @@ UProgressBar::UProgressBar(const FObjectInitializer& ObjectInitializer)
 	: Super(ObjectInitializer)
 {
 	// HACK: THIS SHOULD NOT COME FROM CORESTYLE AND SHOULD INSTEAD BY DEFINED BY ENGINE TEXTURES/PROJECT SETTINGS
-	WidgetStyle = FCoreStyle::Get().GetWidgetStyle<FProgressBarStyle>("ProgressBar");
+	static const FProgressBarStyle StaticProgressBar = FCoreStyle::Get().GetWidgetStyle<FProgressBarStyle>("ProgressBar");
+
+	WidgetStyle = StaticProgressBar;
 	WidgetStyle.FillImage.TintColor = FLinearColor::White;
 
 	BarFillType = EProgressBarFillType::LeftToRight;
diff --git a/Engine/Source/Runtime/UMG/Private/Components/ScrollBar.cpp b/Engine/Source/Runtime/UMG/Private/Components/ScrollBar.cpp
--- a/Engine/Source/Runtime/UMG/Private/Components/ScrollBar.cpp
+++ b/Engine/Source/Runtime/UMG/Private/Components/ScrollBar.cpp
@@ -18,7 +18,8 @@ UScrollBar::UScrollBar(const FObjectInitializer& ObjectInitializer)
 	Thickness = FVector2D(12.0f, 12.0f);
 
 	// HACK: THIS SHOULD NOT COME FROM CORESTYLE AND SHOULD INSTEAD BY DEFINED BY ENGINE TEXTURES/PROJECT SETTINGS
-	WidgetStyle = FCoreStyle::Get().GetWidgetStyle<FScrollBarStyle>("Scrollbar");
+	static const FScrollBarStyle StaticScrollbar = FCoreStyle::Get().GetWidgetStyle<FScrollBarStyle>("Scrollbar");
+	WidgetStyle = StaticScrollbar;
 }
 
 void UScrollBar::ReleaseSlateResources(bool bReleaseChildren)
diff --git a/Engine/Source/Runtime/UMG/Private/Components/ScrollBox.cpp b/Engine/Source/Runtime/UMG/Private/Components/ScrollBox.cpp
--- a/Engine/Source/Runtime/UMG/Private/Components/ScrollBox.cpp
+++ b/Engine/Source/Runtime/UMG/Private/Components/ScrollBox.cpp
@@ -26,9 +26,15 @@ UScrollBox::UScrollBox(const FObjectInitializer& ObjectInitializer)
 	Visibility = ESlateVisibility::Visible;
 	Clipping = EWidgetClipping::ClipToBounds;
 
-	// HACK: THIS SHOULD NOT COME FROM CORESTYLE AND SHOULD INSTEAD BY DEFINED BY ENGINE TEXTURES/PROJECT SETTINGS
-	WidgetStyle = FCoreStyle::Get().GetWidgetStyle<FScrollBoxStyle>("ScrollBox");
-	WidgetBarStyle = FCoreStyle::Get().GetWidgetStyle<FScrollBarStyle>("ScrollBar");
+	{
+		// HACK: THIS SHOULD NOT COME FROM CORESTYLE AND SHOULD INSTEAD BY DEFINED BY ENGINE TEXTURES/PROJECT SETTINGS
+		static const FScrollBoxStyle StaticScrollBoxStyle = FCoreStyle::Get().GetWidgetStyle<FScrollBoxStyle>("ScrollBox");
+		static const FScrollBarStyle StaticScrollBarStyle = FCoreStyle::Get().GetWidgetStyle<FScrollBarStyle>("ScrollBar");
+
+		WidgetStyle = StaticScrollBoxStyle;
+		WidgetBarStyle = StaticScrollBarStyle;
+	}
+
 	bAllowRightClickDragScrolling = true;
 }
 
diff --git a/Engine/Source/Runtime/UMG/Private/Components/Slider.cpp b/Engine/Source/Runtime/UMG/Private/Components/Slider.cpp
--- a/Engine/Source/Runtime/UMG/Private/Components/Slider.cpp
+++ b/Engine/Source/Runtime/UMG/Private/Components/Slider.cpp
@@ -21,7 +21,8 @@ USlider::USlider(const FObjectInitializer& ObjectInitializer)
 	RequiresControllerLock = true;
 
 	// HACK: THIS SHOULD NOT COME FROM CORESTYLE AND SHOULD INSTEAD BY DEFINED BY ENGINE TEXTURES/PROJECT SETTINGS
-	WidgetStyle = FCoreStyle::Get().GetWidgetStyle<FSliderStyle>("Slider");
+	static const FSliderStyle StaticSlider = FCoreStyle::Get().GetWidgetStyle<FSliderStyle>("Slider");
+	WidgetStyle = StaticSlider;
 }
 
 TSharedRef<SWidget> USlider::RebuildWidget()
diff --git a/Engine/Source/Runtime/UMG/Private/Components/SpinBox.cpp b/Engine/Source/Runtime/UMG/Private/Components/SpinBox.cpp
--- a/Engine/Source/Runtime/UMG/Private/Components/SpinBox.cpp
+++ b/Engine/Source/Runtime/UMG/Private/Components/SpinBox.cpp
@@ -31,7 +31,8 @@ USpinBox::USpinBox(const FObjectInitializer& ObjectInitializer)
 	ForegroundColor = FSlateColor(FLinearColor::Black);
 
 	// HACK: THIS SHOULD NOT COME FROM CORESTYLE AND SHOULD INSTEAD BY DEFINED BY ENGINE TEXTURES/PROJECT SETTINGS
-	WidgetStyle = FCoreStyle::Get().GetWidgetStyle<FSpinBoxStyle>("SpinBox");
+	static const FSpinBoxStyle StaticSpinBox = FCoreStyle::Get().GetWidgetStyle<FSpinBoxStyle>("SpinBox");
+	WidgetStyle = StaticSpinBox;
 }
 
 void USpinBox::ReleaseSlateResources(bool bReleaseChildren)
diff --git a/Engine/Source/Runtime/UMG/Private/Components/Throbber.cpp b/Engine/Source/Runtime/UMG/Private/Components/Throbber.cpp
--- a/Engine/Source/Runtime/UMG/Private/Components/Throbber.cpp
+++ b/Engine/Source/Runtime/UMG/Private/Components/Throbber.cpp
@@ -19,7 +19,8 @@ UThrobber::UThrobber(const FObjectInitializer& ObjectInitializer)
 	bAnimateOpacity = true;
 
 	// HACK: THIS SHOULD NOT COME FROM CORESTYLE AND SHOULD INSTEAD BY DEFINED BY ENGINE TEXTURES/PROJECT SETTINGS
-	Image = *FCoreStyle::Get().GetBrush("Throbber.Chunk");
+	static const FSlateBrush StaticImage = *FCoreStyle::Get().GetBrush("Throbber.Chunk");
+	Image = StaticImage;
 }
 
 void UThrobber::ReleaseSlateResources(bool bReleaseChildren)
