Exposed SScrollBox::EndInertialScrolling publicly to allow manual interruption of the inertial scroll.
- Added corresponding UScrollBox::EndInertialScroll as well

diff --git a/Engine/Source/Runtime/UMG/Private/Components/ScrollBox.cpp b/Engine/Source/Runtime/UMG/Private/Components/ScrollBox.cpp
--- a/Engine/Source/Runtime/UMG/Private/Components/ScrollBox.cpp
+++ b/Engine/Source/Runtime/UMG/Private/Components/ScrollBox.cpp
@@ -269,6 +269,14 @@ void UScrollBox::SetAllowOverscroll(bool NewAllowOverscroll)
 	}
 }
 
+void UScrollBox::EndInertialScrolling()
+{
+	if (MyScrollBox.IsValid())
+	{
+		MyScrollBox->EndInertialScrolling();
+	}
+}
+
 void UScrollBox::SlateHandleUserScrolled(float CurrentOffset)
 {
 	OnUserScrolled.Broadcast(CurrentOffset);
diff --git a/Engine/Source/Runtime/UMG/Public/Components/ScrollBox.h b/Engine/Source/Runtime/UMG/Public/Components/ScrollBox.h
--- a/Engine/Source/Runtime/UMG/Public/Components/ScrollBox.h
+++ b/Engine/Source/Runtime/UMG/Public/Components/ScrollBox.h
@@ -97,6 +97,11 @@ public:
 	
 	UFUNCTION(BlueprintCallable, Category = "Scroll")
 	void SetAllowOverscroll(bool NewAllowOverscroll);
+
+	/** Instantly stops any inertial scrolling that is currently in progress */
+	UFUNCTION(BlueprintCallable, Category = "Scroll")
+	void EndInertialScrolling();
+
 public:
 
 	/** Called when the scroll has changed */
