Updated copyright date on DynamicEntryBoxBase

diff --git a/Engine/Source/Runtime/UMG/Private/Components/DynamicEntryBoxBase.cpp b/Engine/Source/Runtime/UMG/Private/Components/DynamicEntryBoxBase.cpp
--- a/Engine/Source/Runtime/UMG/Private/Components/DynamicEntryBoxBase.cpp
+++ b/Engine/Source/Runtime/UMG/Private/Components/DynamicEntryBoxBase.cpp
@@ -1,4 +1,4 @@
-// Copyright 1998-2018 Epic Games, Inc. All Rights Reserved.
+// Copyright 1998-2019 Epic Games, Inc. All Rights Reserved.
 
 #include "Components/DynamicEntryBoxBase.h"
 #include "UMGPrivate.h"
diff --git a/Engine/Source/Runtime/UMG/Public/Components/DynamicEntryBoxBase.h b/Engine/Source/Runtime/UMG/Public/Components/DynamicEntryBoxBase.h
--- a/Engine/Source/Runtime/UMG/Public/Components/DynamicEntryBoxBase.h
+++ b/Engine/Source/Runtime/UMG/Public/Components/DynamicEntryBoxBase.h
@@ -1,4 +1,4 @@
-// Copyright 1998-2018 Epic Games, Inc. All Rights Reserved.
+// Copyright 1998-2019 Epic Games, Inc. All Rights Reserved.
 
 #pragma once
 
