Fixed ListViewBase's pool not being a UPROPERTY and allowing items to get GC'd

diff --git a/Engine/Source/Runtime/UMG/Public/Components/ListViewBase.h b/Engine/Source/Runtime/UMG/Public/Components/ListViewBase.h
--- a/Engine/Source/Runtime/UMG/Public/Components/ListViewBase.h
+++ b/Engine/Source/Runtime/UMG/Public/Components/ListViewBase.h
@@ -569,7 +569,9 @@ private:
 	int32 NumDesignerPreviewEntries = 5;
 #endif
 
+	UPROPERTY(Transient)
 	FUserWidgetPool EntryWidgetPool;
+
 	FTimerHandle EntryGenAnnouncementTimerHandle;
 	TArray<TWeakObjectPtr<UUserWidget>> GeneratedEntriesToAnnounce;
 
