UMG - The PanelSlot's IsDesignTime function should have a FORCEINLINE constant one like we have in UWidget.

diff --git a/Engine/Source/Runtime/UMG/Private/Components/PanelSlot.cpp b/Engine/Source/Runtime/UMG/Private/Components/PanelSlot.cpp
--- a/Engine/Source/Runtime/UMG/Private/Components/PanelSlot.cpp
+++ b/Engine/Source/Runtime/UMG/Private/Components/PanelSlot.cpp
@@ -11,10 +11,12 @@ UPanelSlot::UPanelSlot(const FObjectInitializer& ObjectInitializer)
 {
 }
 
+#if WITH_EDITOR
 bool UPanelSlot::IsDesignTime() const
 {
 	return Parent->IsDesignTime();
 }
+#endif
 
 void UPanelSlot::ReleaseSlateResources(bool bReleaseChildren)
 {
diff --git a/Engine/Source/Runtime/UMG/Public/Components/PanelSlot.h b/Engine/Source/Runtime/UMG/Public/Components/PanelSlot.h
--- a/Engine/Source/Runtime/UMG/Public/Components/PanelSlot.h
+++ b/Engine/Source/Runtime/UMG/Public/Components/PanelSlot.h
@@ -21,7 +21,11 @@ public:
 	UPROPERTY(Instanced)
 	class UWidget* Content;
 	
+#if WITH_EDITOR
 	bool IsDesignTime() const;
+#else
+	FORCEINLINE bool IsDesignTime() const { return false; }
+#endif
 
 	virtual void ReleaseSlateResources(bool bReleaseChildren) override;
 
