UMG - Giving the rich textblock the ability to have the text style sets swapped out at runtime, needed to add the ability to the base RichTextLayouttMarshaller to have the decorator style set swapped out to allow for dynamic rethemeing of rich text content.

Vincent.Gauthier, Jamie.Dale

diff --git a/Engine/Source/Runtime/Slate/Private/Widgets/Text/SRichTextBlock.cpp b/Engine/Source/Runtime/Slate/Private/Widgets/Text/SRichTextBlock.cpp
--- a/Engine/Source/Runtime/Slate/Private/Widgets/Text/SRichTextBlock.cpp
+++ b/Engine/Source/Runtime/Slate/Private/Widgets/Text/SRichTextBlock.cpp
@@ -39,7 +39,7 @@ void SRichTextBlock::Construct( const FArguments& InArgs )
 			Parser = FDefaultRichTextMarkupParser::GetStaticInstance();
 		}
 
-		TSharedPtr<FRichTextLayoutMarshaller> Marshaller = InArgs._Marshaller;
+		Marshaller = InArgs._Marshaller;
 		if (!Marshaller.IsValid())
 		{
 			Marshaller = FRichTextLayoutMarshaller::Create(Parser, nullptr, InArgs._Decorators, InArgs._DecoratorStyleSet);
@@ -166,6 +166,15 @@ void SRichTextBlock::SetMinDesiredWidth(const TAttribute<float>& InMinDesiredWid
 	Invalidate(EInvalidateWidget::LayoutAndVolatility);
 }
 
+void SRichTextBlock::SetDecoratorStyleSet(const ISlateStyle* NewDecoratorStyleSet)
+{
+	if (Marshaller.IsValid())
+	{
+		Marshaller->SetDecoratorStyleSet(NewDecoratorStyleSet);
+		Refresh();
+	}
+}
+
 void SRichTextBlock::Refresh()
 {
 	TextLayoutCache->DirtyContent();
diff --git a/Engine/Source/Runtime/Slate/Public/Framework/Text/RichTextLayoutMarshaller.h b/Engine/Source/Runtime/Slate/Public/Framework/Text/RichTextLayoutMarshaller.h
--- a/Engine/Source/Runtime/Slate/Public/Framework/Text/RichTextLayoutMarshaller.h
+++ b/Engine/Source/Runtime/Slate/Public/Framework/Text/RichTextLayoutMarshaller.h
@@ -44,6 +44,11 @@ public:
 		return *this;
 	}
 
+	inline void SetDecoratorStyleSet(const ISlateStyle* InStyleSet)
+	{
+		DecoratorStyleSet = InStyleSet;
+	}
+
 protected:
 
 	FRichTextLayoutMarshaller(TArray< TSharedRef< ITextDecorator > > InDecorators, const ISlateStyle* const InDecoratorStyleSet);
diff --git a/Engine/Source/Runtime/Slate/Public/Widgets/Text/SRichTextBlock.h b/Engine/Source/Runtime/Slate/Public/Widgets/Text/SRichTextBlock.h
--- a/Engine/Source/Runtime/Slate/Public/Widgets/Text/SRichTextBlock.h
+++ b/Engine/Source/Runtime/Slate/Public/Widgets/Text/SRichTextBlock.h
@@ -208,6 +208,9 @@ public:
 	/** See MinDesiredWidth attribute */
 	void SetMinDesiredWidth(const TAttribute<float>& InMinDesiredWidth);
 
+	/**  */
+	void SetDecoratorStyleSet(const ISlateStyle* NewDecoratorStyleSet);
+
 	/**
 	 * Causes the text to reflow it's layout
 	 */
@@ -253,6 +256,9 @@ private:
 
 	/** Prevents the text block from being smaller than desired in certain cases (e.g. when it is empty) */
 	TAttribute<float> MinDesiredWidth;
+
+	/**  */
+	TSharedPtr<FRichTextLayoutMarshaller> Marshaller;
 };
 
 #endif //WITH_FANCY_TEXT
diff --git a/Engine/Source/Runtime/UMG/Private/Components/RichTextBlock.cpp b/Engine/Source/Runtime/UMG/Private/Components/RichTextBlock.cpp
--- a/Engine/Source/Runtime/UMG/Private/Components/RichTextBlock.cpp
+++ b/Engine/Source/Runtime/UMG/Private/Components/RichTextBlock.cpp
@@ -11,12 +11,37 @@
 #include "Framework/Text/RichTextMarkupProcessing.h"
 #include "Framework/Text/IRichTextMarkupParser.h"
 #include "Framework/Text/IRichTextMarkupWriter.h"
+#include "RenderingThread.h"
 
 #define LOCTEXT_NAMESPACE "UMG"
 
 /////////////////////////////////////////////////////
 // URichTextBlock
 
+template< class ObjectType >
+struct FDeferredDeletor : public FDeferredCleanupInterface
+{
+public:
+	FDeferredDeletor(ObjectType* InInnerObjectToDelete)
+		: InnerObjectToDelete(InInnerObjectToDelete)
+	{
+	}
+
+	virtual ~FDeferredDeletor()
+	{
+		delete InnerObjectToDelete;
+	}
+
+private:
+	ObjectType* InnerObjectToDelete;
+};
+
+template< class ObjectType >
+FORCEINLINE SharedPointerInternals::FRawPtrProxy< ObjectType > MakeShareableDeferredCleanup(ObjectType* InObject)
+{
+	return MakeShareable(InObject, [](ObjectType* ObjectToDelete) { BeginCleanup(new FDeferredDeletor<ObjectType>(ObjectToDelete)); });
+}
+
 URichTextBlock::URichTextBlock(const FObjectInitializer& ObjectInitializer)
 	: Super(ObjectInitializer)
 {
@@ -66,23 +91,7 @@ void URichTextBlock::UpdateStyleData()
 
 	if (!StyleInstance.IsValid())
 	{
-		StyleInstance = MakeShareable(new FSlateStyleSet(TEXT("RichTextStyle")));
-
-		if (TextStyleSet && TextStyleSet->GetRowStruct()->IsChildOf(FRichTextStyleRow::StaticStruct()))
-		{
-			for (const auto& Entry : TextStyleSet->GetRowMap())
-			{
-				FName SubStyleName = Entry.Key;
-				FRichTextStyleRow* RichTextStyle = (FRichTextStyleRow*)Entry.Value;
-
-				if (SubStyleName == FName(TEXT("Default")))
-				{
-					DefaultTextStyle = RichTextStyle->TextStyle;
-				}
-
-				StyleInstance->Set(SubStyleName, RichTextStyle->TextStyle);
-			}
-		}
+		RebuildStyleInstance();
 
 		for (TSubclassOf<URichTextBlockDecorator> DecoratorClass : DecoratorClasses)
 		{
@@ -107,6 +116,40 @@ void URichTextBlock::SetText(const FText& InText)
 	}
 }
 
+void URichTextBlock::RebuildStyleInstance()
+{
+	StyleInstance = MakeShareableDeferredCleanup(new FSlateStyleSet(TEXT("RichTextStyle")));
+
+	if (TextStyleSet && ensure(TextStyleSet->GetRowStruct()->IsChildOf(FRichTextStyleRow::StaticStruct())))
+	{
+		for (const auto& Entry : TextStyleSet->GetRowMap())
+		{
+			FName SubStyleName = Entry.Key;
+			FRichTextStyleRow* RichTextStyle = (FRichTextStyleRow*)Entry.Value;
+
+			if (SubStyleName == FName(TEXT("Default")))
+			{
+				DefaultTextStyle = RichTextStyle->TextStyle;
+			}
+
+			StyleInstance->Set(SubStyleName, RichTextStyle->TextStyle);
+		}
+	}
+}
+
+void URichTextBlock::SetTextStyleSet(class UDataTable* NewTextStyleSet)
+{
+	if (TextStyleSet != NewTextStyleSet)
+	{
+		TextStyleSet = NewTextStyleSet;
+
+		RebuildStyleInstance();
+
+		MyRichTextBlock->SetDecoratorStyleSet(StyleInstance.Get());
+		MyRichTextBlock->SetTextStyle(DefaultTextStyle);
+	}
+}
+
 const FTextBlockStyle& URichTextBlock::GetDefaultTextStyle() const
 {
 	ensure(StyleInstance.IsValid());
diff --git a/Engine/Source/Runtime/UMG/Public/Components/RichTextBlock.h b/Engine/Source/Runtime/UMG/Public/Components/RichTextBlock.h
--- a/Engine/Source/Runtime/UMG/Public/Components/RichTextBlock.h
+++ b/Engine/Source/Runtime/UMG/Public/Components/RichTextBlock.h
@@ -132,6 +132,9 @@ public:
 	UFUNCTION(BlueprintCallable, Category = "Widget")
 	virtual void SetText(const FText& InText);
 
+	UFUNCTION(BlueprintCallable, Category = "Widget")
+	void SetTextStyleSet(class UDataTable* NewTextStyleSet);
+
 	const FTextBlockStyle& GetDefaultTextStyle() const;
 	const FTextBlockStyle& GetCurrentDefaultTextStyle() const;
 
@@ -143,6 +146,7 @@ protected:
 	virtual TSharedRef<SWidget> RebuildWidget() override;
 	
 	virtual void UpdateStyleData();
+	void RebuildStyleInstance();
 	virtual void CreateDecorators(TArray< TSharedRef< class ITextDecorator > >& OutDecorators);
 	virtual TSharedPtr< class IRichTextMarkupParser > CreateMarkupParser();
 	virtual TSharedPtr< class IRichTextMarkupWriter > CreateMarkupWriter();
