Fixing python warning

diff --git a/Engine/Source/Runtime/UMG/Public/Components/MenuAnchor.h b/Engine/Source/Runtime/UMG/Public/Components/MenuAnchor.h
--- a/Engine/Source/Runtime/UMG/Public/Components/MenuAnchor.h
+++ b/Engine/Source/Runtime/UMG/Public/Components/MenuAnchor.h
@@ -43,7 +43,7 @@ public:
 	TEnumAsByte<EMenuPlacement> Placement;
 
 	/** Should the menu anchor attempt to fit the menu inside the window. */
-	UPROPERTY(EditAnywhere, BlueprintReadOnly, Category = "Menu Anchor")
+	UPROPERTY(EditAnywhere, BlueprintReadOnly, Category = "Menu Anchor", meta=(ScriptName="ShouldFitInWindow"))
 	bool bFitInWindow;
 
 	UPROPERTY(EditAnywhere, BlueprintReadOnly, AdvancedDisplay, Category = "Menu Anchor")
