New callback functionality for list views: OnListViewScrolled() - fires when the user scrolls the list view
Different from OnItemScrolledIntoView() - which is fired when an specific item is told to get scrolled into view.

diff --git a/Engine/Source/Runtime/UMG/Public/Components/ListViewBase.h b/Engine/Source/Runtime/UMG/Public/Components/ListViewBase.h
--- a/Engine/Source/Runtime/UMG/Public/Components/ListViewBase.h
+++ b/Engine/Source/Runtime/UMG/Public/Components/ListViewBase.h
@@ -38,20 +38,23 @@ public:
 	//////////////////////////////////////////////////////////////////////////
 	// Automatically implemented via IMPLEMENT_TYPED_UMG_LIST()
 	//////////////////////////////////////////////////////////////////////////
-	DECLARE_EVENT_OneParam(UListView, FSimpleListItemEvent, ItemType);
+	DECLARE_MULTICAST_DELEGATE_OneParam(FSimpleListItemEvent, ItemType);
 	virtual FSimpleListItemEvent& OnItemClicked() const = 0;
 	virtual FSimpleListItemEvent& OnItemDoubleClicked() const = 0;
 
-	DECLARE_EVENT_TwoParams(UListView, FOnItemIsHoveredChanged, ItemType, bool);
+	DECLARE_MULTICAST_DELEGATE_TwoParams(FOnItemIsHoveredChanged, ItemType, bool);
 	virtual FOnItemIsHoveredChanged& OnItemIsHoveredChanged() const = 0;
 
-	DECLARE_EVENT_OneParam(UListView, FOnItemSelectionChanged, NullableItemType);
+	DECLARE_MULTICAST_DELEGATE_OneParam(FOnItemSelectionChanged, NullableItemType);
 	virtual FOnItemSelectionChanged& OnItemSelectionChanged() const = 0;
 
-	DECLARE_EVENT_TwoParams(UListView, FOnItemScrolledIntoView, ItemType, UUserWidget&);
+	DECLARE_MULTICAST_DELEGATE_TwoParams(FOnListViewScrolled, float, float);
+	virtual FOnListViewScrolled& OnListViewScrolled() const = 0;
+
+	DECLARE_MULTICAST_DELEGATE_TwoParams(FOnItemScrolledIntoView, ItemType, UUserWidget&);
 	virtual FOnItemScrolledIntoView& OnItemScrolledIntoView() const = 0;
 
-	DECLARE_EVENT_TwoParams(UTreeView, FOnItemExpansionChanged, ItemType, bool);
+	DECLARE_MULTICAST_DELEGATE_TwoParams(FOnItemExpansionChanged, ItemType, bool);
 	virtual FOnItemExpansionChanged& OnItemExpansionChanged() const = 0;
 
 	DECLARE_DELEGATE_RetVal_OneParam(TSubclassOf<UUserWidget>, FOnGetEntryClassForItem, ItemType);
@@ -228,6 +231,7 @@ protected:
 			.OnSelectionChanged_UObject(Implementer, &UListViewBaseT::HandleSelectionChanged)
 			.OnRowReleased_UObject(Implementer, &UListViewBaseT::HandleRowReleased)
 			.OnItemScrolledIntoView_UObject(Implementer, &UListViewBaseT::HandleItemScrolledIntoView)
+			.OnListViewScrolled_UObject(Implementer, &UListViewBaseT::HandleListViewScrolled)
 			.OnMouseButtonClick_UObject(Implementer, &UListViewBaseT::HandleItemClicked)
 			.OnMouseButtonDoubleClick_UObject(Implementer, &UListViewBaseT::HandleItemDoubleClicked);
 	}
@@ -258,6 +262,7 @@ protected:
 			.OnTileReleased_UObject(Implementer, &UListViewBaseT::HandleRowReleased)
 			.OnSelectionChanged_UObject(Implementer, &UListViewBaseT::HandleSelectionChanged)
 			.OnItemScrolledIntoView_UObject(Implementer, &UListViewBaseT::HandleItemScrolledIntoView)
+			.OnTileViewScrolled_UObject(Implementer, &UListViewBaseT::HandleListViewScrolled)
 			.OnMouseButtonClick_UObject(Implementer, &UListViewBaseT::HandleItemClicked)
 			.OnMouseButtonDoubleClick_UObject(Implementer, &UListViewBaseT::HandleItemDoubleClicked);
 	}
@@ -280,6 +285,7 @@ protected:
 			.OnSelectionChanged_UObject(Implementer, &UListViewBaseT::HandleSelectionChanged)
 			.OnRowReleased_UObject(Implementer, &UListViewBaseT::HandleRowReleased)
 			.OnItemScrolledIntoView_UObject(Implementer, &UListViewBaseT::HandleItemScrolledIntoView)
+			.OnTreeViewScrolled_UObject(Implementer, &UListViewBaseT::HandleListViewScrolled)
 			.OnMouseButtonClick_UObject(Implementer, &UListViewBaseT::HandleItemClicked)
 			.OnMouseButtonDoubleClick_UObject(Implementer, &UListViewBaseT::HandleItemDoubleClicked)
 			.OnGetChildren_UObject(Implementer, &UListViewBaseT::HandleGetChildren)
@@ -317,6 +323,7 @@ protected:
 	virtual void OnItemDoubleClickedInternal(ItemType Item) {}
 	virtual void OnSelectionChangedInternal(NullableItemType FirstSelectedItem) {}
 	virtual void OnItemScrolledIntoViewInternal(ItemType Item, UUserWidget& EntryWidget) {}
+	virtual void OnListViewScrolledInternal(float ItemOffset, float DistanceRemaining) {}
 	virtual void OnItemExpansionChangedInternal(ItemType Item, bool bIsExpanded) {}
 
 private:
@@ -351,6 +358,16 @@ private:
 		OnItemSelectionChanged().Broadcast(Item);
 	}
 
+	void HandleListViewScrolled(double OffsetInItems)
+	{
+		if (SListView<ItemType>* MyListView = GetMyListView())
+		{
+			const FVector2D DistanceRemaining = MyListView->GetScrollDistanceRemaining();
+			OnListViewScrolledInternal(OffsetInItems, DistanceRemaining.Y);
+			OnListViewScrolled().Broadcast(OffsetInItems, DistanceRemaining.Y);
+		}
+	}
+
 	void HandleItemScrolledIntoView(ItemType Item, const TSharedPtr<ITableRow>& InWidget)
 	{
 		UUserWidget* RowWidget = GetEntryWidgetFromItem(Item);
@@ -600,6 +617,7 @@ private:	\
 	mutable FOnItemSelectionChanged OnItemSelectionChangedEvent;	\
 	mutable FOnItemIsHoveredChanged OnItemIsHoveredChangedEvent;	\
 	mutable FOnItemScrolledIntoView OnItemScrolledIntoViewEvent;	\
+	mutable FOnListViewScrolled OnListViewScrolledEvent;	\
 	mutable FOnItemExpansionChanged OnItemExpansionChangedEvent;	\
 	mutable FOnGetEntryClassForItem OnGetEntryClassForItemDelegate;	\
 public:	\
@@ -609,5 +627,6 @@ public:	\
 	virtual FOnItemIsHoveredChanged& OnItemIsHoveredChanged() const override { return OnItemIsHoveredChangedEvent; }	\
 	virtual FOnItemSelectionChanged& OnItemSelectionChanged() const override { return OnItemSelectionChangedEvent; }	\
 	virtual FOnItemScrolledIntoView& OnItemScrolledIntoView() const override { return OnItemScrolledIntoViewEvent; }	\
+	virtual FOnListViewScrolled& OnListViewScrolled() const override { return OnListViewScrolledEvent; }	\
 	virtual FOnItemExpansionChanged& OnItemExpansionChanged() const override { return OnItemExpansionChangedEvent; }	\
 	virtual FOnGetEntryClassForItem& OnGetEntryClassForItem() const override { return OnGetEntryClassForItemDelegate; }
\ No newline at end of file
