Merging CL 6320843 from //UE4/Dev-Editor to //UE4/Main

Fix functional test UI screenshot when running with a window dpi scale != 1.

diff --git a/Engine/Source/Developer/FunctionalTesting/FunctionalTesting.Build.cs b/Engine/Source/Developer/FunctionalTesting/FunctionalTesting.Build.cs
--- a/Engine/Source/Developer/FunctionalTesting/FunctionalTesting.Build.cs
+++ b/Engine/Source/Developer/FunctionalTesting/FunctionalTesting.Build.cs
@@ -13,6 +13,7 @@ public class FunctionalTesting : ModuleRules
                 "Engine",
                 "RenderCore",
                 "Slate",
+				"SlateCore",
                 "MessageLog",
                 "NavigationSystem",
                 "AIModule",
@@ -26,7 +27,12 @@ public class FunctionalTesting : ModuleRules
 
         if (Target.bBuildEditor == true)
 		{
-			PrivateDependencyModuleNames.Add("UnrealEd");
+			PrivateDependencyModuleNames.AddRange(
+				new string[] {
+					"SourceControl",
+					"UnrealEd"
+				}
+			);
 		}
 
         PrivateIncludePaths.AddRange(
@@ -36,15 +42,6 @@ public class FunctionalTesting : ModuleRules
             }
         );
 
-        if (Target.bBuildEditor == true)
-        {
-            PrivateDependencyModuleNames.AddRange(
-                new string[] {
-                    "SourceControl"
-                }
-            );
-        }
-
 		//make sure this is compiled for binary builds
 		if (Target.Configuration != UnrealTargetConfiguration.Shipping)
 		{
diff --git a/Engine/Source/Developer/FunctionalTesting/Private/FunctionalUIScreenshotTest.cpp b/Engine/Source/Developer/FunctionalTesting/Private/FunctionalUIScreenshotTest.cpp
--- a/Engine/Source/Developer/FunctionalTesting/Private/FunctionalUIScreenshotTest.cpp
+++ b/Engine/Source/Developer/FunctionalTesting/Private/FunctionalUIScreenshotTest.cpp
@@ -173,7 +173,7 @@ void AFunctionalUIScreenshotTest::RequestScreenshot()
 			// when rendering to a separate render target
 			FWidgetRenderer* WidgetRenderer = new FWidgetRenderer(true, false);
 			check(WidgetRenderer);
-			WidgetRenderer->DrawWidget(ScreenshotRT, ViewportWidget.ToSharedRef(), ScreenshotSize, 0.f);
+			WidgetRenderer->DrawWidget(ScreenshotRT, ViewportWidget.ToSharedRef(), ViewportWidget->GetCachedGeometry().Scale, ScreenshotSize, 0.f);
 			FlushRenderingCommands();
 			BeginCleanup(WidgetRenderer);
 
diff --git a/Engine/Source/Runtime/UMG/Private/Slate/WidgetRenderer.cpp b/Engine/Source/Runtime/UMG/Private/Slate/WidgetRenderer.cpp
--- a/Engine/Source/Runtime/UMG/Private/Slate/WidgetRenderer.cpp
+++ b/Engine/Source/Runtime/UMG/Private/Slate/WidgetRenderer.cpp
@@ -87,6 +87,22 @@ UTextureRenderTarget2D* FWidgetRenderer::CreateTargetFor(FVector2D DrawSize, Tex
 }
 
 void FWidgetRenderer::DrawWidget(FRenderTarget* RenderTarget, const TSharedRef<SWidget>& Widget, FVector2D DrawSize, float DeltaTime, bool bDeferRenderTargetUpdate)
+{
+	DrawWidget(RenderTarget, Widget, 1.f, DrawSize, DeltaTime, bDeferRenderTargetUpdate);
+}
+
+void FWidgetRenderer::DrawWidget(UTextureRenderTarget2D* RenderTarget, const TSharedRef<SWidget>& Widget, FVector2D DrawSize, float DeltaTime, bool bDeferRenderTargetUpdate)
+{
+	DrawWidget(RenderTarget->GameThread_GetRenderTargetResource(), Widget, DrawSize, DeltaTime, bDeferRenderTargetUpdate);
+}
+
+void FWidgetRenderer::DrawWidget(
+	FRenderTarget* RenderTarget,
+	const TSharedRef<SWidget>& Widget,
+	float Scale,
+	FVector2D DrawSize,
+	float DeltaTime,
+	bool bDeferRenderTargetUpdate)
 {
 	TSharedRef<SVirtualWindow> Window = SNew(SVirtualWindow).Size(DrawSize);
 	TSharedRef<FHittestGrid> HitTestGrid = MakeShareable(new FHittestGrid());
@@ -94,12 +110,18 @@ void FWidgetRenderer::DrawWidget(FRenderTarget* RenderTarget, const TSharedRef<S
 	Window->SetContent(Widget);
 	Window->Resize(DrawSize);
 
-	DrawWindow(RenderTarget, HitTestGrid, Window, 1, DrawSize, DeltaTime, bDeferRenderTargetUpdate);
+	DrawWindow(RenderTarget, HitTestGrid, Window, Scale, DrawSize, DeltaTime, bDeferRenderTargetUpdate);
 }
 
-void FWidgetRenderer::DrawWidget(UTextureRenderTarget2D* RenderTarget, const TSharedRef<SWidget>& Widget, FVector2D DrawSize, float DeltaTime, bool bDeferRenderTargetUpdate)
+void FWidgetRenderer::DrawWidget(
+	UTextureRenderTarget2D* RenderTarget,
+	const TSharedRef<SWidget>& Widget,
+	float Scale,
+	FVector2D DrawSize,
+	float DeltaTime,
+	bool bDeferRenderTargetUpdate)
 {
-	DrawWidget(RenderTarget->GameThread_GetRenderTargetResource(), Widget, DrawSize, DeltaTime, bDeferRenderTargetUpdate);
+	DrawWidget(RenderTarget->GameThread_GetRenderTargetResource(), Widget, Scale, DrawSize, DeltaTime, bDeferRenderTargetUpdate);
 }
 
 void FWidgetRenderer::DrawWindow(
diff --git a/Engine/Source/Runtime/UMG/Public/Slate/WidgetRenderer.h b/Engine/Source/Runtime/UMG/Public/Slate/WidgetRenderer.h
--- a/Engine/Source/Runtime/UMG/Public/Slate/WidgetRenderer.h
+++ b/Engine/Source/Runtime/UMG/Public/Slate/WidgetRenderer.h
@@ -61,6 +61,22 @@ public:
 		float DeltaTime,
 		bool bDeferRenderTargetUpdate = false);
 
+	void DrawWidget(
+		FRenderTarget* RenderTarget,
+		const TSharedRef<SWidget>& Widget,
+		float Scale,
+		FVector2D DrawSize,
+		float DeltaTime,
+		bool bDeferRenderTargetUpdate = false);
+
+	void DrawWidget(
+		UTextureRenderTarget2D* RenderTarget,
+		const TSharedRef<SWidget>& Widget,
+		float Scale,
+		FVector2D DrawSize,
+		float DeltaTime,
+		bool bDeferRenderTargetUpdate = false);
+
 	void DrawWindow(
 		FRenderTarget* RenderTarget,
 		TSharedRef<FHittestGrid> HitTestGrid,
