Exposed support for getting and setting hint text from code/BP for UMultiLineEditableText and UMultiLineEditableTextBox.

This functionality already existed on SMultiLineEditableText and SMultiLineEditableTextBox (setter only), but was not previously exposed by the UMG widgets except as an edit-only field.

C++:
1) Added UMultiLineEditableText::GetHintText() and SetHintText().
2) Added UMultiLineEditableTextBox::GetHintText() and SetHintText().
3) Added SMultiLineEditableTextBox::GetHintText().

diff --git a/Engine/Source/Runtime/Slate/Public/Widgets/Input/SMultiLineEditableTextBox.h b/Engine/Source/Runtime/Slate/Public/Widgets/Input/SMultiLineEditableTextBox.h
--- a/Engine/Source/Runtime/Slate/Public/Widgets/Input/SMultiLineEditableTextBox.h
+++ b/Engine/Source/Runtime/Slate/Public/Widgets/Input/SMultiLineEditableTextBox.h
@@ -267,6 +267,16 @@ public:
 	 */
 	void SetText( const TAttribute< FText >& InNewText );
 
+	/**
+	 * Returns the hint text string
+	 *
+	 * @return  Hint text string
+	 */
+	FText GetHintText() const
+	{
+		return EditableText->GetHintText();
+	}
+
 	/**
 	 * Sets the text that appears when there is no text in the text box
 	 *
diff --git a/Engine/Source/Runtime/UMG/Private/Components/MultiLineEditableText.cpp b/Engine/Source/Runtime/UMG/Private/Components/MultiLineEditableText.cpp
--- a/Engine/Source/Runtime/UMG/Private/Components/MultiLineEditableText.cpp
+++ b/Engine/Source/Runtime/UMG/Private/Components/MultiLineEditableText.cpp
@@ -113,6 +113,27 @@ void UMultiLineEditableText::SetText(FText InText)
 	}
 }
 
+FText UMultiLineEditableText::GetHintText() const
+{
+	if (MyMultiLineEditableText.IsValid())
+	{
+		return MyMultiLineEditableText->GetHintText();
+	}
+
+	return HintText;
+}
+
+void UMultiLineEditableText::SetHintText(FText InHintText)
+{
+	HintText = InHintText;
+	HintTextDelegate.Clear();
+	if (MyMultiLineEditableText.IsValid())
+	{
+		TAttribute<FText> HintTextBinding = PROPERTY_BINDING(FText, HintText);
+		MyMultiLineEditableText->SetHintText(HintTextBinding);
+	}
+}
+
 void UMultiLineEditableText::SetIsReadOnly(bool bReadOnly)
 {
 	bIsReadOnly = bReadOnly;
diff --git a/Engine/Source/Runtime/UMG/Private/Components/MultiLineEditableTextBox.cpp b/Engine/Source/Runtime/UMG/Private/Components/MultiLineEditableTextBox.cpp
--- a/Engine/Source/Runtime/UMG/Private/Components/MultiLineEditableTextBox.cpp
+++ b/Engine/Source/Runtime/UMG/Private/Components/MultiLineEditableTextBox.cpp
@@ -111,6 +111,27 @@ void UMultiLineEditableTextBox::SetText(FText InText)
 	}
 }
 
+FText UMultiLineEditableTextBox::GetHintText() const
+{
+	if (MyEditableTextBlock.IsValid())
+	{
+		return MyEditableTextBlock->GetHintText();
+	}
+
+	return HintText;
+}
+
+void UMultiLineEditableTextBox::SetHintText(FText InHintText)
+{
+	HintText = InHintText;
+	HintTextDelegate.Clear();
+	if ( MyEditableTextBlock.IsValid() )
+	{
+		TAttribute<FText> HintTextBinding = PROPERTY_BINDING(FText, HintText);
+		MyEditableTextBlock->SetHintText(HintTextBinding);
+	}
+}
+
 void UMultiLineEditableTextBox::SetError(FText InError)
 {
 	if ( MyEditableTextBlock.IsValid() )
diff --git a/Engine/Source/Runtime/UMG/Public/Components/MultiLineEditableText.h b/Engine/Source/Runtime/UMG/Public/Components/MultiLineEditableText.h
--- a/Engine/Source/Runtime/UMG/Public/Components/MultiLineEditableText.h
+++ b/Engine/Source/Runtime/UMG/Public/Components/MultiLineEditableText.h
@@ -98,6 +98,14 @@ public:
 	UFUNCTION(BlueprintCallable, Category="Widget", meta=(DisplayName="SetText (Multi-Line Editable Text)"))
 	void SetText(FText InText);
 
+	/**  */
+	UFUNCTION(BlueprintCallable, Category="Widget", meta=(DisplayName="GetHintText (Multi-Line Editable Text)"))
+	FText GetHintText() const;
+
+	/**  */
+	UFUNCTION(BlueprintCallable, Category="Widget", meta=(DisplayName="SetHintText (Multi-Line Editable Text)"))
+	void SetHintText(FText InHintText);
+
 	UFUNCTION(BlueprintCallable, Category="Widget", meta=(DisplayName="SetIsReadOnly (Multi-Line Editable Text"))
 	void SetIsReadOnly(bool bReadOnly);
 	
diff --git a/Engine/Source/Runtime/UMG/Public/Components/MultiLineEditableTextBox.h b/Engine/Source/Runtime/UMG/Public/Components/MultiLineEditableTextBox.h
--- a/Engine/Source/Runtime/UMG/Public/Components/MultiLineEditableTextBox.h
+++ b/Engine/Source/Runtime/UMG/Public/Components/MultiLineEditableTextBox.h
@@ -106,6 +106,14 @@ public:
 	UFUNCTION(BlueprintCallable, Category="Widget", meta=(DisplayName="SetText (Multi-Line Text Box)"))
 	void SetText(FText InText);
 
+	/**  */
+	UFUNCTION(BlueprintCallable, Category="Widget", meta=(DisplayName="GetHintText (Multi-Line Text Box)"))
+	FText GetHintText() const;
+
+	/**  */
+	UFUNCTION(BlueprintCallable, Category="Widget", meta=(DisplayName="SetHintText (Multi-Line Text Box)"))
+	void SetHintText(FText InHintText);
+
 	UFUNCTION(BlueprintCallable, Category="Widget", meta=(DisplayName="SetError (Multi-Line Text Box)"))
 	void SetError(FText InError);
 
