Merging CL 6373836 from //UE4/Dev-Editor to //UE4/Main

RichTextBlock logs a compiler warning when its TextStyleSet DataTable doesn't have a RichTextStyleRow structure instead of an ensure failure.

diff --git a/Engine/Source/Runtime/UMG/Private/Components/RichTextBlock.cpp b/Engine/Source/Runtime/UMG/Private/Components/RichTextBlock.cpp
--- a/Engine/Source/Runtime/UMG/Private/Components/RichTextBlock.cpp
+++ b/Engine/Source/Runtime/UMG/Private/Components/RichTextBlock.cpp
@@ -12,6 +12,7 @@
 #include "Framework/Text/IRichTextMarkupParser.h"
 #include "Framework/Text/IRichTextMarkupWriter.h"
 #include "RenderingThread.h"
+#include "Editor/WidgetCompilerLog.h"
 
 #define LOCTEXT_NAMESPACE "UMG"
 
@@ -120,7 +121,7 @@ void URichTextBlock::RebuildStyleInstance()
 {
 	StyleInstance = MakeShareableDeferredCleanup(new FSlateStyleSet(TEXT("RichTextStyle")));
 
-	if (TextStyleSet && ensure(TextStyleSet->GetRowStruct()->IsChildOf(FRichTextStyleRow::StaticStruct())))
+	if (TextStyleSet && TextStyleSet->GetRowStruct()->IsChildOf(FRichTextStyleRow::StaticStruct()))
 	{
 		for (const auto& Entry : TextStyleSet->GetRowMap())
 		{
@@ -222,7 +223,20 @@ void URichTextBlock::OnCreationFromPalette()
 	//Decorators.Add(NewObject<URichTextBlockDecorator>(this, NAME_None, RF_Transactional));
 }
 
-#endif
+void URichTextBlock::ValidateCompiledDefaults(IWidgetCompilerLog& CompileLog) const
+{
+	Super::ValidateCompiledDefaults(CompileLog);
+
+	if (TextStyleSet && !TextStyleSet->GetRowStruct()->IsChildOf(FRichTextStyleRow::StaticStruct()))
+	{
+		CompileLog.Warning(FText::Format(
+			LOCTEXT("RichTextBlock_InvalidTextStyle", "{0} Text Style Set property expects a Data Table with a Rich Text Style Row structure (currently set to {1})."), 
+			FText::FromString(GetName()), 
+			FText::AsCultureInvariant(TextStyleSet->GetPathName())));
+	}
+}
+
+#endif //if WITH_EDITOR
 
 void URichTextBlock::SetDefaultTextStyle(const FTextBlockStyle& InDefaultTextStyle)
 {
diff --git a/Engine/Source/Runtime/UMG/Public/Components/RichTextBlock.h b/Engine/Source/Runtime/UMG/Public/Components/RichTextBlock.h
--- a/Engine/Source/Runtime/UMG/Public/Components/RichTextBlock.h
+++ b/Engine/Source/Runtime/UMG/Public/Components/RichTextBlock.h
@@ -121,6 +121,7 @@ public:
 	// UWidget interface
 	virtual const FText GetPaletteCategory() override;
 	virtual void OnCreationFromPalette() override;
+	virtual void ValidateCompiledDefaults(IWidgetCompilerLog& CompileLog) const override;
 	// End UWidget interface
 #endif
 
