Promoting the utility to get the owning player via template to the UWidget.

diff --git a/Engine/Source/Runtime/UMG/Public/Components/Widget.h b/Engine/Source/Runtime/UMG/Public/Components/Widget.h
--- a/Engine/Source/Runtime/UMG/Public/Components/Widget.h
+++ b/Engine/Source/Runtime/UMG/Public/Components/Widget.h
@@ -18,6 +18,7 @@
 #include "Components/SlateWrapperTypes.h"
 #include "Slate/WidgetTransform.h"
 #include "UObject/UObjectThreadContext.h"
+#include "GameFramework/PlayerController.h"
 
 #if WITH_EDITOR
 // This violates IWYU, but the alternative is .cpp includes that are invariably not within #if WITH_EDITOR and cause non-editor build failures
@@ -26,7 +27,6 @@
 
 #include "Widget.generated.h"
 
-class APlayerController;
 class ULocalPlayer;
 class SObjectWidget;
 class UPanelSlot;
@@ -672,6 +672,16 @@ public:
 	UFUNCTION(BlueprintCallable, BlueprintCosmetic, Category = "Widget")
 	virtual APlayerController* GetOwningPlayer() const;
 
+	/**
+	 * Gets the player controller associated with this UI cast to the template type.
+	 * @return The player controller that owns the UI. May be NULL if the cast fails.
+	 */
+	template <class TPlayerController = APlayerController >
+	TPlayerController* GetOwningPlayer() const
+	{
+		return Cast<TPlayerController>(GetOwningPlayer());
+	}
+
 	/**
 	 * Gets the local player associated with this UI.
 	 * @return The owning local player.
