UMG - Adding a templated accessor for GetOwningPlayer to UWidget.

diff --git a/Engine/Source/Runtime/UMG/Public/Components/Widget.h b/Engine/Source/Runtime/UMG/Public/Components/Widget.h
--- a/Engine/Source/Runtime/UMG/Public/Components/Widget.h
+++ b/Engine/Source/Runtime/UMG/Public/Components/Widget.h
@@ -688,6 +688,16 @@ public:
 	 */
 	UFUNCTION(BlueprintCallable, BlueprintCosmetic, Category = "Widget")
 	virtual ULocalPlayer* GetOwningLocalPlayer() const;
+	
+	/**
+	 * Gets the local player associated with this UI cast to the template type.
+	 * @return The owning local player. May be NULL if the cast fails.
+	 */
+	template < class T >
+	T* GetOwningLocalPlayer() const
+	{
+		return Cast<T>(GetOwningLocalPlayer());
+	}
 
 	/**
 	 * Applies all properties to the native widget if possible.  This is called after a widget is constructed.
