Slate - Widgets now attach reflection data only once, not every time we syncronize properties, which would actually be a small psuedoleak at editor time.  We no longer include the TWeakObjectPtr<> for the userwidget class responsible for creating the Widget in shipping builds.  Increased the default trip for 'number of excessive invalidation frame streak' to 60.  Fixing a bug where cooked widget templates wouldn't set the widget class responsible for creating the widget, and so none of the debugging info came through as useful information, that should now be fixed.

Matt.Kuhlenschmidt

diff --git a/Engine/Source/Runtime/Slate/Private/Widgets/SInvalidationPanel.cpp b/Engine/Source/Runtime/Slate/Private/Widgets/SInvalidationPanel.cpp
--- a/Engine/Source/Runtime/Slate/Private/Widgets/SInvalidationPanel.cpp
+++ b/Engine/Source/Runtime/Slate/Private/Widgets/SInvalidationPanel.cpp
@@ -77,7 +77,7 @@ static FAutoConsoleVariableRef CVarInvalidationLayerPadding(
 
 #if SLATE_VERBOSE_NAMED_EVENTS
 
-static int32 ExcessiveInvalidationFrameStreak = 10;
+static int32 ExcessiveInvalidationFrameStreak = 60;
 static FAutoConsoleVariableRef CVarExcessiveInvalidationFrameStreak(
 	TEXT("Slate.ExcessiveInvalidationFrameStreak"),
 	ExcessiveInvalidationFrameStreak,
diff --git a/Engine/Source/Runtime/UMG/Private/Components/Widget.cpp b/Engine/Source/Runtime/UMG/Private/Components/Widget.cpp
--- a/Engine/Source/Runtime/UMG/Private/Components/Widget.cpp
+++ b/Engine/Source/Runtime/UMG/Private/Components/Widget.cpp
@@ -779,6 +779,11 @@ TSharedRef<SWidget> UWidget::TakeWidget_Private(ConstructMethodType ConstructMet
 		bRoutedSynchronizeProperties = false;
 #endif
 
+#if WIDGET_INCLUDE_RELFECTION_METADATA
+		// We only need to do this once, when the slate widget is created.
+		PublicWidget->AddMetadata<FReflectionMetaData>(MakeShared<FReflectionMetaData>(GetFName(), GetClass(), this, GetSourceAssetOrClass()));
+#endif
+
 		SynchronizeProperties();
 		VerifySynchronizeProperties();
 		OnWidgetRebuilt();
@@ -1156,10 +1161,6 @@ void UWidget::SynchronizeProperties()
 	{
 		SafeWidget->SetToolTipText(PROPERTY_BINDING(FText, ToolTipText));
 	}
-
-#if !UE_BUILD_SHIPPING
-	SafeWidget->AddMetadata<FReflectionMetaData>(MakeShared<FReflectionMetaData>(GetFName(), GetClass(), this, GetSourceAssetOrClass()));
-#endif
 }
 
 UObject* UWidget::GetSourceAssetOrClass() const
@@ -1171,9 +1172,9 @@ UObject* UWidget::GetSourceAssetOrClass() const
 	// where it comes from, what blueprint, what the name of the widget was...etc.
 	SourceAsset = WidgetGeneratedBy.Get();
 #else
-#if !UE_BUILD_SHIPPING
-	SourceAsset = WidgetGeneratedByClass.Get();
-#endif
+	#if !UE_BUILD_SHIPPING
+		SourceAsset = WidgetGeneratedByClass.Get();
+	#endif
 #endif
 
 	if (!SourceAsset)
diff --git a/Engine/Source/Runtime/UMG/Private/UserWidget.cpp b/Engine/Source/Runtime/UMG/Private/UserWidget.cpp
--- a/Engine/Source/Runtime/UMG/Private/UserWidget.cpp
+++ b/Engine/Source/Runtime/UMG/Private/UserWidget.cpp
@@ -130,7 +130,9 @@ void UUserWidget::TemplateInitInner()
 	{
 		WidgetTree->ForEachWidget([this, WidgetClass] (UWidget* Widget) {
 
+#if !UE_BUILD_SHIPPING
 			Widget->WidgetGeneratedByClass = WidgetClass;
+#endif
 
 			// TODO UMG Make this an FName
 			FString VariableName = Widget->GetName();
diff --git a/Engine/Source/Runtime/UMG/Private/WidgetBlueprintGeneratedClass.cpp b/Engine/Source/Runtime/UMG/Private/WidgetBlueprintGeneratedClass.cpp
--- a/Engine/Source/Runtime/UMG/Private/WidgetBlueprintGeneratedClass.cpp
+++ b/Engine/Source/Runtime/UMG/Private/WidgetBlueprintGeneratedClass.cpp
@@ -201,6 +201,10 @@ void UWidgetBlueprintGeneratedClass::InitializeWidgetStatic(UUserWidget* UserWid
 		return;
 	}
 
+#if !UE_BUILD_SHIPPING
+	UserWidget->WidgetGeneratedByClass = MakeWeakObjectPtr(const_cast<UClass*>(InClass));
+#endif
+
 	UWidgetTree* ClonedTree = UserWidget->WidgetTree;
 
 	if ( UserWidget->bCookedWidgetTree )
@@ -225,6 +229,10 @@ void UWidgetBlueprintGeneratedClass::InitializeWidgetStatic(UUserWidget* UserWid
 		ClonedTree->ForEachWidget([&] (UWidget* Widget) {
 			check(Widget);
 
+#if !UE_BUILD_SHIPPING
+			Widget->WidgetGeneratedByClass = WidgetClass;
+#endif
+
 			if ( UUserWidget* SubUserWidget = Cast<UUserWidget>(Widget) )
 			{
 				SubUserWidget->Initialize();
@@ -259,8 +267,6 @@ void UWidgetBlueprintGeneratedClass::InitializeWidgetStatic(UUserWidget* UserWid
 	UE_LOG(LogUMG, Warning, TEXT("Widget Class %s - Slow Static Duplicate Object."), *InClass->GetName());
 #endif
 
-	UserWidget->WidgetGeneratedByClass = MakeWeakObjectPtr(const_cast<UClass*>(InClass));
-
 #if WITH_EDITOR
 	UserWidget->WidgetGeneratedBy = InClass->ClassGeneratedBy;
 #endif
@@ -278,7 +284,9 @@ void UWidgetBlueprintGeneratedClass::InitializeWidgetStatic(UUserWidget* UserWid
 				return;
 			}
 
+#if !UE_BUILD_SHIPPING
 			Widget->WidgetGeneratedByClass = MakeWeakObjectPtr(const_cast<UClass*>(InClass));
+#endif
 
 #if WITH_EDITOR
 			Widget->WidgetGeneratedBy = InClass->ClassGeneratedBy;
diff --git a/Engine/Source/Runtime/UMG/Public/Components/Widget.h b/Engine/Source/Runtime/UMG/Public/Components/Widget.h
--- a/Engine/Source/Runtime/UMG/Public/Components/Widget.h
+++ b/Engine/Source/Runtime/UMG/Public/Components/Widget.h
@@ -210,6 +210,11 @@ public:
 #endif
 
 
+#ifndef WIDGET_INCLUDE_RELFECTION_METADATA
+	#define WIDGET_INCLUDE_RELFECTION_METADATA !UE_BUILD_SHIPPING
+#endif
+
+
 
 /**
  * This is the base class for all wrapped Slate controls that are exposed to UObjects.
@@ -400,8 +405,10 @@ public:
 
 #endif
 
+#if !UE_BUILD_SHIPPING
 	/** Stores a reference to the class responsible for this widgets construction. */
 	TWeakObjectPtr<UClass> WidgetGeneratedByClass;
+#endif
 
 public:
 
