Fixing the compile issue in the generated class to reference a defined variable for the class to use that the widget was generated as part of.

diff --git a/Engine/Source/Runtime/UMG/Private/WidgetBlueprintGeneratedClass.cpp b/Engine/Source/Runtime/UMG/Private/WidgetBlueprintGeneratedClass.cpp
--- a/Engine/Source/Runtime/UMG/Private/WidgetBlueprintGeneratedClass.cpp
+++ b/Engine/Source/Runtime/UMG/Private/WidgetBlueprintGeneratedClass.cpp
@@ -202,7 +202,8 @@ void UWidgetBlueprintGeneratedClass::InitializeWidgetStatic(UUserWidget* UserWid
 	}
 
 #if !UE_BUILD_SHIPPING
-	UserWidget->WidgetGeneratedByClass = MakeWeakObjectPtr(const_cast<UClass*>(InClass));
+	TWeakObjectPtr<UClass> WidgetGeneratedByClass = MakeWeakObjectPtr(const_cast<UClass*>(InClass));
+	UserWidget->WidgetGeneratedByClass = WidgetGeneratedByClass;
 #endif
 
 	UWidgetTree* ClonedTree = UserWidget->WidgetTree;
@@ -230,7 +231,7 @@ void UWidgetBlueprintGeneratedClass::InitializeWidgetStatic(UUserWidget* UserWid
 			check(Widget);
 
 #if !UE_BUILD_SHIPPING
-			//Widget->WidgetGeneratedByClass = WidgetClass;
+			Widget->WidgetGeneratedByClass = WidgetGeneratedByClass;
 #endif
 
 			if ( UUserWidget* SubUserWidget = Cast<UUserWidget>(Widget) )
@@ -285,7 +286,7 @@ void UWidgetBlueprintGeneratedClass::InitializeWidgetStatic(UUserWidget* UserWid
 			}
 
 #if !UE_BUILD_SHIPPING
-			Widget->WidgetGeneratedByClass = MakeWeakObjectPtr(const_cast<UClass*>(InClass));
+			Widget->WidgetGeneratedByClass = WidgetGeneratedByClass;
 #endif
 
 #if WITH_EDITOR
