PR #5897: Export FRichImageRow for reusing (Contributed by [REDACTED_PERSON])

diff --git a/Engine/Source/Runtime/UMG/Public/Components/RichTextBlockImageDecorator.h b/Engine/Source/Runtime/UMG/Public/Components/RichTextBlockImageDecorator.h
--- a/Engine/Source/Runtime/UMG/Public/Components/RichTextBlockImageDecorator.h
+++ b/Engine/Source/Runtime/UMG/Public/Components/RichTextBlockImageDecorator.h
@@ -17,8 +17,8 @@
 class ISlateStyle;
 
 /** Simple struct for rich text styles */
-USTRUCT()
-struct FRichImageRow : public FTableRowBase
+USTRUCT(Blueprintable, BlueprintType)
+struct UMG_API FRichImageRow : public FTableRowBase
 {
 	GENERATED_USTRUCT_BODY()
 
