Fix WidgetComponent collision width when "Draw at Desired Size" is enabled (both for plane and cylinder cases)

diff --git a/Engine/Source/Runtime/UMG/Private/Components/WidgetComponent.cpp b/Engine/Source/Runtime/UMG/Private/Components/WidgetComponent.cpp
--- a/Engine/Source/Runtime/UMG/Private/Components/WidgetComponent.cpp
+++ b/Engine/Source/Runtime/UMG/Private/Components/WidgetComponent.cpp
@@ -1205,12 +1205,12 @@ float UWidgetComponent::ComputeComponentWidth() const
 	{
 		default:
 		case EWidgetGeometryMode::Plane:
-			return DrawSize.X;
+			return CurrentDrawSize.X;
 		break;
 
 		case EWidgetGeometryMode::Cylinder:
 			const float ArcAngleRadians = FMath::DegreesToRadians(GetCylinderArcAngle());
-			const float Radius = GetDrawSize().X / ArcAngleRadians;
+			const float Radius = CurrentDrawSize.X / ArcAngleRadians;
 			// Chord length is 2*R*Sin(Theta/2)
 			return 2.0f * Radius * FMath::Sin(0.5f*ArcAngleRadians);
 		break;
@@ -1636,15 +1636,16 @@ void UWidgetComponent::UpdateBodySetup( bool bDrawSizeChanged )
 		FVector Origin = FVector(.5f,
 			-( CurrentDrawSize.X * 0.5f ) + ( CurrentDrawSize.X * Pivot.X ),
 			-( CurrentDrawSize.Y * 0.5f ) + ( CurrentDrawSize.Y * Pivot.Y ));
-			const float Width = ComputeComponentWidth();
-			const float Height = CurrentDrawSize.Y;
-			Origin = FVector(.5f,
-				-( Width * 0.5f ) + ( Width * Pivot.X ),
-				-( Height * 0.5f ) + ( Height * Pivot.Y ));
+
+		const float Width = ComputeComponentWidth();
+		const float Height = CurrentDrawSize.Y;
+		Origin = FVector(.5f,
+			-( Width * 0.5f ) + ( Width * Pivot.X ),
+			-( Height * 0.5f ) + ( Height * Pivot.Y ));
 			
 		BoxElem->X = 0.01f;
-		BoxElem->Y = CurrentDrawSize.X;
-		BoxElem->Z = CurrentDrawSize.Y;
+		BoxElem->Y = Width;
+		BoxElem->Z = Height;
 
 		BoxElem->SetTransform(FTransform::Identity);
 		BoxElem->Center = Origin;
