Remove double calculation of origin in UWidgetComponent::UpdateBodySetup

diff --git a/Engine/Source/Runtime/UMG/Private/Components/WidgetComponent.cpp b/Engine/Source/Runtime/UMG/Private/Components/WidgetComponent.cpp
--- a/Engine/Source/Runtime/UMG/Private/Components/WidgetComponent.cpp
+++ b/Engine/Source/Runtime/UMG/Private/Components/WidgetComponent.cpp
@@ -1633,13 +1633,9 @@ void UWidgetComponent::UpdateBodySetup( bool bDrawSizeChanged )
 
 		FKBoxElem* BoxElem = BodySetup->AggGeom.BoxElems.GetData();
 
-		FVector Origin = FVector(.5f,
-			-( CurrentDrawSize.X * 0.5f ) + ( CurrentDrawSize.X * Pivot.X ),
-			-( CurrentDrawSize.Y * 0.5f ) + ( CurrentDrawSize.Y * Pivot.Y ));
-
 		const float Width = ComputeComponentWidth();
 		const float Height = CurrentDrawSize.Y;
-		Origin = FVector(.5f,
+		const FVector Origin = FVector(.5f,
 			-( Width * 0.5f ) + ( Width * Pivot.X ),
 			-( Height * 0.5f ) + ( Height * Pivot.Y ));
 			
