The dynamic entry box now returns the expected class type of what has been passed in, to in a lot of cases avoid re-casting it.

diff --git a/Engine/Source/Runtime/UMG/Public/Components/DynamicEntryBox.h b/Engine/Source/Runtime/UMG/Public/Components/DynamicEntryBox.h
--- a/Engine/Source/Runtime/UMG/Public/Components/DynamicEntryBox.h
+++ b/Engine/Source/Runtime/UMG/Public/Components/DynamicEntryBox.h
@@ -79,7 +79,7 @@ private:
 	UUserWidget* BP_CreateEntry();
 
 	/** Creates and establishes a new dynamic entry in the box using the specified class instead of the default. */
-	UFUNCTION(BlueprintCallable, Category = DynamicEntryBox, meta = (DisplayName = "Create Entry of Class", AllowPrivateAccess = true))
+	UFUNCTION(BlueprintCallable, Category = DynamicEntryBox, meta = (DisplayName = "Create Entry of Class", AllowPrivateAccess = true, DeterminesOutputType = "EntryClass"))
 	UUserWidget* BP_CreateEntryOfClass(TSubclassOf<UUserWidget> EntryClass);
 
 	/**
