refactored world to screen projection in UMG widgets

diff --git a/Engine/Source/Runtime/UMG/Private/Slate/SWorldWidgetScreenLayer.cpp b/Engine/Source/Runtime/UMG/Private/Slate/SWorldWidgetScreenLayer.cpp
--- a/Engine/Source/Runtime/UMG/Private/Slate/SWorldWidgetScreenLayer.cpp
+++ b/Engine/Source/Runtime/UMG/Private/Slate/SWorldWidgetScreenLayer.cpp
@@ -4,6 +4,7 @@
 #include "Widgets/Layout/SBox.h"
 #include "Components/WidgetComponent.h"
 #include "Blueprint/WidgetLayoutLibrary.h"
+#include "Blueprint/SlateBlueprintLibrary.h"
 #include "Engine/GameViewportClient.h"
 #include "Widgets/SViewport.h"
 #include "Slate/SGameLayerManager.h"
@@ -88,6 +89,21 @@ void SWorldWidgetScreenLayer::Tick(const FGeometry& AllottedGeometry, const doub
 		{
 			const FGeometry& ViewportGeometry = ViewportClient->GetGameLayerManager()->GetViewportWidgetHostGeometry();
 
+			// cache projection data here and avoid calls to UWidgetLayoutLibrary.ProjectWorldLocationToWidgetPositionWithDistance
+			FSceneViewProjectionData ProjectionData;
+			FMatrix ViewProjectionMatrix;
+			bool bHasProjectionData = false;
+
+			ULocalPlayer const* const LP = PlayerController->GetLocalPlayer();
+			if (LP && LP->ViewportClient)
+			{
+				bHasProjectionData = LP->GetProjectionData(ViewportClient->Viewport, eSSP_FULL, /*out*/ ProjectionData);
+				if (bHasProjectionData)
+				{
+					ViewProjectionMatrix = ProjectionData.ComputeViewProjectionMatrix();
+				}
+			}
+
 			for ( auto It = ComponentMap.CreateIterator(); It; ++It )
 			{
 				FComponentEntry& Entry = It.Value();
@@ -96,13 +112,17 @@ void SWorldWidgetScreenLayer::Tick(const FGeometry& AllottedGeometry, const doub
 				{
 					FVector WorldLocation = SceneComponent->GetComponentLocation();
 
-					//TODO NDarnell Perf This can be improved, if we get the projection matrix and do all of them in one go, instead of calling the library.
+					FVector2D ScreenPosition2D;
+					const bool bProjected = bHasProjectionData ? FSceneView::ProjectWorldToScreen(WorldLocation, ProjectionData.GetConstrainedViewRect(), ViewProjectionMatrix, ScreenPosition2D) : false;
+					if (bProjected)
+					{
+						const float ViewportDist = FVector::Dist(ProjectionData.ViewOrigin, WorldLocation);
+						const FVector2D RoundedPosition2D(FMath::RoundToInt(ScreenPosition2D.X), FMath::RoundToInt(ScreenPosition2D.Y));
+						FVector2D ViewportPosition2D;
+						USlateBlueprintLibrary::ScreenToViewport(PlayerController, RoundedPosition2D, ViewportPosition2D);
 
-					FVector ViewportPosition;
-					const bool bProjected = UWidgetLayoutLibrary::ProjectWorldLocationToWidgetPositionWithDistance(PlayerController, WorldLocation, ViewportPosition);
+						const FVector ViewportPosition(ViewportPosition2D.X, ViewportPosition2D.Y, ViewportDist);
 
-					if ( bProjected )
-					{
 						Entry.ContainerWidget->SetVisibility(EVisibility::SelfHitTestInvisible);
 
 						if ( SConstraintCanvas::FSlot* CanvasSlot = Entry.Slot )
