Deprecated SetNavigationRule, added three new methods (SetNavigationRuleBase, SetNavigationRuleExplicit, SetNavigationRuleCustom, SetNavigationRuleCustomBoundary) to handle all the possible cases.

diff --git a/Engine/Source/Runtime/UMG/Private/Components/Widget.cpp b/Engine/Source/Runtime/UMG/Private/Components/Widget.cpp
--- a/Engine/Source/Runtime/UMG/Private/Components/Widget.cpp
+++ b/Engine/Source/Runtime/UMG/Private/Components/Widget.cpp
@@ -16,7 +16,6 @@
 #include "Widgets/Layout/SSpacer.h"
 #include "Widgets/SToolTip.h"
 #include "Binding/PropertyBinding.h"
-#include "Blueprint/WidgetNavigation.h"
 #include "Logging/MessageLog.h"
 #include "Blueprint/UserWidget.h"
 #include "Slate/SObjectWidget.h"
@@ -604,7 +603,7 @@ FVector2D UWidget::GetDesiredSize() const
 	return FVector2D(0, 0);
 }
 
-void UWidget::SetNavigationRuleInternal(EUINavigation Direction, EUINavigationRule Rule, FName WidgetToFocus)
+void UWidget::SetNavigationRuleInternal(EUINavigation Direction, EUINavigationRule Rule, FName WidgetToFocus/* = NAME_None*/, UWidget* InWidget/* = nullptr*/, FCustomWidgetNavigationDelegate InCustomDelegate/* = FCustomWidgetNavigationDelegate()*/)
 {
 	if (Navigation == nullptr)
 	{
@@ -614,6 +613,8 @@ void UWidget::SetNavigationRuleInternal(EUINavigation Direction, EUINavigationRu
 	FWidgetNavigationData NavigationData;
 	NavigationData.Rule = Rule;
 	NavigationData.WidgetToFocus = WidgetToFocus;
+	NavigationData.Widget = InWidget;
+	NavigationData.CustomDelegate = InCustomDelegate;
 	switch(Direction)
 	{
 		case EUINavigation::Up:
@@ -645,6 +646,37 @@ void UWidget::SetNavigationRule(EUINavigation Direction, EUINavigationRule Rule,
 	BuildNavigation();
 }
 
+void UWidget::SetNavigationRuleBase(EUINavigation Direction, EUINavigationRule Rule)
+{
+	if (Rule == EUINavigationRule::Explicit || Rule == EUINavigationRule::Custom || Rule == EUINavigationRule::CustomBoundary)
+	{
+#if !(UE_BUILD_SHIPPING || UE_BUILD_TEST)
+		FMessageLog("PIE").Error(LOCTEXT("SetNavigationRuleBaseWrongRule", "Cannot use SetNavigationRuleBase with an Explicit or a Custom or a CustomBoundary Rule."));
+#endif
+		return;
+	}
+	SetNavigationRuleInternal(Direction, Rule);
+	BuildNavigation();
+}
+
+void UWidget::SetNavigationRuleExplicit(EUINavigation Direction, UWidget* InWidget)
+{
+	SetNavigationRuleInternal(Direction, EUINavigationRule::Explicit, NAME_None, InWidget);
+	BuildNavigation();
+}
+
+void UWidget::SetNavigationRuleCustom(EUINavigation Direction, FCustomWidgetNavigationDelegate InCustomDelegate)
+{
+	SetNavigationRuleInternal(Direction, EUINavigationRule::Custom, NAME_None, nullptr, InCustomDelegate);
+	BuildNavigation();
+}
+
+void UWidget::SetNavigationRuleCustomBoundary(EUINavigation Direction, FCustomWidgetNavigationDelegate InCustomDelegate)
+{
+	SetNavigationRuleInternal(Direction, EUINavigationRule::CustomBoundary, NAME_None, nullptr, InCustomDelegate);
+	BuildNavigation();
+}
+
 void UWidget::SetAllNavigationRules(EUINavigationRule Rule, FName WidgetToFocus)
 {
 	SetNavigationRuleInternal(EUINavigation::Up, Rule, WidgetToFocus);
diff --git a/Engine/Source/Runtime/UMG/Public/Components/Widget.h b/Engine/Source/Runtime/UMG/Public/Components/Widget.h
--- a/Engine/Source/Runtime/UMG/Public/Components/Widget.h
+++ b/Engine/Source/Runtime/UMG/Public/Components/Widget.h
@@ -19,6 +19,7 @@
 #include "Slate/WidgetTransform.h"
 #include "UObject/UObjectThreadContext.h"
 #include "GameFramework/PlayerController.h"
+#include "Blueprint/WidgetNavigation.h"
 
 #if WITH_EDITOR
 // This violates IWYU, but the alternative is .cpp includes that are invariably not within #if WITH_EDITOR and cause non-editor build failures
@@ -613,9 +614,42 @@ public:
 	 *	@param Rule The rule to use when navigation is taking place
 	 *	@param WidgetToFocus When using the Explicit rule, focus on this widget
 	 */
+	UE_DEPRECATED(4.23, "SetNavigationRule is deprecated. Please use either SetNavigationRuleBase or SetNavigationRuleExplicit or SetNavigationRuleCustom or SetNavigationRuleCustomBoundary.")
 	UFUNCTION(BlueprintCallable, Category = "Widget")
 	void SetNavigationRule(EUINavigation Direction, EUINavigationRule Rule, FName WidgetToFocus);
 
+	/**
+	 *	Sets the widget navigation rules for a specific direction. This can only be called on widgets that are in a widget tree. This works only for non Explicit, non Custom and non CustomBoundary Rules.
+	 *	@param Direction
+	 *	@param Rule The rule to use when navigation is taking place
+	 */
+	UFUNCTION(BlueprintCallable, Category = "Widget")
+	void SetNavigationRuleBase(EUINavigation Direction, EUINavigationRule Rule);
+
+	/**
+	 *	Sets the widget navigation rules for a specific direction. This can only be called on widgets that are in a widget tree. This works only for Explicit Rule.
+	 *	@param Direction
+	 *	@param InWidget Focus on this widget instance
+	 */
+	UFUNCTION(BlueprintCallable, Category = "Widget")
+	void SetNavigationRuleExplicit(EUINavigation Direction, UWidget* InWidget);
+
+	/**
+	 *	Sets the widget navigation rules for a specific direction. This can only be called on widgets that are in a widget tree. This works only for Custom Rule.
+	 *	@param Direction
+	 *	@param InCustomDelegate Custom Delegate that will be called
+	 */
+	UFUNCTION(BlueprintCallable, Category = "Widget")
+	void SetNavigationRuleCustom(EUINavigation Direction, FCustomWidgetNavigationDelegate InCustomDelegate);
+
+	/**
+	 *	Sets the widget navigation rules for a specific direction. This can only be called on widgets that are in a widget tree. This works only for CustomBoundary Rule.
+	 *	@param Direction
+	 *	@param InCustomDelegate Custom Delegate that will be called
+	 */
+	UFUNCTION(BlueprintCallable, Category = "Widget")
+	void SetNavigationRuleCustomBoundary(EUINavigation Direction, FCustomWidgetNavigationDelegate InCustomDelegate);
+
 	/** Gets the parent widget */
 	UFUNCTION(BlueprintCallable, Category="Widget")
 	class UPanelWidget* GetParent() const;
@@ -954,7 +988,7 @@ protected:
 		return FSlateColor(InLinearColor.Get());
 	}
 
-	void SetNavigationRuleInternal(EUINavigation Direction, EUINavigationRule Rule, FName WidgetToFocus);
+	void SetNavigationRuleInternal(EUINavigation Direction, EUINavigationRule Rule, FName WidgetToFocus = NAME_None, UWidget* InWidget = nullptr, FCustomWidgetNavigationDelegate InCustomDelegate = FCustomWidgetNavigationDelegate());
 
 #if WITH_ACCESSIBILITY
 	/** Gets the widget that accessibility properties should synchronize to. */
