Slate Global Invalidation System

diff --git a/Engine/Source/Runtime/SlateRHIRenderer/Private/SlateRHIResourceManager.h b/Engine/Source/Runtime/SlateRHIRenderer/Private/SlateRHIResourceManager.h
--- a/Engine/Source/Runtime/SlateRHIRenderer/Private/SlateRHIResourceManager.h
+++ b/Engine/Source/Runtime/SlateRHIRenderer/Private/SlateRHIResourceManager.h
@@ -19,7 +19,6 @@ class FSlateAtlasedTextureResource;
 class FSlateDynamicTextureResource;
 class FSlateMaterialResource;
 class FSlateUTextureResource;
-class ILayoutCache;
 class ISlateStyle;
 class UTexture;
 class FSceneInterface;
@@ -121,7 +120,7 @@ struct FCachedRenderBuffers
 /**
  * Stores a mapping of texture names to their RHI texture resource               
  */
-class FSlateRHIResourceManager : public ISlateAtlasProvider, public ISlateRenderDataManager, public FSlateShaderResourceManager, public FTickableGameObject
+class FSlateRHIResourceManager : public ISlateAtlasProvider, public FSlateShaderResourceManager, public FTickableGameObject
 {
 public:
 	FSlateRHIResourceManager();
@@ -133,9 +132,6 @@ public:
 	virtual FSlateShaderResource* GetAtlasPageResource(const int32 InIndex) const override;
 	virtual bool IsAtlasPageResourceAlphaOnly() const override;
 
-	/** ISlateRenderDataManager interface */
-	virtual void BeginReleasingRenderData(const FSlateRenderDataHandle* RenderHandle) override;
-
 	/** FTickableGameObject interface */
 	virtual ETickableTickType GetTickableTickType() const override { return ETickableTickType::Always; }
 	virtual bool IsTickableWhenPaused() const override { return true; }
@@ -203,23 +199,6 @@ public:
 	virtual bool LoadTexture( const FName& TextureName, const FString& ResourcePath, uint32& Width, uint32& Height, TArray<uint8>& DecodedImage );
 	virtual bool LoadTexture( const FSlateBrush& InBrush, uint32& Width, uint32& Height, TArray<uint8>& DecodedImage );
 
-	/**
-	 * 
-	 */
-	FCachedRenderBuffers* FindCachedBuffersForHandle(const FSlateRenderDataHandle* RenderHandle) const
-	{
-		return CachedBuffers.FindRef(RenderHandle);
-	}
-
-	/**
-	 * Creates a vertex and index buffer for a given render handle that it can use to store cached widget geometry to.
-	 */
-	FCachedRenderBuffers* FindOrCreateCachedBuffersForHandle(const TSharedRef<FSlateRenderDataHandle, ESPMode::ThreadSafe>& RenderHandle);
-
-	/**
-	 * Releases all cached buffers allocated by a given layout cacher.  This would happen when an invalidation panel is destroyed.
-	 */
-	void ReleaseCachingResourcesFor(FRHICommandListImmediate& RHICmdList, const ILayoutCache* Cacher);
 
 	/**
 	 * Releases rendering resources
@@ -245,16 +224,6 @@ private:
 
 	void TryToCleanupExpiredResources(bool bForceCleanup);
 
-	void ReleaseCachedBuffer(FRHICommandListImmediate& RHICmdList, FCachedRenderBuffers* PooledBuffer);
-
-	/**
-	* Releases the cached render data for a given render handle.  The layout cacher that owned the handle must also be
-	* provided, as render handle is likely no longer a valid object.
-	*/
-	void ReleaseCachedRenderData(FRHICommandListImmediate& RHICmdList, const FSlateRenderDataHandle* RenderHandle, const ILayoutCache* LayoutCacher);
-
-private:
-
 	/**
 	 * Deletes resources created by the manager
 	 */
@@ -265,11 +234,6 @@ private:
 	 */
 	void DeleteUObjectBrushResources();
 
-	/**
-	 * Deletes cached buffers
-	 */
-	void DeleteCachedBuffers();
-
 	/**
 	 * Debugging command to try forcing a refresh of UObject brushes.
 	 */
@@ -343,18 +307,6 @@ private:
 	/** Needed for displaying an error texture when we end up with bad resources. */
 	UTexture* BadResourceTexture;
 
-	typedef TMap< FSlateRenderDataHandle*, FCachedRenderBuffers* > TCachedBufferMap;
-	TCachedBufferMap CachedBuffers;
-
-	typedef TMap< const ILayoutCache*, TArray< FCachedRenderBuffers* > > TCachedBufferPoolMap;
-	TCachedBufferPoolMap CachedBufferPool;
-
-	/**
-	 * Holds onto a list of pooled buffers that are no longer being used but still need 
-	 * to be deleted after the RHI thread is done with them.
-	 */
-	TArray<FCachedRenderBuffers*> PooledBuffersPendingRelease;
-
 	/**  */
 	TArray<FSceneInterface*> ActiveScenes;
 
diff --git a/Engine/Source/Runtime/SlateRHIRenderer/Private/SlateShaders.h b/Engine/Source/Runtime/SlateRHIRenderer/Private/SlateShaders.h
--- a/Engine/Source/Runtime/SlateRHIRenderer/Private/SlateShaders.h
+++ b/Engine/Source/Runtime/SlateRHIRenderer/Private/SlateShaders.h
@@ -186,7 +186,7 @@ private:
 /** 
  * Pixel shader types for all elements
  */
-template<ESlateShader::Type ShaderType, bool bDrawDisabledEffect, bool bUseTextureAlpha=true>
+template<ESlateShader ShaderType, bool bDrawDisabledEffect, bool bUseTextureAlpha=true>
 class TSlateElementPS : public FSlateElementPS
 {
 	DECLARE_SHADER_TYPE( TSlateElementPS, Global );
diff --git a/Engine/Source/Runtime/SlateRHIRenderer/Private/SlateUpdatableBuffer.cpp b/Engine/Source/Runtime/SlateRHIRenderer/Private/SlateUpdatableBuffer.cpp
--- a/Engine/Source/Runtime/SlateRHIRenderer/Private/SlateUpdatableBuffer.cpp
+++ b/Engine/Source/Runtime/SlateRHIRenderer/Private/SlateUpdatableBuffer.cpp
@@ -7,24 +7,25 @@ DECLARE_CYCLE_STAT(TEXT("UpdateInstanceBuffer Time"), STAT_SlateUpdateInstanceBu
 
 struct FSlateUpdateInstanceBufferCommand final : public FRHICommand<FSlateUpdateInstanceBufferCommand>
 {
-	FVertexBufferRHIRef VertexBufferRHI;
+	TSlateElementVertexBuffer<FVector4>& InstanceBuffer;
 	const TArray<FVector4>& InstanceData;
 
 	FSlateUpdateInstanceBufferCommand(TSlateElementVertexBuffer<FVector4>& InInstanceBuffer, const TArray<FVector4>& InInstanceData )
-		: VertexBufferRHI(InInstanceBuffer.VertexBufferRHI)
+		: InstanceBuffer(InInstanceBuffer)
 		, InstanceData(InInstanceData)
 	{}
 
 	void Execute(FRHICommandListBase& CmdList)
 	{
 		SCOPE_CYCLE_COUNTER( STAT_SlateUpdateInstanceBuffer );
+		const bool bIsInRenderingThread = !IsRunningRHIInSeparateThread() || CmdList.Bypass();
 
 		int32 RequiredVertexBufferSize = InstanceData.Num()*sizeof(FVector4);
-		uint8* InstanceBufferData = (uint8*)GDynamicRHI->RHILockVertexBuffer( VertexBufferRHI, 0, RequiredVertexBufferSize, RLM_WriteOnly );
+		uint8* InstanceBufferData = (uint8*)InstanceBuffer.LockBuffer(RequiredVertexBufferSize, bIsInRenderingThread);
 
 		FMemory::Memcpy( InstanceBufferData, InstanceData.GetData(), InstanceData.Num()*sizeof(FVector4) );
 	
-		GDynamicRHI->RHIUnlockVertexBuffer( VertexBufferRHI );
+		InstanceBuffer.UnlockBuffer(bIsInRenderingThread);
 	}
 };
 
@@ -92,11 +93,8 @@ void FSlateUpdatableInstanceBuffer::UpdateRenderingData_RenderThread(FRHICommand
 
 	if(!IsRunningRHIInSeparateThread() || RHICmdList.Bypass())
 	{
-		uint8* InstanceBufferData = (uint8*)InstanceBufferResource.LockBuffer_RenderThread(RenderThreadBufferData.Num());
-
-		FMemory::Memcpy(InstanceBufferData, RenderThreadBufferData.GetData(), RenderThreadBufferData.Num()*sizeof(FVector4) );
-	
-		InstanceBufferResource.UnlockBuffer_RenderThread();
+		FSlateUpdateInstanceBufferCommand Command(InstanceBufferResource, RenderThreadBufferData);
+		Command.Execute(RHICmdList);
 	}
 	else
 	{
diff --git a/Engine/Source/Runtime/UMG/Private/Components/InvalidationBox.cpp b/Engine/Source/Runtime/UMG/Private/Components/InvalidationBox.cpp
--- a/Engine/Source/Runtime/UMG/Private/Components/InvalidationBox.cpp
+++ b/Engine/Source/Runtime/UMG/Private/Components/InvalidationBox.cpp
@@ -28,7 +28,6 @@ TSharedRef<SWidget> UInvalidationBox::RebuildWidget()
 {
 	MyInvalidationPanel =
 		SNew(SInvalidationPanel)
-		.CacheRelativeTransforms(CacheRelativeTransforms)
 #if !UE_BUILD_SHIPPING
 		.DebugName(GetPathName())
 #endif
@@ -64,10 +63,6 @@ void UInvalidationBox::OnSlotRemoved(UPanelSlot* InSlot)
 
 void UInvalidationBox::InvalidateCache()
 {
-	if ( MyInvalidationPanel.IsValid() )
-	{
-		return MyInvalidationPanel->InvalidateCache();
-	}
 }
 
 bool UInvalidationBox::GetCanCache() const
diff --git a/Engine/Source/Runtime/UMG/Private/Components/RetainerBox.cpp b/Engine/Source/Runtime/UMG/Private/Components/RetainerBox.cpp
--- a/Engine/Source/Runtime/UMG/Private/Components/RetainerBox.cpp
+++ b/Engine/Source/Runtime/UMG/Private/Components/RetainerBox.cpp
@@ -142,11 +142,11 @@ const FText URetainerBox::GetPaletteCategory()
 
 #endif
 
-const FGeometry& URetainerBox::GetCachedAllottedGeometry() const
+FGeometry URetainerBox::GetCachedAllottedGeometry() const
 {
 	if (MyRetainerWidget.IsValid())
 	{
-		return MyRetainerWidget->GetCachedGeometry();
+		return MyRetainerWidget->GetTickSpaceGeometry();
 	}
 
 	static const FGeometry TempGeo;
diff --git a/Engine/Source/Runtime/UMG/Private/Components/Widget.cpp b/Engine/Source/Runtime/UMG/Private/Components/Widget.cpp
--- a/Engine/Source/Runtime/UMG/Private/Components/Widget.cpp
+++ b/Engine/Source/Runtime/UMG/Private/Components/Widget.cpp
@@ -579,7 +579,7 @@ void UWidget::ForceLayoutPrepass()
 	TSharedPtr<SWidget> SafeWidget = GetCachedWidget();
 	if (SafeWidget.IsValid())
 	{
-		SafeWidget->SlatePrepass(SafeWidget->GetCachedGeometry().Scale);
+		SafeWidget->SlatePrepass(SafeWidget->GetTickSpaceGeometry().Scale);
 	}
 }
 
@@ -722,14 +722,30 @@ void UWidget::RemoveFromParent()
 }
 
 const FGeometry& UWidget::GetCachedGeometry() const
+{
+	return GetTickSpaceGeometry();
+}
+
+const FGeometry& UWidget::GetTickSpaceGeometry() const
 {
 	TSharedPtr<SWidget> SafeWidget = GetCachedWidget();
 	if ( SafeWidget.IsValid() )
 	{
-		return SafeWidget->GetCachedGeometry();
+		return SafeWidget->GetTickSpaceGeometry();
 	}
 
-	return SNullWidget::NullWidget->GetCachedGeometry();
+	return SNullWidget::NullWidget->GetTickSpaceGeometry();
+}
+
+const FGeometry& UWidget::GetPaintSpaceGeometry() const
+{
+	TSharedPtr<SWidget> SafeWidget = GetCachedWidget();
+	if (SafeWidget.IsValid())
+	{
+		return SafeWidget->GetPaintSpaceGeometry();
+	}
+
+	return SNullWidget::NullWidget->GetPaintSpaceGeometry();
 }
 
 #if !(UE_BUILD_SHIPPING || UE_BUILD_TEST)
@@ -1174,10 +1190,6 @@ void UWidget::SynchronizeProperties()
 
 	SafeWidget->SetRenderOpacity(RenderOpacity);
 
-#if !UE_BUILD_SHIPPING
-	SafeWidget->SetTag(GetFName());
-#endif
-
 	UpdateRenderTransform();
 	SafeWidget->SetRenderTransformPivot(RenderTransformPivot);
 
@@ -1203,6 +1215,10 @@ void UWidget::SynchronizeProperties()
 		SafeWidget->SetToolTipText(PROPERTY_BINDING(FText, ToolTipText));
 	}
 
+#if WITH_SLATE_DEBUGGING
+	SafeWidget->AddMetadata<FReflectionMetaData>(MakeShared<FReflectionMetaData>(GetFName(), GetClass(), this, GetSourceAssetOrClass()));
+#endif
+
 #if WITH_ACCESSIBILITY
 	TSharedPtr<SWidget> AccessibleWidget = GetAccessibleWidget();
 	if (AccessibleWidget.IsValid())
diff --git a/Engine/Source/Runtime/UMG/Private/Components/WidgetComponent.cpp b/Engine/Source/Runtime/UMG/Private/Components/WidgetComponent.cpp
--- a/Engine/Source/Runtime/UMG/Private/Components/WidgetComponent.cpp
+++ b/Engine/Source/Runtime/UMG/Private/Components/WidgetComponent.cpp
@@ -1845,7 +1845,7 @@ TArray<FWidgetAndPointer> UWidgetComponent::GetHitWidgetPath(FVector2D WidgetSpa
 	TArray<FWidgetAndPointer> ArrangedWidgets;
 	if ( SlateWindow.IsValid() )
 	{
-		ArrangedWidgets = SlateWindow->GetHittestGrid()->GetBubblePath( LocalHitLocation, CursorRadius, bIgnoreEnabledStatus );
+		ArrangedWidgets = SlateWindow->GetHittestGrid().GetBubblePath( LocalHitLocation, CursorRadius, bIgnoreEnabledStatus );
 
 		for( FWidgetAndPointer& ArrangedWidget : ArrangedWidgets )
 		{
diff --git a/Engine/Source/Runtime/UMG/Private/Slate/SObjectWidget.cpp b/Engine/Source/Runtime/UMG/Private/Slate/SObjectWidget.cpp
--- a/Engine/Source/Runtime/UMG/Private/Slate/SObjectWidget.cpp
+++ b/Engine/Source/Runtime/UMG/Private/Slate/SObjectWidget.cpp
@@ -132,9 +132,10 @@ int32 SObjectWidget::OnPaint(const FPaintArgs& Args, const FGeometry& AllottedGe
 
 bool SObjectWidget::ComputeVolatility() const
 {
-	if ( CanRouteEvent() )
+	// @todo eliminate this
+	if (!GSlateEnableGlobalInvalidation && CanRouteEvent())
 	{
-		return SCompoundWidget::ComputeVolatility() || WidgetObject->IsPlayingAnimation();
+		return WidgetObject->IsPlayingAnimation() || SCompoundWidget::ComputeVolatility();
 	}
 
 	return SCompoundWidget::ComputeVolatility();
@@ -535,4 +536,4 @@ void SObjectWidget::OnMouseCaptureLost(const FCaptureLostEvent& CaptureLostEvent
 	{
 		return WidgetObject->NativeOnMouseCaptureLost(CaptureLostEvent);
 	}
-}
\ No newline at end of file
+}
diff --git a/Engine/Source/Runtime/UMG/Private/Slate/SRetainerWidget.cpp b/Engine/Source/Runtime/UMG/Private/Slate/SRetainerWidget.cpp
--- a/Engine/Source/Runtime/UMG/Private/Slate/SRetainerWidget.cpp
+++ b/Engine/Source/Runtime/UMG/Private/Slate/SRetainerWidget.cpp
@@ -7,8 +7,6 @@
 #include "Engine/TextureRenderTarget2D.h"
 #include "Framework/Application/SlateApplication.h"
 #include "Engine/World.h"
-#include "Layout/WidgetCaching.h"
-
 
 DECLARE_CYCLE_STAT(TEXT("Retainer Widget Tick"), STAT_SlateRetainerWidgetTick, STATGROUP_Slate);
 DECLARE_CYCLE_STAT(TEXT("Retainer Widget Paint"), STAT_SlateRetainerWidgetPaint, STATGROUP_Slate);
@@ -87,22 +85,28 @@ TFrameValue<int32> SRetainerWidget::Shared_RetainerWorkThisFrame(0);
 
 SRetainerWidget::SRetainerWidget()
 	: EmptyChildSlot(this)
+	, VirtualWindow(SNew(SVirtualWindow))
 	, RenderingResources(new FRetainerWidgetRenderingResources)
+
 {
+	FSlateApplicationBase::Get().OnGlobalInvalidationToggled().AddRaw(this, &SRetainerWidget::OnGlobalInvalidationToggled);
+	if (FSlateApplication::IsInitialized())
+	{
+#if !UE_BUILD_SHIPPING
+		OnRetainerModeChangedDelegate.RemoveAll(this);
+#endif
+	}
+	bHasCustomPrepass = true;
+	SetInvalidationRootWidget(*this);
+	SetInvalidationRootHittestGrid(HittestGrid);
 	SetCanTick(false);
 }
 
 SRetainerWidget::~SRetainerWidget()
 {
-	for ( int32 i = 0; i < NodePool.Num(); i++ )
-	{
-		delete NodePool[i];
-	}
-	NodePool.Empty();
-	
 	if( FSlateApplication::IsInitialized() )
 	{
-		FSlateApplicationBase::Get().OnGlobalInvalidate().RemoveAll( this );
+		FSlateApplicationBase::Get().OnGlobalInvalidationToggled().RemoveAll(this);
 #if !UE_BUILD_SHIPPING
 		OnRetainerModeChangedDelegate.RemoveAll( this );
 #endif
@@ -151,7 +155,7 @@ void SRetainerWidget::UpdateWidgetRenderer()
 
 void SRetainerWidget::Construct(const FArguments& InArgs)
 {
-	FSlateApplicationBase::Get().OnGlobalInvalidate().AddSP(this, &SRetainerWidget::OnGlobalInvalidate);
+	FSlateApplicationBase::Get().OnInvalidateAllWidgets().AddSP(this, &SRetainerWidget::OnGlobalInvalidate);
 
 	STAT(MyStatId = FDynamicStats::CreateStatId<FStatGroup_STATGROUP_Slate>(InArgs._StatId);)
 
@@ -163,11 +167,9 @@ void SRetainerWidget::Construct(const FArguments& InArgs)
 	RenderingResources->RenderTarget = RenderTarget;
 	SurfaceBrush.SetResourceObject(RenderTarget);
 
-	Window = SNew(SVirtualWindow)
-    .Visibility(EVisibility::SelfHitTestInvisible);  // deubanks: We don't want Retainer Widgets blocking hit testing for tooltips
+	VirtualWindow->SetVisibility(EVisibility::SelfHitTestInvisible);  // deubanks: We don't want Retainer Widgets blocking hit testing for tooltips
+	VirtualWindow->SetShouldResolveDeferred(false);
 
-	Window->SetShouldResolveDeferred(false);
-	
 	UpdateWidgetRenderer();
 
 	MyWidget = InArgs._Content.Widget;
@@ -184,16 +186,12 @@ void SRetainerWidget::Construct(const FArguments& InArgs)
 	bEnableRetainedRenderingDesire = true;
 	bEnableRetainedRendering = false;
 
+	RefreshRenderingMode();
 	bRenderRequested = true;
 
-	RootCacheNode = nullptr;
-	LastUsedCachedNodeIndex = 0;
-
-	Window->SetContent(MyWidget.ToSharedRef());
-
 	ChildSlot
 	[
-		Window.ToSharedRef()
+		MyWidget.ToSharedRef()
 	];
 
 	if ( FSlateApplication::IsInitialized() )
@@ -225,7 +223,11 @@ bool SRetainerWidget::IsAnythingVisibleToRender() const
 void SRetainerWidget::OnRetainerModeChanged()
 {
 	RefreshRenderingMode();
-	Invalidate(EInvalidateWidget::Layout);
+	// Invalidate myself 
+	InvalidateRoot();
+
+	// Nested invalidation: Invalidate whatever invalidation root I am in.
+	Invalidate(EInvalidateWidget::ChildOrder);
 }
 
 void SRetainerWidget::OnGlobalInvalidate()
@@ -245,6 +247,8 @@ void SRetainerWidget::OnRetainerModeCVarChanged( IConsoleVariable* CVar )
 void SRetainerWidget::SetRetainedRendering(bool bRetainRendering)
 {
 	bEnableRetainedRenderingDesire = bRetainRendering;
+
+	RefreshRenderingMode();
 }
 
 void SRetainerWidget::RefreshRenderingMode()
@@ -254,15 +258,17 @@ void SRetainerWidget::RefreshRenderingMode()
 	if ( bEnableRetainedRendering != bShouldBeRenderingOffscreen )
 	{
 		bEnableRetainedRendering = bShouldBeRenderingOffscreen;
-
-		Window->SetContent(MyWidget.ToSharedRef());
+		InvalidateChildOrder();
 	}
 }
 
 void SRetainerWidget::SetContent(const TSharedRef< SWidget >& InContent)
 {
 	MyWidget = InContent;
-	Window->SetContent(InContent);
+	ChildSlot
+	[
+		InContent
+	];
 }
 
 UMaterialInstanceDynamic* SRetainerWidget::GetEffectMaterial() const
@@ -304,7 +310,7 @@ void SRetainerWidget::SetWorld(UWorld* World)
 
 FChildren* SRetainerWidget::GetChildren()
 {
-	if ( bEnableRetainedRendering )
+	if (bEnableRetainedRendering && !GSlateEnableGlobalInvalidation && !NeedsPrepass())
 	{
 		return &EmptyChildSlot;
 	}
@@ -314,35 +320,9 @@ FChildren* SRetainerWidget::GetChildren()
 	}
 }
 
-bool SRetainerWidget::ComputeVolatility() const
-{
-	return true;
-}
-
-FCachedWidgetNode* SRetainerWidget::CreateCacheNode() const
-{
-	// If the node pool is empty, allocate a few
-	if (LastUsedCachedNodeIndex >= NodePool.Num())
-	{
-		for (int32 i = 0; i < 10; i++)
-		{
-			NodePool.Add(new FCachedWidgetNode());
-		}
-	}
-
-	// Return one of the preallocated nodes and increment the next node index.
-	FCachedWidgetNode* NewNode = NodePool[LastUsedCachedNodeIndex];
-	++LastUsedCachedNodeIndex;
-
-	return NewNode;
-}
-
-void SRetainerWidget::InvalidateWidget(SWidget* InvalidateWidget)
+FChildren* SRetainerWidget::GetAllChildren()
 {
-	if (RenderOnInvalidation)
-	{
-		bRenderRequested = true;
-	}
+	return SCompoundWidget::GetChildren();
 }
 
 void SRetainerWidget::SetRenderingPhase(int32 InPhase, int32 InPhaseCount)
@@ -354,18 +334,27 @@ void SRetainerWidget::SetRenderingPhase(int32 InPhase, int32 InPhaseCount)
 void SRetainerWidget::RequestRender()
 {
 	bRenderRequested = true;
+	InvalidateRoot();
 }
 
-bool SRetainerWidget::PaintRetainedContent(const FPaintArgs& Args, const FGeometry& AllottedGeometry)
+bool SRetainerWidget::PaintRetainedContent(const FSlateInvalidationContext& Context, const FGeometry& AllottedGeometry)
 {
 	if (RenderOnPhase)
 	{
 		if (LastTickedFrame != GFrameCounter && (GFrameCounter % PhaseCount) == Phase)
 		{
+			// If doing some phase based invalidation, just redraw everything again
+			InvalidateRoot();
 			bRenderRequested = true;
 		}
 	}
 
+	if (RenderOnInvalidation)
+	{
+		// the invalidation root will take care of whether or not we actually rendered
+		bRenderRequested = true;
+	}
+
 	if (Shared_MaxRetainerWorkPerFrame > 0)
 	{
 		if (Shared_RetainerWorkThisFrame.TryGetValue(0) > Shared_MaxRetainerWorkPerFrame)
@@ -384,7 +373,7 @@ bool SRetainerWidget::PaintRetainedContent(const FPaintArgs& Args, const FGeomet
 		bRenderRequested = true;
 	}
 
-	if ( bRenderRequested )
+	if (bRenderRequested)
 	{
 		// In order to get material parameter collections to function properly, we need the current world's Scene
 		// properly propagated through to any widgets that depend on that functionality. The SceneViewport and RetainerWidget the 
@@ -411,16 +400,6 @@ bool SRetainerWidget::PaintRetainedContent(const FPaintArgs& Args, const FGeomet
 
 		const FVector2D ViewOffset = PaintGeometry.DrawPosition.RoundToVector();
 
-		// Keep the visibilities the same, the proxy window should maintain the same visible/non-visible hit-testing of the retainer.
-		Window->SetVisibility(GetVisibility());
-
-		// Need to prepass.
-		Window->SlatePrepass(AllottedGeometry.Scale);
-
-		// Reset the cached node pool index so that we effectively reset the pool.
-		LastUsedCachedNodeIndex = 0;
-		RootCacheNode = nullptr;
-
 		UTextureRenderTarget2D* RenderTarget = RenderingResources->RenderTarget;
 		FWidgetRenderer* WidgetRenderer = RenderingResources->WidgetRenderer;
 
@@ -448,36 +427,28 @@ bool SRetainerWidget::PaintRetainedContent(const FPaintArgs& Args, const FGeomet
 				const float Scale = AllottedGeometry.Scale;
 
 				const FVector2D DrawSize = FVector2D(RenderTargetWidth, RenderTargetHeight);
-				const FGeometry WindowGeometry = FGeometry::MakeRoot(DrawSize * ( 1 / Scale ), FSlateLayoutTransform(Scale, PaintGeometry.DrawPosition));
+				//const FGeometry WindowGeometry = FGeometry::MakeRoot(DrawSize * ( 1 / Scale ), FSlateLayoutTransform(Scale, PaintGeometry.DrawPosition));
 
 				// Update the surface brush to match the latest size.
 				SurfaceBrush.ImageSize = DrawSize;
 
 				WidgetRenderer->ViewOffset = -ViewOffset;
 
-				SRetainerWidget* MutableThis = const_cast<SRetainerWidget*>(this);
-				TSharedRef<SRetainerWidget> SharedMutableThis = SharedThis(MutableThis);
-				
-				FPaintArgs PaintArgs(*this, Args.GetGrid(), Args.GetWindowToDesktopTransform(), FApp::GetCurrentTime(), Args.GetDeltaTime());
-
-				RootCacheNode = CreateCacheNode();
-				RootCacheNode->Initialize(Args, SharedMutableThis, WindowGeometry);
+				bool bRepaintedWidgets = WidgetRenderer->DrawInvalidationRoot(VirtualWindow, RenderTarget, *this, Context, GDeferRetainedRenderingRenderThread != 0);
 
-				WidgetRenderer->DrawWindow(
-					PaintArgs.EnableCaching(SharedMutableThis, RootCacheNode, true, true),
-					RenderTarget,
-					Window.ToSharedRef(),
-					WindowGeometry,
-					WindowGeometry.GetLayoutBoundingRect(),
-					TimeSinceLastDraw,
-					GDeferRetainedRenderingRenderThread!=0);
+#if WITH_SLATE_DEBUGGING
+				if (bRepaintedWidgets)
+				{
+					FSlateDebugging::DrawInvalidationRoot(*this, Context.IncomingLayerId+1, *Context.WindowElementList);
+				}
+#endif
 
-				bRenderRequested = false;
+				bRenderRequested = true;
 				Shared_WaitingToRender.Remove(this);
 
 				LastDrawTime = FApp::GetCurrentTime();
 
-				return true;
+				return bRepaintedWidgets;
 			}
 		}
 	}
@@ -487,31 +458,39 @@ bool SRetainerWidget::PaintRetainedContent(const FPaintArgs& Args, const FGeomet
 
 int32 SRetainerWidget::OnPaint(const FPaintArgs& Args, const FGeometry& AllottedGeometry, const FSlateRect& MyCullingRect, FSlateWindowElementList& OutDrawElements, int32 LayerId, const FWidgetStyle& InWidgetStyle, bool bParentEnabled) const
 {
-	STAT(FScopeCycleCounter PaintCycleCounter(MyStatId);)
+	STAT(FScopeCycleCounter PaintCycleCounter(MyStatId););
 
-	SRetainerWidget* MutableThis = const_cast<SRetainerWidget*>( this );
+	SRetainerWidget* MutableThis = const_cast<SRetainerWidget*>(this);
 
-	MutableThis->RefreshRenderingMode();
-
-	if ( bEnableRetainedRendering && IsAnythingVisibleToRender() )
+	if (bEnableRetainedRendering && IsAnythingVisibleToRender())
 	{
-		SCOPE_CYCLE_COUNTER( STAT_SlateRetainerWidgetPaint );
+		SCOPE_CYCLE_COUNTER(STAT_SlateRetainerWidgetPaint);
+
+		// Copy hit test grid settings from the root
+		const bool bHittestCleared = HittestGrid.SetHittestArea(Args.RootGrid.GetGridOrigin(), Args.RootGrid.GetGridSize(), Args.RootGrid.GetGridWindowOrigin());
+		FPaintArgs NewArgs = Args.WithNewHitTestGrid(HittestGrid);
 
-		TSharedRef<SRetainerWidget> SharedMutableThis = SharedThis(MutableThis);
+		FSlateInvalidationContext Context(OutDrawElements, InWidgetStyle);
+		Context.bParentEnabled = bParentEnabled;
+		Context.bAllowFastPathUpdate = true;
+		Context.LayoutScaleMultiplier = GetPrepassLayoutScaleMultiplier();
+		Context.PaintArgs = &NewArgs;
+		Context.IncomingLayerId = LayerId;
+		Context.CullingRect = MyCullingRect;
 
-		const bool bNewFramePainted = MutableThis->PaintRetainedContent(Args, AllottedGeometry);
+		MutableThis->PaintRetainedContent(Context, AllottedGeometry);
 
 		UTextureRenderTarget2D* RenderTarget = RenderingResources->RenderTarget;
 
-		if ( RenderTarget->GetSurfaceWidth() >= 1 && RenderTarget->GetSurfaceHeight() >= 1 )
+		if (RenderTarget->GetSurfaceWidth() >= 1 && RenderTarget->GetSurfaceHeight() >= 1)
 		{
-			const FLinearColor ComputedColorAndOpacity(InWidgetStyle.GetColorAndOpacityTint() * ColorAndOpacity.Get() * SurfaceBrush.GetTint(InWidgetStyle));
+			const FLinearColor ComputedColorAndOpacity(Context.WidgetStyle.GetColorAndOpacityTint() * ColorAndOpacity.Get() * SurfaceBrush.GetTint(Context.WidgetStyle));
 			// Retainer widget uses pre-multiplied alpha, so pre-multiply the color by the alpha to respect opacity.
 			const FLinearColor PremultipliedColorAndOpacity(ComputedColorAndOpacity * ComputedColorAndOpacity.A);
 
 			FWidgetRenderer* WidgetRenderer = RenderingResources->WidgetRenderer;
 			UMaterialInstanceDynamic* DynamicEffect = RenderingResources->DynamicEffect;
-	
+
 			const bool bDynamicMaterialInUse = (DynamicEffect != nullptr);
 			if (bDynamicMaterialInUse)
 			{
@@ -519,8 +498,8 @@ int32 SRetainerWidget::OnPaint(const FPaintArgs& Args, const FGeometry& Allotted
 			}
 
 			FSlateDrawElement::MakeBox(
-				OutDrawElements,
-				LayerId,
+				*Context.WindowElementList,
+				Context.IncomingLayerId,
 				AllottedGeometry.ToPaintGeometry(),
 				&SurfaceBrush,
 				// We always write out the content in gamma space, so when we render the final version we need to
@@ -528,24 +507,17 @@ int32 SRetainerWidget::OnPaint(const FPaintArgs& Args, const FGeometry& Allotted
 				ESlateDrawEffect::PreMultipliedAlpha | ESlateDrawEffect::NoGamma,
 				FLinearColor(PremultipliedColorAndOpacity.R, PremultipliedColorAndOpacity.G, PremultipliedColorAndOpacity.B, PremultipliedColorAndOpacity.A)
 			);
-			
-			if (RootCacheNode)
-			{
-				RootCacheNode->RecordHittestGeometry(Args.GetGrid(), Args.GetLastHitTestIndex(), LayerId, FVector2D(0, 0));
-			}
-
-			// Any deferred painted elements of the retainer should be drawn directly by the main renderer, not rendered into the render target,
-			// as most of those sorts of things will break the rendering rect, things like tooltips, and popup menus.
-			for ( auto& DeferredPaint : WidgetRenderer->DeferredPaints )
-			{
-				OutDrawElements.QueueDeferredPainting(DeferredPaint->Copy(Args));
-			}
 		}
 
-		return LayerId;
-	}
+		// add our widgets to the root hit test grid
+		Context.PaintArgs->RootGrid.AppendGrid(HittestGrid);
 
-	return SCompoundWidget::OnPaint(Args, AllottedGeometry, MyCullingRect, OutDrawElements, LayerId, InWidgetStyle, bParentEnabled);
+		return Context.IncomingLayerId;
+	}
+	else
+	{
+		return SCompoundWidget::OnPaint(Args, AllottedGeometry, MyCullingRect, OutDrawElements, LayerId, InWidgetStyle, bParentEnabled);
+	}
 }
 
 FVector2D SRetainerWidget::ComputeDesiredSize(float LayoutScaleMuliplier) const
@@ -559,3 +531,43 @@ FVector2D SRetainerWidget::ComputeDesiredSize(float LayoutScaleMuliplier) const
 		return SCompoundWidget::ComputeDesiredSize(LayoutScaleMuliplier);
 	}
 }
+
+bool SRetainerWidget::ComputeVolatility() const
+{
+	// We need to be volatile when global invalidation is turned on so that we have a chance to update
+	// @todo instead of being volatile it may be possible to move this to tick
+	return GSlateEnableGlobalInvalidation;
+}
+
+void SRetainerWidget::OnGlobalInvalidationToggled(bool bGlobalInvalidationEnabled)
+{
+	InvalidateRoot();
+
+	ClearAllFastPathData();
+}
+
+/*
+void SRetainerWidget::Advanced_InvalidateRoot()
+{
+	InvalidateRoot();
+}
+*/
+
+bool SRetainerWidget::CustomPrepass(float LayoutScaleMultiplier)
+{
+	if (bEnableRetainedRendering)
+	{
+		ProcessInvalidation();
+
+		return NeedsPrepass();
+	}
+	else
+	{
+		return true;
+	}
+}
+
+int32 SRetainerWidget::PaintSlowPath(const FSlateInvalidationContext& Context)
+{
+	return SCompoundWidget::OnPaint(*Context.PaintArgs, GetPaintSpaceGeometry(), Context.CullingRect, *Context.WindowElementList, Context.IncomingLayerId, Context.WidgetStyle, Context.bParentEnabled);
+}
diff --git a/Engine/Source/Runtime/UMG/Private/Slate/WidgetRenderer.cpp b/Engine/Source/Runtime/UMG/Private/Slate/WidgetRenderer.cpp
--- a/Engine/Source/Runtime/UMG/Private/Slate/WidgetRenderer.cpp
+++ b/Engine/Source/Runtime/UMG/Private/Slate/WidgetRenderer.cpp
@@ -68,7 +68,7 @@ UTextureRenderTarget2D* FWidgetRenderer::DrawWidget(const TSharedRef<SWidget>& W
 	{
 		UTextureRenderTarget2D* RenderTarget = FWidgetRenderer::CreateTargetFor(DrawSize, TF_Bilinear, bUseGammaSpace);
 
-		DrawWidget(RenderTarget, Widget, DrawSize, 0);
+		DrawWidget(RenderTarget, Widget, DrawSize, 0, false);
 
 		return RenderTarget;
 	}
@@ -115,12 +115,12 @@ void FWidgetRenderer::DrawWidget(
 	bool bDeferRenderTargetUpdate)
 {
 	TSharedRef<SVirtualWindow> Window = SNew(SVirtualWindow).Size(DrawSize);
-	TSharedRef<FHittestGrid> HitTestGrid = MakeShareable(new FHittestGrid());
+	TUniquePtr<FHittestGrid> HitTestGrid = MakeUnique<FHittestGrid>();
 
 	Window->SetContent(Widget);
 	Window->Resize(DrawSize);
 
-	DrawWindow(RenderTarget, HitTestGrid, Window, Scale, DrawSize, DeltaTime, bDeferRenderTargetUpdate);
+	DrawWindow(RenderTarget, *HitTestGrid, Window, Scale, DrawSize, DeltaTime, bDeferRenderTargetUpdate);
 }
 
 void FWidgetRenderer::DrawWidget(
@@ -136,7 +136,7 @@ void FWidgetRenderer::DrawWidget(
 
 void FWidgetRenderer::DrawWindow(
 	FRenderTarget* RenderTarget,
-	TSharedRef<FHittestGrid> HitTestGrid,
+	FHittestGrid& HitTestGrid,
 	TSharedRef<SWindow> Window,
 	float Scale,
 	FVector2D DrawSize,
@@ -159,7 +159,7 @@ void FWidgetRenderer::DrawWindow(
 
 void FWidgetRenderer::DrawWindow(
 	UTextureRenderTarget2D* RenderTarget,
-	TSharedRef<FHittestGrid> HitTestGrid,
+	FHittestGrid& HitTestGrid,
 	TSharedRef<SWindow> Window,
 	float Scale,
 	FVector2D DrawSize,
@@ -171,20 +171,20 @@ void FWidgetRenderer::DrawWindow(
 
 void FWidgetRenderer::DrawWindow(
 	FRenderTarget* RenderTarget,
-	TSharedRef<FHittestGrid> HitTestGrid,
+	FHittestGrid& HitTestGrid,
 	TSharedRef<SWindow> Window,
 	FGeometry WindowGeometry,
 	FSlateRect WindowClipRect,
 	float DeltaTime,
 	bool bDeferRenderTargetUpdate)
 {
-	FPaintArgs PaintArgs(Window.Get(), HitTestGrid.Get(), FVector2D::ZeroVector, FApp::GetCurrentTime(), DeltaTime);
+	FPaintArgs PaintArgs(nullptr, HitTestGrid, FVector2D::ZeroVector, FApp::GetCurrentTime(), DeltaTime);
 	DrawWindow(PaintArgs, RenderTarget, Window, WindowGeometry, WindowClipRect, DeltaTime, bDeferRenderTargetUpdate);
 }
 
 void FWidgetRenderer::DrawWindow(
 	UTextureRenderTarget2D* RenderTarget,
-	TSharedRef<FHittestGrid> HitTestGrid,
+	FHittestGrid& HitTestGrid,
 	TSharedRef<SWindow> Window,
 	FGeometry WindowGeometry,
 	FSlateRect WindowClipRect,
@@ -214,17 +214,19 @@ void FWidgetRenderer::DrawWindow(
 		    // Ticking can cause geometry changes.  Recompute
 		    Window->SlatePrepass(WindowGeometry.Scale);
 	    }
-    
+	
+		PaintArgs.GetHittestGrid().SetHittestArea(WindowClipRect.GetTopLeft(), WindowClipRect.GetSize());
+
 		if ( bClearHitTestGrid )
 		{
 			// Prepare the test grid 
-			PaintArgs.GetGrid().ClearGridForNewFrame(WindowClipRect);
+			PaintArgs.GetHittestGrid().Clear();
 		}
     
 	    // Get the free buffer & add our virtual window
 	    FSlateDrawBuffer& DrawBuffer = Renderer->GetDrawBuffer();
 	    FSlateWindowElementList& WindowElementList = DrawBuffer.AddWindowElementList(Window);
-    
+
 	    int32 MaxLayerId = 0;
 	    {
 		    // Paint the window
@@ -244,7 +246,7 @@ void FWidgetRenderer::DrawWindow(
 
 		DrawBuffer.ViewOffset = ViewOffset;
 
-		FRenderThreadUpdateContext Context =
+		FRenderThreadUpdateContext RenderThreadUpdateContext =
 		{
 			&DrawBuffer,
 			static_cast<float>(FApp::GetCurrentTime() - GStartTime),
@@ -255,7 +257,7 @@ void FWidgetRenderer::DrawWindow(
 			bClearTarget
 		};
 
-		FSlateApplication::Get().GetRenderer()->AddWidgetRendererUpdate(Context, bDeferRenderTargetUpdate);
+		FSlateApplication::Get().GetRenderer()->AddWidgetRendererUpdate(RenderThreadUpdateContext, bDeferRenderTargetUpdate);
 	}
 #endif // !UE_SERVER
 }
@@ -271,3 +273,69 @@ void FWidgetRenderer::DrawWindow(
 {
 	DrawWindow(PaintArgs, RenderTarget->GameThread_GetRenderTargetResource(), Window, WindowGeometry, WindowClipRect, DeltaTime, bDeferRenderTargetUpdate);
 }
+
+bool FWidgetRenderer::DrawInvalidationRoot(TSharedRef<SVirtualWindow>& VirtualWindow, UTextureRenderTarget2D* RenderTarget, FSlateInvalidationRoot& Root, const FSlateInvalidationContext& Context, bool bDeferRenderTargetUpdate)
+{
+	bool bRepaintedWidgets = false;
+#if !UE_SERVER
+	FSlateRenderer* MainSlateRenderer = FSlateApplication::Get().GetRenderer();
+	FScopeLock ScopeLock(MainSlateRenderer->GetResourceCriticalSection());
+
+
+	if (LIKELY(FApp::CanEverRender()))
+	{
+		// Need to set a new window element list so make a copy
+		FSlateInvalidationContext ContextCopy = Context;
+		
+		// Get the free buffer & add our virtual window
+		FSlateDrawBuffer& DrawBuffer = Renderer->GetDrawBuffer();
+		FSlateWindowElementList& WindowElementList = DrawBuffer.AddWindowElementList(VirtualWindow);
+
+		ContextCopy.WindowElementList = &WindowElementList;
+		FSlateInvalidationResult Result = Root.PaintInvalidationRoot(ContextCopy);
+
+		const int32 MaxLayerId = Result.MaxLayerIdPainted;
+
+		if(Result.bRepaintedWidgets)
+		{
+			//MaxLayerId = WindowElementList.PaintDeferred(MaxLayerId);
+			DeferredPaints = WindowElementList.GetDeferredPaintList();
+
+			Renderer->DrawWindow_GameThread(DrawBuffer);
+
+			DrawBuffer.ViewOffset = ViewOffset;
+
+			FRenderThreadUpdateContext RenderThreadUpdateContext =
+			{
+				&DrawBuffer,
+				static_cast<float>(FApp::GetCurrentTime() - GStartTime),
+				static_cast<float>(FApp::GetDeltaTime()),
+				static_cast<float>(FPlatformTime::Seconds() - GStartTime),
+				RenderTarget->GameThread_GetRenderTargetResource(),
+				Renderer.Get(),
+				bClearTarget
+			};
+
+			bRepaintedWidgets = Result.bRepaintedWidgets;
+			FSlateApplication::Get().GetRenderer()->AddWidgetRendererUpdate(RenderThreadUpdateContext, bDeferRenderTargetUpdate);
+
+
+			// Any deferred painted elements of the retainer should be drawn directly by the main renderer, not rendered into the render target,
+			// as most of those sorts of things will break the rendering rect, things like tooltips, and popup menus.
+			for (auto& DeferredPaint : DeferredPaints)
+			{
+				Context.WindowElementList->QueueDeferredPainting(DeferredPaint->Copy(*Context.PaintArgs));
+			}
+		}
+		else
+		{
+			WindowElementList.ResetElementList();
+			DrawBuffer.Unlock();
+		}
+	}
+#endif // !UE_SERVER
+
+	return bRepaintedWidgets;
+}
+
+
diff --git a/Engine/Source/Runtime/UMG/Private/UserWidget.cpp b/Engine/Source/Runtime/UMG/Private/UserWidget.cpp
--- a/Engine/Source/Runtime/UMG/Private/UserWidget.cpp
+++ b/Engine/Source/Runtime/UMG/Private/UserWidget.cpp
@@ -39,7 +39,7 @@ static FWidgetStyle NullStyle;
 
 FSlateWindowElementList& GetNullElementList()
 {
-	static FSlateWindowElementList NullElementList;
+	static FSlateWindowElementList NullElementList(nullptr);
 	return NullElementList;
 }
 
@@ -618,6 +618,7 @@ void UUserWidget::Invalidate(EInvalidateWidget InvalidateReason)
 	TSharedPtr<SWidget> CachedWidget = GetCachedWidget();
 	if (CachedWidget.IsValid())
 	{
+		UpdateCanTick();
 		CachedWidget->Invalidate(InvalidateReason);
 	}
 }
@@ -1533,6 +1534,7 @@ void UUserWidget::CancelLatentActions()
 	{
 		World->GetLatentActionManager().RemoveActionsForObject(this);
 		World->GetTimerManager().ClearAllTimersForObject(this);
+		UpdateCanTick();
 	}
 }
 
diff --git a/Engine/Source/Runtime/UMG/Public/Components/InvalidationBox.h b/Engine/Source/Runtime/UMG/Public/Components/InvalidationBox.h
--- a/Engine/Source/Runtime/UMG/Public/Components/InvalidationBox.h
+++ b/Engine/Source/Runtime/UMG/Public/Components/InvalidationBox.h
@@ -25,7 +25,7 @@ public:
 	/**
 	 * 
 	 */
-	UFUNCTION(BlueprintCallable, Category="Invalidation Box")
+	UFUNCTION(BlueprintCallable, Category="Invalidation Box", meta=(DeprecatedFunction))
 	void InvalidateCache();
 
 	/**
@@ -70,8 +70,8 @@ protected:
 	 * this adds extra overhead to drawing them every frame.  However, in cases where
 	 * the position of the invalidation boxes changes every frame this can be a big savings.
 	 */
-	UPROPERTY(EditAnywhere, BlueprintReadOnly, Category="Caching")
-	bool CacheRelativeTransforms;
+	UPROPERTY()
+	bool CacheRelativeTransforms_DEPRECATED;
 
 protected:
 	TSharedPtr<class SInvalidationPanel> MyInvalidationPanel;
diff --git a/Engine/Source/Runtime/UMG/Public/Components/RetainerBox.h b/Engine/Source/Runtime/UMG/Public/Components/RetainerBox.h
--- a/Engine/Source/Runtime/UMG/Public/Components/RetainerBox.h
+++ b/Engine/Source/Runtime/UMG/Public/Components/RetainerBox.h
@@ -100,7 +100,7 @@ public:
 	virtual const FText GetPaletteCategory() override;
 #endif
 
-	const FGeometry& GetCachedAllottedGeometry() const;
+	FGeometry GetCachedAllottedGeometry() const;
 
 protected:
 
diff --git a/Engine/Source/Runtime/UMG/Public/Components/Widget.h b/Engine/Source/Runtime/UMG/Public/Components/Widget.h
--- a/Engine/Source/Runtime/UMG/Public/Components/Widget.h
+++ b/Engine/Source/Runtime/UMG/Public/Components/Widget.h
@@ -670,9 +670,14 @@ public:
 	 * or what are referred to as hysteresis problems, both caused by depending on geometry from the previous frame
 	 * being used to advise how to layout a dependent object the current frame.
 	 */
-	UFUNCTION(BlueprintCallable, Category="Widget")
+	UFUNCTION(BlueprintCallable, Category="Widget",  meta=(DeprecatedFunction, DeprecationMessage = "Use GetTickSpaceGeometry instead"))
 	const FGeometry& GetCachedGeometry() const;
 
+	UFUNCTION(BlueprintCallable, Category="Widget")
+	const FGeometry& GetTickSpaceGeometry() const;
+
+	UFUNCTION(BlueprintCallable, Category="Widget")
+	const FGeometry& GetPaintSpaceGeometry() const;
 	/**
 	 * Gets the underlying slate widget or constructs it if it doesn't exist.  If you're looking to replace
 	 * what slate widget gets constructed look for RebuildWidget.  For extremely special cases where you actually
diff --git a/Engine/Source/Runtime/UMG/Public/Slate/SRetainerWidget.h b/Engine/Source/Runtime/UMG/Public/Slate/SRetainerWidget.h
--- a/Engine/Source/Runtime/UMG/Public/Slate/SRetainerWidget.h
+++ b/Engine/Source/Runtime/UMG/Public/Slate/SRetainerWidget.h
@@ -36,7 +36,7 @@ DECLARE_MULTICAST_DELEGATE( FOnRetainedModeChanged );
  * frequency of the main game render.  It also has the side benefit of allow materials
  * to be applied to the render target after drawing the widgets to apply a simple post process.
  */
-class UMG_API SRetainerWidget : public SCompoundWidget,  public ILayoutCache
+class UMG_API SRetainerWidget : public SCompoundWidget,  public FSlateInvalidationRoot
 {
 public:
 	static int32 Shared_MaxRetainerWorkPerFrame;
@@ -44,7 +44,7 @@ public:
 public:
 	SLATE_BEGIN_ARGS(SRetainerWidget)
 	{
-		_Visibility = EVisibility::Visible;
+		_Visibility = EVisibility::SelfHitTestInvisible;
 		_Phase = 0;
 		_PhaseCount = 1;
 		_RenderOnPhase = true;
@@ -79,25 +79,25 @@ public:
 
 	void SetTextureParameter(FName TextureParameter);
 
-	// ILayoutCache overrides
-	virtual void InvalidateWidget(SWidget* InvalidateWidget) override;
-	virtual FCachedWidgetNode* CreateCacheNode() const override;
-	// End ILayoutCache
-
-	// SWidget
+	/** SWidget interface */
 	virtual FChildren* GetChildren() override;
-	virtual bool ComputeVolatility() const override;
-	// SWidget
+	virtual FChildren* GetAllChildren() override;
 
-	virtual bool PaintRetainedContent(const FPaintArgs& Args, const FGeometry& AllottedGeometry);
+	/** FInvalidationRoot interface */
+	virtual bool PaintRetainedContent(const FSlateInvalidationContext& Context, const FGeometry& AllottedGeometry);
 
 	void SetWorld(UWorld* World);
 
 protected:
-	// BEGIN SLeafWidget interface
+	/** SCompoundWidget interface */
 	virtual int32 OnPaint(const FPaintArgs& Args, const FGeometry& AllottedGeometry, const FSlateRect& MyCullingRect, FSlateWindowElementList& OutDrawElements, int32 LayerId, const FWidgetStyle& InWidgetStyle, bool bParentEnabled) const override;
 	virtual FVector2D ComputeDesiredSize(float Scale) const override;
-	// END SLeafWidget interface
+	virtual bool ComputeVolatility() const override;
+	virtual bool Advanced_IsInvalidationRoot() const { return bEnableRetainedRendering; }
+	virtual bool CustomPrepass(float LayoutScaleMultiplier) override;
+
+	/** FSlateInvalidationRoot interface */
+	int32 PaintSlowPath(const FSlateInvalidationContext& Context);
 
 	void RefreshRenderingMode();
 	bool ShouldBeRenderingOffscreen() const;
@@ -105,6 +105,7 @@ protected:
 	void OnRetainerModeChanged();
 	void OnGlobalInvalidate();
 private:
+	void OnGlobalInvalidationToggled(bool bGlobalInvalidationEnabled);
 #if !UE_BUILD_SHIPPING
 	static void OnRetainerModeCVarChanged( IConsoleVariable* CVar );
 	static FOnRetainedModeChanged OnRetainerModeChangedDelegate;
@@ -117,14 +118,17 @@ private:
 
 	void UpdateWidgetRenderer();
 
-	mutable TSharedPtr<SWidget> MyWidget;
+	TSharedPtr<SWidget> MyWidget;
+	TSharedRef<SVirtualWindow> VirtualWindow;
+
+	mutable FHittestGrid HittestGrid;
 
-	bool bEnableRetainedRenderingDesire;
-	bool bEnableRetainedRendering;
-	
 	int32 Phase;
 	int32 PhaseCount;
 
+	bool bEnableRetainedRenderingDesire;
+	bool bEnableRetainedRendering;
+
 	bool RenderOnPhase;
 	bool RenderOnInvalidation;
 
@@ -133,7 +137,6 @@ private:
 	double LastDrawTime;
 	int64 LastTickedFrame;
 
-	TSharedPtr<SVirtualWindow> Window;
 	TWeakObjectPtr<UWorld> OuterWorld;
 
 	FRetainerWidgetRenderingResources* RenderingResources;
@@ -146,8 +149,4 @@ private:
 
 	static TArray<SRetainerWidget*, TInlineAllocator<3>> Shared_WaitingToRender;
 	static TFrameValue<int32> Shared_RetainerWorkThisFrame;
-
-	mutable FCachedWidgetNode* RootCacheNode;
-	mutable TArray< FCachedWidgetNode* > NodePool;
-	mutable int32 LastUsedCachedNodeIndex;
 };
diff --git a/Engine/Source/Runtime/UMG/Public/Slate/WidgetRenderer.h b/Engine/Source/Runtime/UMG/Public/Slate/WidgetRenderer.h
--- a/Engine/Source/Runtime/UMG/Public/Slate/WidgetRenderer.h
+++ b/Engine/Source/Runtime/UMG/Public/Slate/WidgetRenderer.h
@@ -81,7 +81,7 @@ public:
 
 	void DrawWindow(
 		FRenderTarget* RenderTarget,
-		TSharedRef<FHittestGrid> HitTestGrid,
+		FHittestGrid& HitTestGrid,
 		TSharedRef<SWindow> Window,
 		float Scale,
 		FVector2D DrawSize,
@@ -90,7 +90,7 @@ public:
 
 	void DrawWindow(
 		UTextureRenderTarget2D* RenderTarget,
-		TSharedRef<FHittestGrid> HitTestGrid,
+		FHittestGrid& HitTestGrid,
 		TSharedRef<SWindow> Window,
 		float Scale,
 		FVector2D DrawSize,
@@ -99,7 +99,7 @@ public:
 
 	void DrawWindow(
 		FRenderTarget* RenderTarget,
-		TSharedRef<FHittestGrid> HitTestGrid,
+		FHittestGrid& HitTestGrid,
 		TSharedRef<SWindow> Window,
 		FGeometry WindowGeometry,
 		FSlateRect WindowClipRect,
@@ -108,7 +108,7 @@ public:
 
 	void DrawWindow(
 		UTextureRenderTarget2D* RenderTarget,
-		TSharedRef<FHittestGrid> HitTestGrid,
+		FHittestGrid& HitTestGrid,
 		TSharedRef<SWindow> Window,
 		FGeometry WindowGeometry,
 		FSlateRect WindowClipRect,
@@ -133,6 +133,13 @@ public:
 		float DeltaTime,
 		bool bDeferRenderTargetUpdate = false);
 
+	bool DrawInvalidationRoot(
+		TSharedRef<SVirtualWindow>& VirtualWindow,
+		UTextureRenderTarget2D* RenderTarget,
+		FSlateInvalidationRoot& Root,
+		const FSlateInvalidationContext& Context,
+		bool bDeferRenderTargetUpdate = false);
+
 	TArray< TSharedPtr<FSlateWindowElementList::FDeferredPaint> > DeferredPaints;
 private:
 	/** The slate 3D renderer used to render the user slate widget */
