#UE4 Removing deprecation on GetCachedGeometry until content can be converted.

diff --git a/Engine/Source/Runtime/UMG/Public/Components/Widget.h b/Engine/Source/Runtime/UMG/Public/Components/Widget.h
--- a/Engine/Source/Runtime/UMG/Public/Components/Widget.h
+++ b/Engine/Source/Runtime/UMG/Public/Components/Widget.h
@@ -670,7 +670,7 @@ public:
 	 * or what are referred to as hysteresis problems, both caused by depending on geometry from the previous frame
 	 * being used to advise how to layout a dependent object the current frame.
 	 */
-	UFUNCTION(BlueprintCallable, Category="Widget",  meta=(DeprecatedFunction, DeprecationMessage = "Use GetTickSpaceGeometry instead"))
+	UFUNCTION(BlueprintCallable, Category="Widget")
 	const FGeometry& GetCachedGeometry() const;
 
 	UFUNCTION(BlueprintCallable, Category="Widget")
