#UE4 Remove deprecation of InvalidateCache until content can be converted

diff --git a/Engine/Source/Runtime/UMG/Public/Components/InvalidationBox.h b/Engine/Source/Runtime/UMG/Public/Components/InvalidationBox.h
--- a/Engine/Source/Runtime/UMG/Public/Components/InvalidationBox.h
+++ b/Engine/Source/Runtime/UMG/Public/Components/InvalidationBox.h
@@ -25,7 +25,7 @@ public:
 	/**
 	 * 
 	 */
-	UFUNCTION(BlueprintCallable, Category="Invalidation Box", meta=(DeprecatedFunction))
+	UFUNCTION(BlueprintCallable, Category="Invalidation Box")
 	void InvalidateCache();
 
 	/**
