Several widgets were running NativeDestruct code after reinstancing from blueprint compiles, turns out the enum flags for the designer won't serialize if they're flags and resolve to more than one enum value, so instead the enum flags are now stored as a uint8 to make UObject serializer happy, but as far as the outside world is concerned, they are UENUMs, and I've converted them over to be enum class instead of the older style namespace::type method, this allows the values to properly persist after reinstancing.

diff --git a/Engine/Source/Editor/UMGEditor/Private/WidgetBlueprintEditor.cpp b/Engine/Source/Editor/UMGEditor/Private/WidgetBlueprintEditor.cpp
--- a/Engine/Source/Editor/UMGEditor/Private/WidgetBlueprintEditor.cpp
+++ b/Engine/Source/Editor/UMGEditor/Private/WidgetBlueprintEditor.cpp
@@ -1065,20 +1065,20 @@ TArray< TFunction<void()> >& FWidgetBlueprintEditor::GetQueuedDesignerActions()
 	return QueuedDesignerActions;
 }
 
-EWidgetDesignFlags::Type FWidgetBlueprintEditor::GetCurrentDesignerFlags() const
+EWidgetDesignFlags FWidgetBlueprintEditor::GetCurrentDesignerFlags() const
 {
-	EWidgetDesignFlags::Type Flags = EWidgetDesignFlags::Designing;
+	EWidgetDesignFlags Flags = EWidgetDesignFlags::Designing;
 	
 	if ( bShowDashedOutlines )
 	{
-		Flags = ( EWidgetDesignFlags::Type )( Flags | EWidgetDesignFlags::ShowOutline );
+		Flags |= EWidgetDesignFlags::ShowOutline;
 	}
 
 	if ( const UWidgetDesignerSettings* Designer = GetDefault<UWidgetDesignerSettings>() )
 	{
 		if ( Designer->bExecutePreConstructEvent )
 		{
-			Flags = ( EWidgetDesignFlags::Type )( Flags | EWidgetDesignFlags::ExecutePreConstruct );
+			Flags |= EWidgetDesignFlags::ExecutePreConstruct;
 		}
 	}
 
diff --git a/Engine/Source/Editor/UMGEditor/Public/WidgetBlueprintEditor.h b/Engine/Source/Editor/UMGEditor/Public/WidgetBlueprintEditor.h
--- a/Engine/Source/Editor/UMGEditor/Public/WidgetBlueprintEditor.h
+++ b/Engine/Source/Editor/UMGEditor/Public/WidgetBlueprintEditor.h
@@ -165,7 +165,7 @@ public:
 	TArray< TFunction<void()> >& GetQueuedDesignerActions();
 
 	/** Get the current designer flags that are in effect for the current user widget we're editing. */
-	EWidgetDesignFlags::Type GetCurrentDesignerFlags() const;
+	EWidgetDesignFlags GetCurrentDesignerFlags() const;
 
 	bool GetShowDashedOutlines() const;
 	void SetShowDashedOutlines(bool Value);
diff --git a/Engine/Source/Runtime/UMG/Private/Components/NamedSlotInterface.cpp b/Engine/Source/Runtime/UMG/Private/Components/NamedSlotInterface.cpp
--- a/Engine/Source/Runtime/UMG/Private/Components/NamedSlotInterface.cpp
+++ b/Engine/Source/Runtime/UMG/Private/Components/NamedSlotInterface.cpp
@@ -40,8 +40,10 @@ void INamedSlotInterface::ReleaseNamedSlotSlateResources(bool bReleaseChildren)
 		}
 	}
 }
+
 #if WITH_EDITOR
-void INamedSlotInterface::SetNamedSlotDesignerFlags(EWidgetDesignFlags::Type NewFlags)
+
+void INamedSlotInterface::SetNamedSlotDesignerFlags(EWidgetDesignFlags NewFlags)
 {
 	TArray<FName> SlotNames;
 	GetSlotNames(SlotNames);
@@ -54,4 +56,5 @@ void INamedSlotInterface::SetNamedSlotDesignerFlags(EWidgetDesignFlags::Type New
 		}
 	}
 }
+
 #endif
\ No newline at end of file
diff --git a/Engine/Source/Runtime/UMG/Private/Components/PanelWidget.cpp b/Engine/Source/Runtime/UMG/Private/Components/PanelWidget.cpp
--- a/Engine/Source/Runtime/UMG/Private/Components/PanelWidget.cpp
+++ b/Engine/Source/Runtime/UMG/Private/Components/PanelWidget.cpp
@@ -193,7 +193,7 @@ void UPanelWidget::ShiftChild(int32 Index, UWidget* Child)
 	Slots.Insert(Child->Slot, FMath::Clamp(Index, 0, Slots.Num()));
 }
 
-void UPanelWidget::SetDesignerFlags(EWidgetDesignFlags::Type NewFlags)
+void UPanelWidget::SetDesignerFlags(EWidgetDesignFlags NewFlags)
 {
 	Super::SetDesignerFlags(NewFlags);
 
diff --git a/Engine/Source/Runtime/UMG/Private/Components/Widget.cpp b/Engine/Source/Runtime/UMG/Private/Components/Widget.cpp
--- a/Engine/Source/Runtime/UMG/Private/Components/Widget.cpp
+++ b/Engine/Source/Runtime/UMG/Private/Components/Widget.cpp
@@ -154,7 +154,7 @@ UWidget::UWidget(const FObjectInitializer& ObjectInitializer)
 	bIsEnabled = true;
 	bIsVariable = true;
 #if WITH_EDITOR
-	DesignerFlags = EWidgetDesignFlags::None;
+	DesignerFlags = static_cast<uint8>(EWidgetDesignFlags::None);
 #endif
 	Visibility = ESlateVisibility::Visible;
 	RenderOpacity = 1.0f;
@@ -936,9 +936,9 @@ ULocalPlayer* UWidget::GetOwningLocalPlayer() const
 #undef LOCTEXT_NAMESPACE
 #define LOCTEXT_NAMESPACE "UMGEditor"
 
-void UWidget::SetDesignerFlags(EWidgetDesignFlags::Type NewFlags)
+void UWidget::SetDesignerFlags(EWidgetDesignFlags NewFlags)
 {
-	DesignerFlags = ( EWidgetDesignFlags::Type )( DesignerFlags | NewFlags );
+	DesignerFlags = static_cast<uint8>(GetDesignerFlags() | NewFlags);
 
 	INamedSlotInterface* NamedSlotWidget = Cast<INamedSlotInterface>(this);
 	if (NamedSlotWidget)
diff --git a/Engine/Source/Runtime/UMG/Private/UserWidget.cpp b/Engine/Source/Runtime/UMG/Private/UserWidget.cpp
--- a/Engine/Source/Runtime/UMG/Private/UserWidget.cpp
+++ b/Engine/Source/Runtime/UMG/Private/UserWidget.cpp
@@ -1322,7 +1322,7 @@ const FText UUserWidget::GetPaletteCategory()
 	return PaletteCategory;
 }
 
-void UUserWidget::SetDesignerFlags(EWidgetDesignFlags::Type NewFlags)
+void UUserWidget::SetDesignerFlags(EWidgetDesignFlags NewFlags)
 {
 	UWidget::SetDesignerFlags(NewFlags);
 
diff --git a/Engine/Source/Runtime/UMG/Public/Blueprint/UserWidget.h b/Engine/Source/Runtime/UMG/Public/Blueprint/UserWidget.h
--- a/Engine/Source/Runtime/UMG/Public/Blueprint/UserWidget.h
+++ b/Engine/Source/Runtime/UMG/Public/Blueprint/UserWidget.h
@@ -1073,7 +1073,7 @@ public:
 	virtual const FText GetPaletteCategory() override;
 	//~ End UWidget Interface
 
-	virtual void SetDesignerFlags(EWidgetDesignFlags::Type NewFlags) override;
+	virtual void SetDesignerFlags(EWidgetDesignFlags NewFlags) override;
 	virtual void OnDesignerChanged(const FDesignerChangedEventArgs& EventArgs) override;
 	virtual void PostEditChangeProperty(FPropertyChangedEvent& PropertyChangedEvent) override;
 
diff --git a/Engine/Source/Runtime/UMG/Public/Components/NamedSlotInterface.h b/Engine/Source/Runtime/UMG/Public/Components/NamedSlotInterface.h
--- a/Engine/Source/Runtime/UMG/Public/Components/NamedSlotInterface.h
+++ b/Engine/Source/Runtime/UMG/Public/Components/NamedSlotInterface.h
@@ -39,7 +39,8 @@ public:
 
 	/**  */
 	void ReleaseNamedSlotSlateResources(bool bReleaseChildren);
+
 #if WITH_EDITOR
-	void SetNamedSlotDesignerFlags(EWidgetDesignFlags::Type NewFlags);
+	void SetNamedSlotDesignerFlags(EWidgetDesignFlags NewFlags);
 #endif
 };
diff --git a/Engine/Source/Runtime/UMG/Public/Components/PanelWidget.h b/Engine/Source/Runtime/UMG/Public/Components/PanelWidget.h
--- a/Engine/Source/Runtime/UMG/Public/Components/PanelWidget.h
+++ b/Engine/Source/Runtime/UMG/Public/Components/PanelWidget.h
@@ -88,7 +88,7 @@ public:
 	 */
 	void ShiftChild(int32 Index, UWidget* Child);
 
-	void SetDesignerFlags(EWidgetDesignFlags::Type NewFlags);
+	void SetDesignerFlags(EWidgetDesignFlags NewFlags);
 
 #endif
 
diff --git a/Engine/Source/Runtime/UMG/Public/Components/Widget.h b/Engine/Source/Runtime/UMG/Public/Components/Widget.h
--- a/Engine/Source/Runtime/UMG/Public/Components/Widget.h
+++ b/Engine/Source/Runtime/UMG/Public/Components/Widget.h
@@ -174,16 +174,15 @@ namespace UMWidget
  * Flags used by the widget designer.
  */
 UENUM()
-namespace EWidgetDesignFlags
+enum class EWidgetDesignFlags : uint8
 {
-	enum Type
-	{
-		None				= 0,
-		Designing			= 1,
-		ShowOutline			= 2,
-		ExecutePreConstruct	= 4
-	};
-}
+	None				= 0,
+	Designing			= 1 << 0,
+	ShowOutline			= 1 << 1,
+	ExecutePreConstruct	= 1 << 2
+};
+
+ENUM_CLASS_FLAGS(EWidgetDesignFlags);
 
 
 #if WITH_EDITOR
@@ -797,18 +796,18 @@ public:
 	}
 
 	/** Sets the designer flags on the widget. */
-	virtual void SetDesignerFlags(EWidgetDesignFlags::Type NewFlags);
+	virtual void SetDesignerFlags(EWidgetDesignFlags NewFlags);
 
 	/** Gets the designer flags currently set on the widget. */
-	FORCEINLINE EWidgetDesignFlags::Type GetDesignerFlags() const
+	FORCEINLINE EWidgetDesignFlags GetDesignerFlags() const
 	{
-		return DesignerFlags;
+		return static_cast<EWidgetDesignFlags>(DesignerFlags);
 	}
 
 	/** Tests if any of the flags exist on this widget. */
-	FORCEINLINE bool HasAnyDesignerFlags(EWidgetDesignFlags::Type FlagToCheck) const
+	FORCEINLINE bool HasAnyDesignerFlags(EWidgetDesignFlags FlagsToCheck) const
 	{
-		return ( DesignerFlags&FlagToCheck ) != 0;
+		return EnumHasAnyFlags(GetDesignerFlags(), FlagsToCheck);
 	}
 
 	/** Returns the friendly name of the widget to display in the editor */
@@ -1018,7 +1017,7 @@ private:
 #if WITH_EDITORONLY_DATA
 	/** Any flags used by the designer at edit time. */
 	UPROPERTY(Transient)
-	TEnumAsByte<EWidgetDesignFlags::Type> DesignerFlags;
+	uint8 DesignerFlags;
 
 	/** The friendly name for this widget displayed in the designer and BP graph. */
 	UPROPERTY()
