PR #4944: Added template cast versions of multiple methods. (Contributed by [REDACTED_PERSON])

diff --git a/Engine/Source/Runtime/Engine/Classes/Engine/World.h b/Engine/Source/Runtime/Engine/Classes/Engine/World.h
--- a/Engine/Source/Runtime/Engine/Classes/Engine/World.h
+++ b/Engine/Source/Runtime/Engine/Classes/Engine/World.h
@@ -1985,8 +1985,32 @@ public:
 	/** @return Returns an iterator for the player controller list. */
 	FConstPlayerControllerIterator GetPlayerControllerIterator() const;
 	
+	/** 
+	 * @return Returns the first player controller cast to the template type, or NULL if there is not one.
+	 *
+	 * May return NULL if the cast fails.
+	 */
+	template< class T >
+	T* GetFirstPlayerController() const
+	{
+		return Cast<T>(GetFirstPlayerController());
+	}
+	
 	/** @return Returns the first player controller, or NULL if there is not one. */	
 	APlayerController* GetFirstPlayerController() const;
+	
+	/*
+	 *	Get the first valid local player via the first player controller.
+	 *
+	 *  @return Pointer to the first valid ULocalPlayer cast to the template type, or NULL if there is not one.
+	 *
+	 *  May Return NULL if the cast fails.
+	 */	
+	template< class T >
+	T* GetFirstLocalPlayerFromController() const
+	{
+		return Cast<T>(GetFirstLocalPlayerFromController());
+	}
 
 	/*
 	 *	Get the first valid local player via the first player controller.
diff --git a/Engine/Source/Runtime/Engine/Classes/GameFramework/Actor.h b/Engine/Source/Runtime/Engine/Classes/GameFramework/Actor.h
--- a/Engine/Source/Runtime/Engine/Classes/GameFramework/Actor.h
+++ b/Engine/Source/Runtime/Engine/Classes/GameFramework/Actor.h
@@ -2687,6 +2687,13 @@ public:
 
 	/** Gets the GameInstance that ultimately contains this actor. */
 	class UGameInstance* GetGameInstance() const;
+	
+	/** 
+	 * Gets the GameInstance that ultimately contains this actor cast to the template type.
+	 * May return NULL if the cast fails. 
+	 */
+	template< class T >
+	T* GetGameInstance() const;
 
 	/** Returns true if this is a replicated actor that was placed in the map */
 	bool IsNetStartupActor() const;
diff --git a/Engine/Source/Runtime/Engine/Classes/GameFramework/Controller.h b/Engine/Source/Runtime/Engine/Classes/GameFramework/Controller.h
--- a/Engine/Source/Runtime/Engine/Classes/GameFramework/Controller.h
+++ b/Engine/Source/Runtime/Engine/Classes/GameFramework/Controller.h
@@ -171,6 +171,16 @@ public:
 
 	UFUNCTION()
 	virtual void OnRep_PlayerState();
+	
+	/**
+	 * @return this controller's PlayerState cast to the template type, or NULL if there is not one.
+	 * May return null if the cast fails.
+	 */
+	template < class T >
+	T* GetPlayerState() const
+	{
+		return Cast<T>(PlayerState);
+	}
 
 	/** DEPRECATED! Use the standard "Cast To" node instead. Casts this Controller to a Player Controller, if possible. */
 	DEPRECATED(4.11, "CastToPlayerController has been replaced by the standard Cast To node.")
diff --git a/Engine/Source/Runtime/Engine/Classes/GameFramework/Pawn.h b/Engine/Source/Runtime/Engine/Classes/GameFramework/Pawn.h
--- a/Engine/Source/Runtime/Engine/Classes/GameFramework/Pawn.h
+++ b/Engine/Source/Runtime/Engine/Classes/GameFramework/Pawn.h
@@ -199,6 +199,13 @@ public:
 	/** Returns controller for this actor. */
 	UFUNCTION(BlueprintCallable, Category=Pawn)
 	AController* GetController() const;
+	
+	/** Returns controller for this actor cast to the template type. May return NULL is the cast fails. */
+	template < class T >
+	T* GetController() const
+	{
+		Cast<T>(GetController());
+	}
 
 	/** Get the rotation of the Controller, often the 'view' rotation of this Pawn. */
 	UFUNCTION(BlueprintCallable, Category=Pawn)
diff --git a/Engine/Source/Runtime/Engine/Private/Actor.cpp b/Engine/Source/Runtime/Engine/Private/Actor.cpp
--- a/Engine/Source/Runtime/Engine/Private/Actor.cpp
+++ b/Engine/Source/Runtime/Engine/Private/Actor.cpp
@@ -352,6 +352,12 @@ UGameInstance* AActor::GetGameInstance() const
 	return GetWorld()->GetGameInstance();
 }
 
+template< class T >
+T* AActor::GetGameInstance() const
+{
+	return GetWorld()->GetGameInstance<T>();
+}
+
 bool AActor::IsNetStartupActor() const
 {
 	return bNetStartup;
diff --git a/Engine/Source/Runtime/UMG/Public/Blueprint/UserWidget.h b/Engine/Source/Runtime/UMG/Public/Blueprint/UserWidget.h
--- a/Engine/Source/Runtime/UMG/Public/Blueprint/UserWidget.h
+++ b/Engine/Source/Runtime/UMG/Public/Blueprint/UserWidget.h
@@ -293,6 +293,16 @@ public:
 	 * @return The owning local player.
 	 */
 	virtual ULocalPlayer* GetOwningLocalPlayer() const override;
+	
+	/**
+	 * Gets the local player associated with this UI cast to the template type.
+	 * @return The owning local player. May be NULL if the cast fails.
+	 */
+	template < class T >
+	T* GetOwningLocalPlayer() const
+	{
+		return Cast<T>(GetOwningLocalPlayer());
+	}
 
 	/**
 	 * Sets the player associated with this UI via LocalPlayer reference.
@@ -305,7 +315,17 @@ public:
 	 * @return The player controller that owns the UI.
 	 */
 	virtual APlayerController* GetOwningPlayer() const override;
-
+	
+	/**
+	 * Gets the player controller associated with this UI cast to the template type.
+	 * @return The player controller that owns the UI. May be NULL if the cast fails.
+	 */
+	template < class T >
+	T* GetOwningPlayer() const
+	{
+		return Cast<T>(GetOwningPlayer());
+	}
+	
 	/**
 	 * Sets the local player associated with this UI via PlayerController reference.
 	 * @param LocalPlayerController The PlayerController of the local player you want to be the conceptual owner of this UI.
@@ -319,6 +339,17 @@ public:
 	 */
 	UFUNCTION(BlueprintCallable, BlueprintCosmetic, Category="Player")
 	class APawn* GetOwningPlayerPawn() const;
+	
+	/**
+	 * Gets the player pawn associated with this UI cast to the template type.
+	 * @return Gets the owning player pawn that's owned by the player controller assigned to this widget.
+	 * May be NULL if the cast fails.
+	 */
+	template < class T >
+	T* GetOwningPlayerPawn() const
+	{
+		return Cast<T>(GetOwningPlayerPawn());
+	}
 
 	/**
 	 * Get the owning player's PlayerState.
