Fix for ListView not propagating selection change events

diff --git a/Engine/Source/Runtime/UMG/Public/Slate/SObjectTableRow.h b/Engine/Source/Runtime/UMG/Public/Slate/SObjectTableRow.h
--- a/Engine/Source/Runtime/UMG/Public/Slate/SObjectTableRow.h
+++ b/Engine/Source/Runtime/UMG/Public/Slate/SObjectTableRow.h
@@ -79,6 +79,8 @@ public:
 			[
 				ContentWidget.ToSharedRef()
 			], &InWidgetObject);
+
+		SetCanTick(true);
 	}
 
 	virtual ~SObjectTableRow()
