Fix crash UWidget::HasUserFocus when a PlayerController has a null Player

diff --git a/Engine/Source/Runtime/UMG/Private/Components/Widget.cpp b/Engine/Source/Runtime/UMG/Private/Components/Widget.cpp
--- a/Engine/Source/Runtime/UMG/Private/Components/Widget.cpp
+++ b/Engine/Source/Runtime/UMG/Private/Components/Widget.cpp
@@ -461,7 +461,7 @@ void UWidget::SetKeyboardFocus()
 
 bool UWidget::HasUserFocus(APlayerController* PlayerController) const
 {
-	if ( PlayerController == nullptr || !PlayerController->IsLocalPlayerController() )
+	if (PlayerController == nullptr || PlayerController->Player == nullptr || !PlayerController->IsLocalPlayerController())
 	{
 		return false;
 	}
