Fix for UMG ListView not propagating OnItemSelectionChanged events to its items

diff --git a/Engine/Source/Runtime/UMG/Public/Components/ListViewBase.h b/Engine/Source/Runtime/UMG/Public/Components/ListViewBase.h
--- a/Engine/Source/Runtime/UMG/Public/Components/ListViewBase.h
+++ b/Engine/Source/Runtime/UMG/Public/Components/ListViewBase.h
@@ -321,7 +321,9 @@ private:
 
 		UUserWidget& EntryWidget = OnGenerateEntryWidgetInternal(Item, DesiredEntryClass, OwnerTable);
 		EntryWidget.SetPadding(GetDesiredEntryPadding(Item));
-		return StaticCastSharedPtr<SObjectTableRow<ItemType>>(EntryWidget.GetCachedWidget()).ToSharedRef();
+		TSharedPtr<SWidget> CachedWidget = EntryWidget.GetCachedWidget();
+		CachedWidget->SetCanTick(true); // this is a hack to force ticking to true so selection works (which should NOT require ticking! but currently does)
+		return StaticCastSharedPtr<SObjectTableRow<ItemType>>(CachedWidget).ToSharedRef();
 	}
 
 	void HandleItemClicked(ItemType Item)
diff --git a/Engine/Source/Runtime/UMG/Public/Slate/SObjectTableRow.h b/Engine/Source/Runtime/UMG/Public/Slate/SObjectTableRow.h
--- a/Engine/Source/Runtime/UMG/Public/Slate/SObjectTableRow.h
+++ b/Engine/Source/Runtime/UMG/Public/Slate/SObjectTableRow.h
@@ -79,8 +79,6 @@ public:
 			[
 				ContentWidget.ToSharedRef()
 			], &InWidgetObject);
-
-		SetCanTick(true);
 	}
 
 	virtual ~SObjectTableRow()
