Fix for Setting focus not working in Fast Path for WidgetComponents in world space.
Fix for rare assert caused by widget components creating MID's in async tasks (along side screne proxy creation) moved it out.

diff --git a/Engine/Source/Runtime/UMG/Private/Components/WidgetComponent.cpp b/Engine/Source/Runtime/UMG/Private/Components/WidgetComponent.cpp
--- a/Engine/Source/Runtime/UMG/Private/Components/WidgetComponent.cpp
+++ b/Engine/Source/Runtime/UMG/Private/Components/WidgetComponent.cpp
@@ -614,14 +614,32 @@ void UWidgetComponent::EndPlay(const EEndPlayReason::Type EndPlayReason)
 	Super::EndPlay(EndPlayReason);
 }
 
-FPrimitiveSceneProxy* UWidgetComponent::CreateSceneProxy()
+
+void UWidgetComponent::SetMaterial(int32 ElementIndex, UMaterialInterface* Material)
+{
+	Super::SetMaterial(ElementIndex, Material);
+
+	UpdateMaterialInstance();
+}
+
+void UWidgetComponent::UpdateMaterialInstance()
 {
 	// Always clear the material instance in case we're going from 3D to 2D.
-	if ( MaterialInstance )
+	MaterialInstance = nullptr;
+
+	if (Space == EWidgetSpace::Screen)
 	{
-		MaterialInstance = nullptr;
+		return;
 	}
 
+	UMaterialInterface* BaseMaterial = GetMaterial(0);
+	MaterialInstance = UMaterialInstanceDynamic::Create(BaseMaterial, this);
+
+	UpdateMaterialInstanceParameters();
+}
+
+FPrimitiveSceneProxy* UWidgetComponent::CreateSceneProxy()
+{
 	if (Space == EWidgetSpace::Screen)
 	{
 		return nullptr;
@@ -629,15 +647,6 @@ FPrimitiveSceneProxy* UWidgetComponent::CreateSceneProxy()
 
 	if (WidgetRenderer && CurrentSlateWidget.IsValid())
 	{
-		// Create a new MID for the current base material
-		{
-			UMaterialInterface* BaseMaterial = GetMaterial(0);
-
-			MaterialInstance = UMaterialInstanceDynamic::Create(BaseMaterial, this);
-
-			UpdateMaterialInstanceParameters();
-		}
-
 		RequestRedraw();
 		LastWidgetRenderTime = 0;
 
@@ -890,6 +899,23 @@ void UWidgetComponent::RegisterWindow()
 		{
 			FSlateApplication::Get().RegisterVirtualWindow(SlateWindow.ToSharedRef());
 		}
+
+#if SLATE_PARENT_POINTERS
+		if (Widget && !Widget->IsDesignTime())
+		{
+			if (UWorld* LocalWorld = GetWorld())
+			{
+				UGameInstance* GameInstance = LocalWorld->GetGameInstance();
+				check(GameInstance);
+
+				UGameViewportClient* GameViewportClient = GameInstance->GetGameViewportClient();
+				if (GameViewportClient)
+				{
+					SlateWindow->AssignParentWidget(GameViewportClient->GetGameViewportWidget());
+				}
+			}
+		}
+#endif
 	}
 }
 
@@ -1377,6 +1403,8 @@ void UWidgetComponent::UpdateWidget()
 			bool bNeededNewWindow = false;
 			if ( !SlateWindow.IsValid() )
 			{
+				UpdateMaterialInstance();
+
 				SlateWindow = SNew(SVirtualWindow).Size(DrawSize);
 				SlateWindow->SetIsFocusable(bWindowFocusable);
 				RegisterWindow();
diff --git a/Engine/Source/Runtime/UMG/Public/Components/WidgetComponent.h b/Engine/Source/Runtime/UMG/Public/Components/WidgetComponent.h
--- a/Engine/Source/Runtime/UMG/Public/Components/WidgetComponent.h
+++ b/Engine/Source/Runtime/UMG/Public/Components/WidgetComponent.h
@@ -85,6 +85,7 @@ public:
 	virtual void OnUnregister() override;
 	virtual void DestroyComponent(bool bPromoteChildren = false) override;
 	UMaterialInterface* GetMaterial(int32 MaterialIndex) const override;
+	virtual void SetMaterial(int32 ElementIndex, UMaterialInterface* Material) override;
 	int32 GetNumMaterials() const override;
 	virtual void TickComponent(float DeltaTime, enum ELevelTick TickType, FActorComponentTickFunction *ThisTickFunction) override;
 
@@ -287,6 +288,8 @@ protected:
 	/** @return the width of the widget component taking GeometryMode into account. */
 	float ComputeComponentWidth() const;
 
+	void UpdateMaterialInstance();
+
 protected:
 	/** The coordinate space in which to render the widget */
 	UPROPERTY(EditAnywhere, Category=UserInterface)
