PR #5065: make RichTextBlockImageDecorator be blueprintable (Contributed by [REDACTED_PERSON])

diff --git a/Engine/Source/Runtime/UMG/Public/Components/RichTextBlockImageDecorator.h b/Engine/Source/Runtime/UMG/Public/Components/RichTextBlockImageDecorator.h
--- a/Engine/Source/Runtime/UMG/Public/Components/RichTextBlockImageDecorator.h
+++ b/Engine/Source/Runtime/UMG/Public/Components/RichTextBlockImageDecorator.h
@@ -33,7 +33,7 @@ public:
  * to map certain keys to certain images.  We recommend you subclass this
  * as a blueprint to configure the instance.
  */
-UCLASS()
+UCLASS(Blueprintable)
 class UMG_API URichTextBlockImageDecorator : public URichTextBlockDecorator
 {
 	GENERATED_BODY()
