Fix for ListView not propagating selection change events

diff --git a/Engine/Source/Runtime/UMG/Public/Slate/SObjectTableRow.h b/Engine/Source/Runtime/UMG/Public/Slate/SObjectTableRow.h
--- a/Engine/Source/Runtime/UMG/Public/Slate/SObjectTableRow.h
+++ b/Engine/Source/Runtime/UMG/Public/Slate/SObjectTableRow.h
@@ -74,6 +74,7 @@ public:
 		{
 			ChildSlot [ SNew(STextBlock).Text(NSLOCTEXT("SObjectTableRow", "InvalidWidgetClass", "Any UserWidget generated as a table row must implement the IUserListEntry interface")) ];
 		}
+		SetCanTick(true);
 	}
 
 	virtual ~SObjectTableRow()
