PR #5094: UE-49592 Fix WidgetInteractionComponent broke when using Cylinder geometry on a 3D Widget (Contributed by [REDACTED_PERSON])

diff --git a/Engine/Source/Runtime/UMG/Private/Components/WidgetInteractionComponent.cpp b/Engine/Source/Runtime/UMG/Private/Components/WidgetInteractionComponent.cpp
--- a/Engine/Source/Runtime/UMG/Private/Components/WidgetInteractionComponent.cpp
+++ b/Engine/Source/Runtime/UMG/Private/Components/WidgetInteractionComponent.cpp
@@ -113,13 +113,15 @@ UWidgetInteractionComponent::FWidgetTraceResult UWidgetInteractionComponent::Per
 
 	TArray<FHitResult> MultiHits;
 
+	FVector WorldDirection;
+
 	switch( InteractionSource )
 	{
 		case EWidgetInteractionSource::World:
 		{
 			const FVector WorldLocation = GetComponentLocation();
 			const FTransform WorldTransform = GetComponentTransform();
-			const FVector Direction = WorldTransform.GetUnitAxis(EAxis::X);
+			WorldDirection = WorldTransform.GetUnitAxis(EAxis::X);
 
 			TArray<UPrimitiveComponent*> PrimitiveChildren;
 			GetRelatedComponentsToIgnoreInAutomaticHitTesting(PrimitiveChildren);
@@ -128,7 +130,7 @@ UWidgetInteractionComponent::FWidgetTraceResult UWidgetInteractionComponent::Per
 			Params.AddIgnoredComponents(PrimitiveChildren);
 
 			TraceResult.LineStartLocation = WorldLocation;
-			TraceResult.LineEndLocation = WorldLocation + (Direction * InteractionDistance);
+			TraceResult.LineEndLocation = WorldLocation + (WorldDirection * InteractionDistance);
 
 			GetWorld()->LineTraceMultiByChannel(MultiHits, TraceResult.LineStartLocation, TraceResult.LineEndLocation, TraceChannel, Params);
 			break;
@@ -153,7 +155,6 @@ UWidgetInteractionComponent::FWidgetTraceResult UWidgetInteractionComponent::Per
 					if ( LocalPlayer->ViewportClient->GetMousePosition(MousePosition) )
 					{
 						FVector WorldOrigin;
-						FVector WorldDirection;
 						if ( UGameplayStatics::DeprojectScreenToWorld(PlayerController, MousePosition, WorldOrigin, WorldDirection) == true )
 						{
 							TraceResult.LineStartLocation = WorldOrigin;
@@ -169,7 +170,6 @@ UWidgetInteractionComponent::FWidgetTraceResult UWidgetInteractionComponent::Per
 					LocalPlayer->ViewportClient->GetViewportSize(ViewportSize);
 
 					FVector WorldOrigin;
-					FVector WorldDirection;
 					if ( UGameplayStatics::DeprojectScreenToWorld(PlayerController, ViewportSize * 0.5f, WorldOrigin, WorldDirection) == true )
 					{
 						TraceResult.LineStartLocation = WorldOrigin;
@@ -183,6 +183,8 @@ UWidgetInteractionComponent::FWidgetTraceResult UWidgetInteractionComponent::Per
 		}
 		case EWidgetInteractionSource::Custom:
 		{
+			const FTransform WorldTransform = GetComponentTransform();
+			WorldDirection = WorldTransform.GetUnitAxis(EAxis::X);
 			TraceResult.HitResult = CustomHitResult;
 			TraceResult.bWasHit = CustomHitResult.bBlockingHit;
 			TraceResult.LineStartLocation = CustomHitResult.TraceStart;
@@ -221,11 +223,8 @@ UWidgetInteractionComponent::FWidgetTraceResult UWidgetInteractionComponent::Per
 			// @todo WASTED WORK: GetLocalHitLocation() gets called in GetHitWidgetPath();
 
 			if (TraceResult.HitWidgetComponent->GetGeometryMode() == EWidgetGeometryMode::Cylinder)
-			{
-				const FTransform WorldTransform = GetComponentTransform();
-				const FVector Direction = WorldTransform.GetUnitAxis(EAxis::X);
-				
-				TTuple<FVector, FVector2D> CylinderHitLocation = TraceResult.HitWidgetComponent->GetCylinderHitLocation(TraceResult.HitResult.ImpactPoint, Direction);
+			{				
+				TTuple<FVector, FVector2D> CylinderHitLocation = TraceResult.HitWidgetComponent->GetCylinderHitLocation(TraceResult.HitResult.ImpactPoint, WorldDirection);
 				TraceResult.HitResult.ImpactPoint = CylinderHitLocation.Get<0>();
 				TraceResult.LocalHitLocation = CylinderHitLocation.Get<1>();
 			}
