Adding UMG_API to FRichTextDecorator

diff --git a/Engine/Source/Runtime/UMG/Public/Components/RichTextBlockDecorator.h b/Engine/Source/Runtime/UMG/Public/Components/RichTextBlockDecorator.h
--- a/Engine/Source/Runtime/UMG/Public/Components/RichTextBlockDecorator.h
+++ b/Engine/Source/Runtime/UMG/Public/Components/RichTextBlockDecorator.h
@@ -16,7 +16,7 @@ class ISlateStyle;
 class URichTextBlockDecorator;
 class URichTextBlock;
 
-class FRichTextDecorator : public ITextDecorator
+class UMG_API FRichTextDecorator : public ITextDecorator
 {
 public:
 	FRichTextDecorator(URichTextBlock* InOwner);
