PR #5110: Fix Slate Widget's support in UWidgetComponent (Contributed by [REDACTED_PERSON])

diff --git a/Engine/Source/Runtime/UMG/Private/Components/WidgetComponent.cpp b/Engine/Source/Runtime/UMG/Private/Components/WidgetComponent.cpp
--- a/Engine/Source/Runtime/UMG/Private/Components/WidgetComponent.cpp
+++ b/Engine/Source/Runtime/UMG/Private/Components/WidgetComponent.cpp
@@ -48,6 +48,10 @@ public:
 				{
 					ScreenLayer.Pin()->AddComponent(Component, UserWidget->TakeWidget());
 				}
+				else if (Component->GetSlateWidget().IsValid())
+				{
+					ScreenLayer.Pin()->AddComponent(Component, Component->GetSlateWidget().ToSharedRef());
+				}
 			}
 		}
 	}
@@ -84,6 +88,10 @@ public:
 				{
 					NewScreenLayer->AddComponent(Component, UserWidget->TakeWidget());
 				}
+				else if (Component->GetSlateWidget().IsValid())
+				{
+					NewScreenLayer->AddComponent(Component, Component->GetSlateWidget().ToSharedRef()); 
+				}
 			}
 		}
 
@@ -882,6 +890,12 @@ void UWidgetComponent::ReleaseResources()
 		Widget = nullptr;
 	}
 
+	if (SlateWidget.IsValid())
+	{
+		RemoveWidgetFromScreen();
+		SlateWidget.Reset();
+	}
+	
 	if (WidgetRenderer)
 	{
 		BeginCleanup(WidgetRenderer);
@@ -995,7 +1009,11 @@ void UWidgetComponent::TickComponent(float DeltaTime, enum ELevelTick TickType,
 								
 									bAddedToScreen = true;
 								
-									Widget->SetPlayerContext(PlayerContext);
+									if (Widget && Widget->IsValidLowLevel())
+									{
+										Widget->SetPlayerContext(PlayerContext);
+									}
+									
 									ScreenLayer->AddComponent(this);
 								}
 							}
