PR #5141: Fixed incorrect deprecated version number in UserWidget.h (Contributed by [REDACTED_PERSON])

diff --git a/Engine/Source/Runtime/UMG/Public/Blueprint/UserWidget.h b/Engine/Source/Runtime/UMG/Public/Blueprint/UserWidget.h
--- a/Engine/Source/Runtime/UMG/Public/Blueprint/UserWidget.h
+++ b/Engine/Source/Runtime/UMG/Public/Blueprint/UserWidget.h
@@ -1054,7 +1054,7 @@ protected:
 
 	virtual void NativeTick(const FGeometry& MyGeometry, float InDeltaTime);
 
-	DEPRECATED(2.20, "Please override the other version of NativePaint that accepts all the parameters, not just the paint context.")
+	DEPRECATED(4.20, "Please override the other version of NativePaint that accepts all the parameters, not just the paint context.")
 	virtual void NativePaint(FPaintContext& InContext) const { }
 
 	/**
