PR #5145: UE-48051: Don't override DrawSize variable with CurrentDrawSize (Contributed by [REDACTED_PERSON])

diff --git a/Engine/Source/Runtime/UMG/Private/Components/WidgetComponent.cpp b/Engine/Source/Runtime/UMG/Private/Components/WidgetComponent.cpp
--- a/Engine/Source/Runtime/UMG/Private/Components/WidgetComponent.cpp
+++ b/Engine/Source/Runtime/UMG/Private/Components/WidgetComponent.cpp
@@ -140,7 +140,7 @@ public:
 						{
 							if ( WidgetComponent->GetReceiveHardwareInput() )
 							{
-								if ( WidgetComponent->GetDrawSize().X != 0 && WidgetComponent->GetDrawSize().Y != 0 )
+								if ( WidgetComponent->GetCurrentDrawSize().X != 0 && WidgetComponent->GetCurrentDrawSize().Y != 0 )
 								{
 									// Get the "forward" vector based on the current rotation system.
 									const FVector ForwardVector = WidgetComponent->GetForwardVector();
@@ -171,7 +171,7 @@ public:
 			{
 				FGeometry WidgetGeom;
 
-				ArrangedChildren.AddWidget( FArrangedWidget( WidgetComponent->GetSlateWindow().ToSharedRef(), WidgetGeom.MakeChild( WidgetComponent->GetDrawSize(), FSlateLayoutTransform() ) ) );
+				ArrangedChildren.AddWidget( FArrangedWidget( WidgetComponent->GetSlateWindow().ToSharedRef(), WidgetGeom.MakeChild( WidgetComponent->GetCurrentDrawSize(), FSlateLayoutTransform() ) ) );
 			}
 		}
 	}
@@ -196,7 +196,7 @@ public:
 						{
 							if ( WidgetComponent->GetReceiveHardwareInput() )
 							{
-								if ( WidgetComponent->GetDrawSize().X != 0 && WidgetComponent->GetDrawSize().Y != 0 )
+								if ( WidgetComponent->GetCurrentDrawSize().X != 0 && WidgetComponent->GetCurrentDrawSize().Y != 0 )
 								{
 									if ( WidgetComponent == HitComponent )
 									{
@@ -681,7 +681,7 @@ FPrimitiveSceneProxy* UWidgetComponent::CreateSceneProxy()
 
 		FWidgetBoxProxy(const UWidgetComponent* InComponent)
 			: FPrimitiveSceneProxy(InComponent)
-			, BoxExtents(1.f, InComponent->GetDrawSize().X / 2.0f, InComponent->GetDrawSize().Y / 2.0f)
+			, BoxExtents(1.f, InComponent->GetCurrentDrawSize().X / 2.0f, InComponent->GetCurrentDrawSize().Y / 2.0f)
 		{
 			bWillEverBeLit = false;
 		}
@@ -738,12 +738,12 @@ FBoxSphereBounds UWidgetComponent::CalcBounds(const FTransform & LocalToWorld) c
 	if ( Space != EWidgetSpace::Screen )
 	{
 		const FVector Origin = FVector(.5f,
-			-( DrawSize.X * 0.5f ) + ( DrawSize.X * Pivot.X ),
-			-( DrawSize.Y * 0.5f ) + ( DrawSize.Y * Pivot.Y ));
+			-( CurrentDrawSize.X * 0.5f ) + (CurrentDrawSize.X * Pivot.X ),
+			-( CurrentDrawSize.Y * 0.5f ) + ( CurrentDrawSize.Y * Pivot.Y ));
 
-		const FVector BoxExtent = FVector(1.f, DrawSize.X / 2.0f, DrawSize.Y / 2.0f);
+		const FVector BoxExtent = FVector(1.f, CurrentDrawSize.X / 2.0f, CurrentDrawSize.Y / 2.0f);
 
-		FBoxSphereBounds NewBounds(Origin, BoxExtent, DrawSize.Size() / 2.0f);
+		FBoxSphereBounds NewBounds(Origin, BoxExtent, CurrentDrawSize.Size() / 2.0f);
 		NewBounds = NewBounds.TransformBy(LocalToWorld);
 
 		NewBounds.BoxExtent *= BoundsScale;
@@ -767,7 +767,7 @@ FCollisionShape UWidgetComponent::GetCollisionShape(float Inflation) const
 {
 	if ( Space != EWidgetSpace::Screen )
 	{
-		FVector BoxHalfExtent = ( FVector(0.01f, DrawSize.X * 0.5f, DrawSize.Y * 0.5f) * GetComponentTransform().GetScale3D() ) + Inflation;
+		FVector BoxHalfExtent = ( FVector(0.01f, CurrentDrawSize.X * 0.5f, CurrentDrawSize.Y * 0.5f) * GetComponentTransform().GetScale3D() ) + Inflation;
 
 		if ( Inflation < 0.0f )
 		{
@@ -1079,6 +1079,7 @@ void UWidgetComponent::DrawWidgetToRenderTarget(float DeltaTime)
 		return;
 	}
 
+	const FIntPoint PreviousDrawSize = CurrentDrawSize;
 	CurrentDrawSize = DrawSize;
 
 	const float DrawScale = 1.0f;
@@ -1099,9 +1100,8 @@ void UWidgetComponent::DrawWidgetToRenderTarget(float DeltaTime)
 		WidgetRenderer->SetIsPrepassNeeded(true);
 	}
 
-	if ( CurrentDrawSize != DrawSize )
+	if ( CurrentDrawSize != PreviousDrawSize )
 	{
-		DrawSize = CurrentDrawSize;
 		UpdateBodySetup(true);
 		RecreatePhysicsState();
 	}
@@ -1436,14 +1436,14 @@ void UWidgetComponent::UpdateWidget()
 			{
 				UpdateMaterialInstance();
 
-				SlateWindow = SNew(SVirtualWindow).Size(DrawSize);
+				SlateWindow = SNew(SVirtualWindow).Size(CurrentDrawSize);
 				SlateWindow->SetIsFocusable(bWindowFocusable);
 				RegisterWindow();
 
 				bNeededNewWindow = true;
 			}
 
-			SlateWindow->Resize(DrawSize);
+			SlateWindow->Resize(CurrentDrawSize);
 
 			bool bWidgetChanged = false;
 			if ( NewSlateWidget.IsValid() )
@@ -1572,17 +1572,17 @@ void UWidgetComponent::UpdateBodySetup( bool bDrawSizeChanged )
 		FKBoxElem* BoxElem = BodySetup->AggGeom.BoxElems.GetData();
 
 		FVector Origin = FVector(.5f,
-			-( DrawSize.X * 0.5f ) + ( DrawSize.X * Pivot.X ),
-			-( DrawSize.Y * 0.5f ) + ( DrawSize.Y * Pivot.Y ));
+			-( CurrentDrawSize.X * 0.5f ) + ( CurrentDrawSize.X * Pivot.X ),
+			-( CurrentDrawSize.Y * 0.5f ) + ( CurrentDrawSize.Y * Pivot.Y ));
 			const float Width = ComputeComponentWidth();
-			const float Height = DrawSize.Y;
+			const float Height = CurrentDrawSize.Y;
 			Origin = FVector(.5f,
 				-( Width * 0.5f ) + ( Width * Pivot.X ),
 				-( Height * 0.5f ) + ( Height * Pivot.Y ));
 			
 		BoxElem->X = 0.01f;
-		BoxElem->Y = DrawSize.X;
-		BoxElem->Z = DrawSize.Y;
+		BoxElem->Y = CurrentDrawSize.X;
+		BoxElem->Z = CurrentDrawSize.Y;
 
 		BoxElem->SetTransform(FTransform::Identity);
 		BoxElem->Center = Origin;
@@ -1651,7 +1651,7 @@ TTuple<FVector, FVector2D> UWidgetComponent::GetCylinderHitLocation(FVector Worl
 
 
 	const float ArcAngleRadians = FMath::DegreesToRadians(GetCylinderArcAngle());
-	const float Radius = GetDrawSize().X / ArcAngleRadians;
+	const float Radius = CurrentDrawSize.X / ArcAngleRadians;
 	const float Apothem = Radius * FMath::Cos(0.5f*ArcAngleRadians);
 	const float ChordLength = 2.0f * Radius * FMath::Sin(0.5f*ArcAngleRadians);
 
@@ -1791,6 +1791,11 @@ FVector2D UWidgetComponent::GetDrawSize() const
 	return DrawSize;
 }
 
+FVector2D UWidgetComponent::GetCurrentDrawSize() const
+{
+	return CurrentDrawSize;
+}
+
 void UWidgetComponent::SetDrawSize(FVector2D Size)
 {
 	FIntPoint NewDrawSize((int32)Size.X, (int32)Size.Y);
@@ -1799,8 +1804,6 @@ void UWidgetComponent::SetDrawSize(FVector2D Size)
 	{
 		DrawSize = NewDrawSize;
 		MarkRenderStateDirty();
-		UpdateBodySetup( true );
-		RecreatePhysicsState();
 	}
 }
 
diff --git a/Engine/Source/Runtime/UMG/Public/Components/WidgetComponent.h b/Engine/Source/Runtime/UMG/Public/Components/WidgetComponent.h
--- a/Engine/Source/Runtime/UMG/Public/Components/WidgetComponent.h
+++ b/Engine/Source/Runtime/UMG/Public/Components/WidgetComponent.h
@@ -194,10 +194,14 @@ public:
 	UFUNCTION(BlueprintCallable, Category=UserInterface)
 	ULocalPlayer* GetOwnerPlayer() const;
 
-	/** @return The draw size of the quad in the world */
+	/** @return The "specified" draw size of the quad in the world */
 	UFUNCTION(BlueprintCallable, Category=UserInterface)
 	FVector2D GetDrawSize() const;
 
+	/** @return The "actual" draw size of the quad in the world */
+	UFUNCTION(BlueprintCallable, Category = UserInterface)
+	FVector2D GetCurrentDrawSize() const;
+
 	/** Sets the draw size of the quad in the world */
 	UFUNCTION(BlueprintCallable, Category=UserInterface)
 	void SetDrawSize(FVector2D Size);
