Mark ComboBoxString as UsingCustomVersion to properly support  backwards compatibility for the EnableGamepadNavigationMode flag

diff --git a/Engine/Source/Runtime/UMG/Private/Components/ComboBoxString.cpp b/Engine/Source/Runtime/UMG/Private/Components/ComboBoxString.cpp
--- a/Engine/Source/Runtime/UMG/Private/Components/ComboBoxString.cpp
+++ b/Engine/Source/Runtime/UMG/Private/Components/ComboBoxString.cpp
@@ -66,6 +66,13 @@ void UComboBoxString::ReleaseSlateResources(bool bReleaseChildren)
 	ComboBoxContent.Reset();
 }
 
+void UComboBoxString::Serialize(FArchive& Ar)
+{
+	Super::Serialize(Ar);
+
+	Ar.UsingCustomVersion(FEditorObjectVersion::GUID);
+}
+
 void UComboBoxString::PostLoad()
 {
 	Super::PostLoad();
diff --git a/Engine/Source/Runtime/UMG/Public/Components/ComboBoxString.h b/Engine/Source/Runtime/UMG/Public/Components/ComboBoxString.h
--- a/Engine/Source/Runtime/UMG/Public/Components/ComboBoxString.h
+++ b/Engine/Source/Runtime/UMG/Public/Components/ComboBoxString.h
@@ -137,6 +137,7 @@ public:
 
 	//~ Begin UObject Interface
 	virtual void PostInitProperties() override;
+	virtual void Serialize(FArchive& Ar) override;
 	virtual void PostLoad() override;
 	//~ End UObject Interface
 
