Expose color deficiency settings to blueprint

diff --git a/Engine/Source/Editor/UnrealEd/Private/EditorEngine.cpp b/Engine/Source/Editor/UnrealEd/Private/EditorEngine.cpp
--- a/Engine/Source/Editor/UnrealEd/Private/EditorEngine.cpp
+++ b/Engine/Source/Editor/UnrealEd/Private/EditorEngine.cpp
@@ -930,6 +930,18 @@ void UEditorEngine::Init(IEngineLoop* InEngineLoop)
 
 		// Always resume the dynamic resolution state to ensure it is same state as in game builds when starting PIE.
 		GEngine->ResumeDynamicResolution();
+
+		if (FSlateApplication::IsInitialized())
+		{
+			//Reset color deficiency settings in case they have been modified during PIE
+			const UEditorStyleSettings* EditorSettings = GetDefault<UEditorStyleSettings>();
+			const EColorVisionDeficiency DeficiencyType = EditorSettings->ColorVisionDeficiencyPreviewType;
+			const int32 Severity = EditorSettings->ColorVisionDeficiencySeverity;
+			const bool bCorrectDeficiency = EditorSettings->bColorVisionDeficiencyCorrection;
+			const bool bShowCorrectionWithDeficiency = EditorSettings->bColorVisionDeficiencyCorrectionPreviewWithDeficiency;
+			FSlateApplication::Get().GetRenderer()->SetColorVisionDeficiencyType(DeficiencyType, Severity, bCorrectDeficiency, bShowCorrectionWithDeficiency);
+		}
+
 	});
 
 	// Initialize vanilla status before other systems that consume its status are started inside InitEditor()
diff --git a/Engine/Source/Runtime/UMG/Private/WidgetBlueprintLibrary.cpp b/Engine/Source/Runtime/UMG/Private/WidgetBlueprintLibrary.cpp
--- a/Engine/Source/Runtime/UMG/Private/WidgetBlueprintLibrary.cpp
+++ b/Engine/Source/Runtime/UMG/Private/WidgetBlueprintLibrary.cpp
@@ -608,6 +608,13 @@ void UWidgetBlueprintLibrary::GetSafeZonePadding(UObject* WorldContextObject, FV
 	SpillOverPadding = SafePadding;
 }
 
+void UWidgetBlueprintLibrary::SetColorVisionDeficiencyType(EColorVisionDeficiency Type, float Severity, bool CorrectDeficiency, bool ShowCorrectionWithDeficiency)
+{
+	int32 AdjustedSeverity = FMath::Clamp(Severity, 0.f, 1.f) * 10;
+	FSlateApplicationBase::Get().GetRenderer()->SetColorVisionDeficiencyType(Type, AdjustedSeverity, CorrectDeficiency, ShowCorrectionWithDeficiency);
+}
+
+
 bool UWidgetBlueprintLibrary::SetHardwareCursor(UObject* WorldContextObject, EMouseCursor::Type CursorShape, FName CursorName, FVector2D HotSpot)
 {
 	UWorld* World = GEngine->GetWorldFromContextObject(WorldContextObject, EGetWorldErrorMode::LogAndReturnNull);
diff --git a/Engine/Source/Runtime/UMG/Public/Blueprint/WidgetBlueprintLibrary.h b/Engine/Source/Runtime/UMG/Public/Blueprint/WidgetBlueprintLibrary.h
--- a/Engine/Source/Runtime/UMG/Public/Blueprint/WidgetBlueprintLibrary.h
+++ b/Engine/Source/Runtime/UMG/Public/Blueprint/WidgetBlueprintLibrary.h
@@ -318,6 +318,16 @@ public:
 	UFUNCTION(BlueprintCallable, Category="Widget|Safe Zone", meta=( WorldContext="WorldContextObject" ))
 	static void GetSafeZonePadding(UObject* WorldContextObject, FVector4& SafePadding, FVector2D& SafePaddingScale, FVector4& SpillOverPadding);
 
+	/**
+	* Apply color deficiency correction settings to the game window 
+	* @param Type The type of color deficiency correction to apply.
+	* @param Severity Intensity of the color deficiency correction effect, from 0 to 1.
+	* @param CorrectDeficiency Shifts the color spectrum to the visible range based on the current deficiency type.
+	* @param ShowCorrectionWithDeficiency If you're correcting the color deficiency, you can use this to visualize what the correction looks like with the deficiency.
+	*/
+	UFUNCTION(BlueprintCallable, BlueprintCosmetic, Category = "Widget|Accessibility", meta = (AdvancedDisplay = "3"))
+	static void SetColorVisionDeficiencyType(EColorVisionDeficiency Type, float Severity, bool CorrectDeficiency, bool ShowCorrectionWithDeficiency);
+
 	/**
 	 * Loads or sets a hardware cursor from the content directory in the game.
 	 */
