PR #5201: WidgetBlueprintLibrary::DrawLines and WidgetBlueprintLibrary::DrawLine functionality fix. (Contributed by [REDACTED_PERSON])

diff --git a/Engine/Source/Runtime/UMG/Private/WidgetBlueprintLibrary.cpp b/Engine/Source/Runtime/UMG/Private/WidgetBlueprintLibrary.cpp
--- a/Engine/Source/Runtime/UMG/Private/WidgetBlueprintLibrary.cpp
+++ b/Engine/Source/Runtime/UMG/Private/WidgetBlueprintLibrary.cpp
@@ -158,7 +158,7 @@ void UWidgetBlueprintLibrary::DrawBox(FPaintContext& Context, FVector2D Position
 	}
 }
 
-void UWidgetBlueprintLibrary::DrawLine(FPaintContext& Context, FVector2D PositionA, FVector2D PositionB, FLinearColor Tint, bool bAntiAlias)
+void UWidgetBlueprintLibrary::DrawLine(FPaintContext& Context, FVector2D PositionA, FVector2D PositionB, FLinearColor Tint, bool bAntiAlias, float Thickness)
 {
 	Context.MaxLayer++;
 
@@ -173,10 +173,11 @@ void UWidgetBlueprintLibrary::DrawLine(FPaintContext& Context, FVector2D Positio
 		Points,
 		ESlateDrawEffect::None,
 		Tint,
-		bAntiAlias);
+		bAntiAlias,
+		Thickness);
 }
 
-void UWidgetBlueprintLibrary::DrawLines(FPaintContext& Context, const TArray<FVector2D>& Points, FLinearColor Tint, bool bAntiAlias)
+void UWidgetBlueprintLibrary::DrawLines(FPaintContext& Context, const TArray<FVector2D>& Points, FLinearColor Tint, bool bAntiAlias, float Thickness)
 {
 	Context.MaxLayer++;
 
@@ -187,7 +188,8 @@ void UWidgetBlueprintLibrary::DrawLines(FPaintContext& Context, const TArray<FVe
 		Points,
 		ESlateDrawEffect::None,
 		Tint,
-		bAntiAlias);
+		bAntiAlias,
+		Thickness);
 }
 
 void UWidgetBlueprintLibrary::DrawText(FPaintContext& Context, const FString& InString, FVector2D Position, FLinearColor Tint)
diff --git a/Engine/Source/Runtime/UMG/Public/Blueprint/WidgetBlueprintLibrary.h b/Engine/Source/Runtime/UMG/Public/Blueprint/WidgetBlueprintLibrary.h
--- a/Engine/Source/Runtime/UMG/Public/Blueprint/WidgetBlueprintLibrary.h
+++ b/Engine/Source/Runtime/UMG/Public/Blueprint/WidgetBlueprintLibrary.h
@@ -76,21 +76,23 @@ public:
 	 *
 	 * @param PositionA		Starting position of the line in local space.
 	 * @param PositionB		Ending position of the line in local space.
-	 * @param Thickness		How many pixels thick this line should be.
 	 * @param Tint			Color to render the line.
+	 * @param bAntialias	Whether the line should be antialiased.
+	 * @param Thickness		How many pixels thick this line should be.
 	 */
-	UFUNCTION(BlueprintCallable, meta=( AdvancedDisplay = "4" ), Category="Painting" )
-	static void DrawLine(UPARAM(ref) FPaintContext& Context, FVector2D PositionA, FVector2D PositionB, FLinearColor Tint = FLinearColor::White, bool bAntiAlias = true);
+	UFUNCTION(BlueprintCallable, meta = (AdvancedDisplay = "4"), Category = "Painting")
+	static void DrawLine(UPARAM(ref) FPaintContext& Context, FVector2D PositionA, FVector2D PositionB, FLinearColor Tint = FLinearColor::White, bool bAntiAlias = true, float Thickness = 1.0f);
 
 	/**
 	 * Draws several line segments.
 	 *
 	 * @param Points		Line pairs, each line needs to be 2 separate points in the array.
-	 * @param Thickness		How many pixels thick this line should be.
 	 * @param Tint			Color to render the line.
+	 * @param bAntialias	Whether the line should be antialiased.
+	 * @param Thickness		How many pixels thick this line should be.
 	 */
-	UFUNCTION(BlueprintCallable, meta=( AdvancedDisplay = "3" ), Category="Painting" )
-	static void DrawLines(UPARAM(ref) FPaintContext& Context, const TArray<FVector2D>& Points, FLinearColor Tint = FLinearColor::White, bool bAntiAlias = true);
+	UFUNCTION(BlueprintCallable, meta = (AdvancedDisplay = "3"), Category = "Painting")
+	static void DrawLines(UPARAM(ref) FPaintContext& Context, const TArray<FVector2D>& Points, FLinearColor Tint = FLinearColor::White, bool bAntiAlias = true, float Thickness = 1.0f);
 
 	/**
 	 * Draws text.
