PR #5249: UE-66138: Assign PressMethod in UButton::RebuildWidget (Contributed by [REDACTED_PERSON])

diff --git a/Engine/Source/Runtime/UMG/Private/Components/Button.cpp b/Engine/Source/Runtime/UMG/Private/Components/Button.cpp
--- a/Engine/Source/Runtime/UMG/Private/Components/Button.cpp
+++ b/Engine/Source/Runtime/UMG/Private/Components/Button.cpp
@@ -44,6 +44,7 @@ TSharedRef<SWidget> UButton::RebuildWidget()
 		.ButtonStyle(&WidgetStyle)
 		.ClickMethod(ClickMethod)
 		.TouchMethod(TouchMethod)
+		.PressMethod(PressMethod)
 		.IsFocusable(IsFocusable)
 		;
 
