PR #5279: Added setters for widget compoent layer properties (Contributed by [REDACTED_PERSON])

diff --git a/Engine/Source/Runtime/UMG/Public/Components/WidgetComponent.h b/Engine/Source/Runtime/UMG/Public/Components/WidgetComponent.h
--- a/Engine/Source/Runtime/UMG/Public/Components/WidgetComponent.h
+++ b/Engine/Source/Runtime/UMG/Public/Components/WidgetComponent.h
@@ -272,6 +272,12 @@ public:
 
 	/** Defines the curvature of the widget component when using EWidgetGeometryMode::Cylinder; ignored otherwise.  */
 	float GetCylinderArcAngle() const { return CylinderArcAngle; }
+	
+	/** Sets shared layer name used when this widget is initialized */
+	void SetInitialSharedLayerName(FName NewSharedLayerName) { SharedLayerName = NewSharedLayerName; }
+	
+	/** Sets layer z order used when this widget is initialized */
+	void SetInitialLayerZOrder(int32 NewLayerZOrder) { LayerZOrder = NewLayerZOrder; }
 
 protected:
 	/** Just because the user attempts to receive hardware input does not mean it's possible. */
