UCanvasPanelSlot::PostEditChangeChainProperty always calls the super now

diff --git a/Engine/Source/Runtime/UMG/Private/Components/CanvasPanelSlot.cpp b/Engine/Source/Runtime/UMG/Private/Components/CanvasPanelSlot.cpp
--- a/Engine/Source/Runtime/UMG/Private/Components/CanvasPanelSlot.cpp
+++ b/Engine/Source/Runtime/UMG/Private/Components/CanvasPanelSlot.cpp
@@ -319,24 +319,16 @@ void UCanvasPanelSlot::PostEditChangeChainProperty(struct FPropertyChangedChainE
 
 	static FName AnchorsProperty(TEXT("Anchors"));
 
-	FEditPropertyChain::TDoubleLinkedListNode* AnchorNode = PropertyChangedEvent.PropertyChain.GetHead()->GetNextNode();
-	if ( !AnchorNode )
+	if (FEditPropertyChain::TDoubleLinkedListNode* AnchorNode = PropertyChangedEvent.PropertyChain.GetHead()->GetNextNode())
 	{
-		return;
-	}
-
-	FEditPropertyChain::TDoubleLinkedListNode* LayoutDataNode = AnchorNode->GetNextNode();
-
-	if ( !LayoutDataNode )
-	{
-		return;
-	}
-
-	UProperty* AnchorProperty = LayoutDataNode->GetValue();
-
-	if ( AnchorProperty && AnchorProperty->GetFName() == AnchorsProperty )
-	{
-		RebaseLayout();
+		if (FEditPropertyChain::TDoubleLinkedListNode* LayoutDataNode = AnchorNode->GetNextNode())
+		{
+			UProperty* AnchorProperty = LayoutDataNode->GetValue();
+			if (AnchorProperty && AnchorProperty->GetFName() == AnchorsProperty)
+			{
+				RebaseLayout();
+			}
+		}
 	}
 
 	Super::PostEditChangeProperty(PropertyChangedEvent);
