Removed an early out for abstract classes from UWidgetBlueprintLibrary::Create() it ultimately avoided the CreateWidgetHelpers::ValidateUserWidgetClass() which catches this and gives a readable message. In shipping this code dosn't run but NewObject fails witha an ensure and returns null anyways, admittedly the error from NewObject() isn't as good.

diff --git a/Engine/Source/Runtime/UMG/Private/WidgetBlueprintLibrary.cpp b/Engine/Source/Runtime/UMG/Private/WidgetBlueprintLibrary.cpp
--- a/Engine/Source/Runtime/UMG/Private/WidgetBlueprintLibrary.cpp
+++ b/Engine/Source/Runtime/UMG/Private/WidgetBlueprintLibrary.cpp
@@ -36,7 +36,7 @@ UWidgetBlueprintLibrary::UWidgetBlueprintLibrary(const FObjectInitializer& Objec
 
 UUserWidget* UWidgetBlueprintLibrary::Create(UObject* WorldContextObject, TSubclassOf<UUserWidget> WidgetType, APlayerController* OwningPlayer)
 {
-	if ( WidgetType == nullptr || WidgetType->HasAnyClassFlags(CLASS_Abstract) )
+	if (WidgetType == nullptr)
 	{
 		return nullptr;
 	}
