PR #5315: Update RetainerBox SetRenderingPhase function (Contributed by [REDACTED_PERSON])

diff --git a/Engine/Source/Runtime/UMG/Private/Components/RetainerBox.cpp b/Engine/Source/Runtime/UMG/Private/Components/RetainerBox.cpp
--- a/Engine/Source/Runtime/UMG/Private/Components/RetainerBox.cpp
+++ b/Engine/Source/Runtime/UMG/Private/Components/RetainerBox.cpp
@@ -28,6 +28,11 @@ void URetainerBox::SetRenderingPhase(int PhaseToRenderOn, int32 TotalRenderingPh
 {
 	Phase = PhaseToRenderOn;
 	PhaseCount = TotalRenderingPhases;
+	
+	if ( PhaseCount < 1 )
+	{
+		PhaseCount = 1;
+	}
 
 	if (MyRetainerWidget.IsValid())
 	{
