PR #5368: Function GetAllChildren() (Contributed by [REDACTED_PERSON])

diff --git a/Engine/Source/Runtime/UMG/Private/Components/PanelWidget.cpp b/Engine/Source/Runtime/UMG/Private/Components/PanelWidget.cpp
--- a/Engine/Source/Runtime/UMG/Private/Components/PanelWidget.cpp
+++ b/Engine/Source/Runtime/UMG/Private/Components/PanelWidget.cpp
@@ -50,6 +50,18 @@ UWidget* UPanelWidget::GetChildAt(int32 Index) const
 	return nullptr;
 }
 
+TArray<UWidget*> GetAllChildren() const
+{
+	TArray<UWidget*> Result;
+
+	for(UPanelSlot* ChildSlot : Slots)
+	{
+		Result.Add(ChildSlot->Content);
+	}
+
+	return Result;
+}
+
 int32 UPanelWidget::GetChildIndex(UWidget* Content) const
 {
 	const int32 ChildCount = GetChildrenCount();
diff --git a/Engine/Source/Runtime/UMG/Public/Components/PanelWidget.h b/Engine/Source/Runtime/UMG/Public/Components/PanelWidget.h
--- a/Engine/Source/Runtime/UMG/Public/Components/PanelWidget.h
+++ b/Engine/Source/Runtime/UMG/Public/Components/PanelWidget.h
@@ -34,6 +34,10 @@ public:
 	 */
 	UFUNCTION(BlueprintCallable, Category="Widget|Panel")
 	UWidget* GetChildAt(int32 Index) const;
+	
+	/** Gets all widgets in the container */
+	UFUNCTION(BlueprintCallable, Category="Widget|Panel")
+	TArray<UWidget*> GetAllChildren() const;
 
 	/** Gets the index of a specific child widget */
 	UFUNCTION(BlueprintCallable, Category="Widget|Panel")
