Compile fix.

diff --git a/Engine/Source/Runtime/UMG/Private/Components/PanelWidget.cpp b/Engine/Source/Runtime/UMG/Private/Components/PanelWidget.cpp
--- a/Engine/Source/Runtime/UMG/Private/Components/PanelWidget.cpp
+++ b/Engine/Source/Runtime/UMG/Private/Components/PanelWidget.cpp
@@ -50,7 +50,7 @@ UWidget* UPanelWidget::GetChildAt(int32 Index) const
 	return nullptr;
 }
 
-TArray<UWidget*> GetAllChildren() const
+TArray<UWidget*> UPanelWidget::GetAllChildren() const
 {
 	TArray<UWidget*> Result;
 
