Removed Transient from the properties in FWidgetNavigationData. According to Nick Darnell, checking these in with Transient set in the previous check-in was not intentional.

diff --git a/Engine/Source/Runtime/UMG/Public/Blueprint/WidgetNavigation.h b/Engine/Source/Runtime/UMG/Public/Blueprint/WidgetNavigation.h
--- a/Engine/Source/Runtime/UMG/Public/Blueprint/WidgetNavigation.h
+++ b/Engine/Source/Runtime/UMG/Public/Blueprint/WidgetNavigation.h
@@ -30,10 +30,10 @@ public:
 	UPROPERTY(EditAnywhere, BlueprintReadOnly, Category="Navigation")
 	FName WidgetToFocus;
 
-	UPROPERTY(Transient)
+	UPROPERTY()
 	TWeakObjectPtr<UWidget> Widget;
 
-	UPROPERTY(Transient)
+	UPROPERTY()
 	FCustomWidgetNavigationDelegate CustomDelegate;
 
 	void Resolve(class UUserWidget* Outer, class UWidgetTree* WidgetTree);
