UMG - Fixing the designer's hidden eyeball to actually be an inherited state.  Since the designer does it's own hit testing, it wasn't excluding children of 'hidden in designer' widgets, and so even though it was visually invisible, it would still capture clicks, now, if you hide the parent widget, it and all children will also no longer be hit testable.

diff --git a/Engine/Source/Editor/UMGEditor/Private/Designer/SDesignerView.cpp b/Engine/Source/Editor/UMGEditor/Private/Designer/SDesignerView.cpp
--- a/Engine/Source/Editor/UMGEditor/Private/Designer/SDesignerView.cpp
+++ b/Engine/Source/Editor/UMGEditor/Private/Designer/SDesignerView.cpp
@@ -2023,7 +2023,7 @@ void SDesignerView::PopulateWidgetGeometryCache_Loop(FArrangedWidget& CurrentWid
 	{
 		bool bRespectLocks = IsRespectingLocks();
 
-		if (!CandidateUWidget->GetVisibilityInDesigner().IsVisible())
+		if (!CandidateUWidget->IsVisibleInDesigner())
 		{
 			bIncludeInHitTestGrid = false;
 		}
diff --git a/Engine/Source/Runtime/UMG/Private/Components/Widget.cpp b/Engine/Source/Runtime/UMG/Private/Components/Widget.cpp
--- a/Engine/Source/Runtime/UMG/Private/Components/Widget.cpp
+++ b/Engine/Source/Runtime/UMG/Private/Components/Widget.cpp
@@ -980,6 +980,27 @@ EVisibility UWidget::GetVisibilityInDesigner() const
 	return bHiddenInDesigner ? EVisibility::Collapsed : EVisibility::Visible;
 }
 
+bool UWidget::IsVisibleInDesigner() const
+{
+	if (bHiddenInDesigner)
+	{
+		return false;
+	}
+
+	UWidget* Parent = GetParent();
+	while (Parent != nullptr)
+	{
+		if (Parent->bHiddenInDesigner)
+		{
+			return false;
+		}
+
+		Parent = Parent->GetParent();
+	}
+
+	return true;
+}
+
 void UWidget::PostEditChangeProperty(FPropertyChangedEvent& PropertyChangedEvent)
 {
 	Super::PostEditChangeProperty(PropertyChangedEvent);
diff --git a/Engine/Source/Runtime/UMG/Public/Components/Widget.h b/Engine/Source/Runtime/UMG/Public/Components/Widget.h
--- a/Engine/Source/Runtime/UMG/Public/Components/Widget.h
+++ b/Engine/Source/Runtime/UMG/Public/Components/Widget.h
@@ -804,8 +804,8 @@ public:
 	virtual void PostEditChangeProperty(struct FPropertyChangedEvent& PropertyChangedEvent) override;
 	// End of UObject interface
 
-	/** Gets the visibility of the widget inside the designer. */
-	EVisibility GetVisibilityInDesigner() const;
+	/** Is the widget visible in the designer?  If this widget is 'hidden in the designer' or a parent is, this widget will also return false here. */
+	bool IsVisibleInDesigner() const;
 
 	// Begin Designer contextual events
 	void SelectByDesigner();
@@ -836,6 +836,12 @@ public:
 	static FString GetDefaultFontName();
 
 protected:
+#if WITH_EDITOR
+	// This is an implementation detail that allows us to show and hide the widget in the designer
+	// regardless of the actual visibility state set by the user.
+	EVisibility GetVisibilityInDesigner() const;
+#endif
+
 	virtual void OnBindingChanged(const FName& Property);
 
 protected:
