Add MaterialInstanceDynamic created by WidgetComponent to the component's GC cluster

diff --git a/Engine/Source/Runtime/UMG/Private/Components/WidgetComponent.cpp b/Engine/Source/Runtime/UMG/Private/Components/WidgetComponent.cpp
--- a/Engine/Source/Runtime/UMG/Private/Components/WidgetComponent.cpp
+++ b/Engine/Source/Runtime/UMG/Private/Components/WidgetComponent.cpp
@@ -649,7 +649,10 @@ void UWidgetComponent::UpdateMaterialInstance()
 
 	UMaterialInterface* BaseMaterial = GetMaterial(0);
 	MaterialInstance = UMaterialInstanceDynamic::Create(BaseMaterial, this);
-
+	if (MaterialInstance)
+	{
+			MaterialInstance->AddToCluster(this);
+	}
 	UpdateMaterialInstanceParameters();
 }
 
