Resolving robomerge block

--------------------------------------
- Undo/Redo now works again on properties of UserWidget instances in a UMG hierarchy (including the root widget)
- Added a couple WidgetTree static helpers for determining the child index of a PanelWidget that a given widget ultimately occupies (even if nested within other panels)

diff --git a/Engine/Source/Runtime/UMG/Private/Components/PanelWidget.cpp b/Engine/Source/Runtime/UMG/Private/Components/PanelWidget.cpp
--- a/Engine/Source/Runtime/UMG/Private/Components/PanelWidget.cpp
+++ b/Engine/Source/Runtime/UMG/Private/Components/PanelWidget.cpp
@@ -62,7 +62,7 @@ TArray<UWidget*> UPanelWidget::GetAllChildren() const
 	return Result;
 }
 
-int32 UPanelWidget::GetChildIndex(UWidget* Content) const
+int32 UPanelWidget::GetChildIndex(const UWidget* Content) const
 {
 	const int32 ChildCount = GetChildrenCount();
 	for ( int32 ChildIndex = 0; ChildIndex < ChildCount; ChildIndex++ )
