Fix small change that got lost in the (strangely baseless) robomerge resolve in CL 4754220

diff --git a/Engine/Source/Runtime/UMG/Public/Components/PanelWidget.h b/Engine/Source/Runtime/UMG/Public/Components/PanelWidget.h
--- a/Engine/Source/Runtime/UMG/Public/Components/PanelWidget.h
+++ b/Engine/Source/Runtime/UMG/Public/Components/PanelWidget.h
@@ -41,7 +41,7 @@ public:
 
 	/** Gets the index of a specific child widget */
 	UFUNCTION(BlueprintCallable, Category="Widget|Panel")
-	int32 GetChildIndex(UWidget* Content) const;
+	int32 GetChildIndex(const UWidget* Content) const;
 
 	/** @return true if panel contains this widget */
 	UFUNCTION(BlueprintCallable, Category="Widget|Panel")
