RE-re-resolving borked merge from yesterday, should fix compile

diff --git a/Engine/Source/Editor/UMGEditor/Private/WidgetBlueprintEditor.cpp b/Engine/Source/Editor/UMGEditor/Private/WidgetBlueprintEditor.cpp
--- a/Engine/Source/Editor/UMGEditor/Private/WidgetBlueprintEditor.cpp
+++ b/Engine/Source/Editor/UMGEditor/Private/WidgetBlueprintEditor.cpp
@@ -588,6 +588,7 @@ static bool MigratePropertyValue(UObject* SourceObject, UObject* DestinationObje
 		{
 			if (DestinationObject)
 			{
+				DestinationObject->SetFlags(RF_Transactional);
 				DestinationObject->Modify();
 			}
 			return true;
@@ -665,6 +666,7 @@ void FWidgetBlueprintEditor::PostUndo(bool bSuccessful)
 	FBlueprintEditor::PostUndo(bSuccessful);
 
 	OnWidgetBlueprintTransaction.Broadcast();
+	RefreshPreview();
 }
 
 void FWidgetBlueprintEditor::PostRedo(bool bSuccessful)
@@ -672,6 +674,7 @@ void FWidgetBlueprintEditor::PostRedo(bool bSuccessful)
 	FBlueprintEditor::PostRedo(bSuccessful);
 
 	OnWidgetBlueprintTransaction.Broadcast();
+	RefreshPreview();
 }
 
 TSharedRef<SWidget> FWidgetBlueprintEditor::CreateSequencerWidget()
diff --git a/Engine/Source/Runtime/UMG/Private/Components/Widget.cpp b/Engine/Source/Runtime/UMG/Private/Components/Widget.cpp
--- a/Engine/Source/Runtime/UMG/Private/Components/Widget.cpp
+++ b/Engine/Source/Runtime/UMG/Private/Components/Widget.cpp
@@ -1046,7 +1046,8 @@ bool UWidget::Modify(bool bAlwaysMarkDirty)
 
 	if ( Slot )
 	{
-		Modified &= Slot->Modify(bAlwaysMarkDirty);
+		Slot->SetFlags(RF_Transactional);
+		Modified |= Slot->Modify(bAlwaysMarkDirty);
 	}
 
 	return Modified;
diff --git a/Engine/Source/Runtime/UMG/Private/UserWidget.cpp b/Engine/Source/Runtime/UMG/Private/UserWidget.cpp
--- a/Engine/Source/Runtime/UMG/Private/UserWidget.cpp
+++ b/Engine/Source/Runtime/UMG/Private/UserWidget.cpp
@@ -2186,7 +2186,7 @@ UUserWidget* UUserWidget::CreateInstanceInternal(UObject* Outer, TSubclassOf<UUs
 #endif
 
 			FObjectInstancingGraph ObjectInstancingGraph;
-			NewWidget = NewObject<UUserWidget>(Outer, UserWidgetClass, InstanceName, RF_NoFlags, Template, false, &ObjectInstancingGraph);
+			NewWidget = NewObject<UUserWidget>(Outer, UserWidgetClass, InstanceName, RF_Transactional, Template, false, &ObjectInstancingGraph);
 		}
 #if !WITH_EDITOR && (UE_BUILD_DEBUG || UE_BUILD_DEVELOPMENT)
 		else
@@ -2210,7 +2210,7 @@ UUserWidget* UUserWidget::CreateInstanceInternal(UObject* Outer, TSubclassOf<UUs
 
 	if (!NewWidget)
 	{
-		NewWidget = NewObject<UUserWidget>(Outer, UserWidgetClass, InstanceName, RF_NoFlags);
+		NewWidget = NewObject<UUserWidget>(Outer, UserWidgetClass, InstanceName, RF_Transactional);
 	}
 	
 	if (LocalPlayer)
@@ -2276,4 +2276,4 @@ bool CreateWidgetHelpers::ValidateUserWidgetClass(const UClass* UserWidgetClass)
 	return true;
 }
 
-#undef LOCTEXT_NAMESPACE
\ No newline at end of file
+#undef LOCTEXT_NAMESPACE
diff --git a/Engine/Source/Runtime/UMG/Private/WidgetTree.cpp b/Engine/Source/Runtime/UMG/Private/WidgetTree.cpp
--- a/Engine/Source/Runtime/UMG/Private/WidgetTree.cpp
+++ b/Engine/Source/Runtime/UMG/Private/WidgetTree.cpp
@@ -67,6 +67,63 @@ UPanelWidget* UWidgetTree::FindWidgetParent(UWidget* Widget, int32& OutChildInde
 	return Parent;
 }
 
+>>>> ORIGINAL //UE4/Main/Engine/Source/Runtime/UMG/Private/WidgetTree.cpp#none
+==== THEIRS //UE4/Main/Engine/Source/Runtime/UMG/Private/WidgetTree.cpp@=4753498
+UWidget* UWidgetTree::FindWidgetChild(UPanelWidget* ParentWidget, FName ChildWidgetName, int32& OutChildIndex)
+{
+	OutChildIndex = INDEX_NONE;
+	UWidget* FoundChild = nullptr;
+
+	if (ParentWidget)
+	{
+		const auto CheckWidgetFunc = [&FoundChild, ChildWidgetName](UWidget* Widget) 
+			{
+				if (!FoundChild && Widget->GetFName() == ChildWidgetName)
+				{
+					FoundChild = Widget;
+				}
+			};
+
+		// Check all the children of the given ParentWidget, but only track the index at the top level
+		for (int32 ChildIdx = 0; ChildIdx < ParentWidget->GetChildrenCount(); ++ChildIdx)
+		{
+			CheckWidgetFunc(ParentWidget->GetChildAt(ChildIdx));
+
+			if (!FoundChild)
+			{
+				ForWidgetAndChildren(ParentWidget->GetChildAt(ChildIdx), CheckWidgetFunc);
+			}
+
+			if (FoundChild)
+			{
+				OutChildIndex = ChildIdx;
+				break;
+			}
+		}
+	}
+
+	return FoundChild;
+}
+
+int32 UWidgetTree::FindChildIndex(const UPanelWidget* ParentWidget, const UWidget* ChildWidget)
+{
+	const UWidget* CurrentWidget = ChildWidget;
+	while (CurrentWidget)
+	{
+		const UPanelWidget* NextParent = CurrentWidget->Slot ? CurrentWidget->Slot->Parent : nullptr;
+		if (NextParent && NextParent == ParentWidget)
+		{
+			return NextParent->GetChildIndex(CurrentWidget);
+		}
+
+		CurrentWidget = NextParent;
+	}
+
+	return INDEX_NONE;
+}
+
+==== YOURS //DanH_FortSpecial/Engine/Source/Runtime/UMG/Private/WidgetTree.cpp
+<<<<
 bool UWidgetTree::RemoveWidget(UWidget* InRemovedWidget)
 {
 	bool bRemoved = false;
diff --git a/Engine/Source/Runtime/UMG/Public/Blueprint/UserWidget.h b/Engine/Source/Runtime/UMG/Public/Blueprint/UserWidget.h
--- a/Engine/Source/Runtime/UMG/Public/Blueprint/UserWidget.h
+++ b/Engine/Source/Runtime/UMG/Public/Blueprint/UserWidget.h
@@ -1081,7 +1081,7 @@ public:
 	 * @see ValidateCompiledDefaults
 	 * @see ValidateCompiledWidgetTree
 	 */
-	void ValidateBlueprint(const UWidgetTree& BlueprintWidgetTree, FCompilerResultsLog& OutErrors) const;
+	void ValidateBlueprint(const UWidgetTree& BlueprintWidgetTree, FCompilerResultsLog& CompileLog) const;
 
 	/**
 	 * Override to perform any custom inspections of the default widget tree at the end of compilation.
@@ -1091,7 +1091,7 @@ public:
 	 *
 	 * Tip: If you need to validate properties of BindWidget members, you can search for them by property name within the widget tree.
 	 */
-	virtual void ValidateCompiledWidgetTree(const UWidgetTree& BlueprintWidgetTree, FCompilerResultsLog& OutErrors) const {};
+	virtual void ValidateCompiledWidgetTree(const UWidgetTree& BlueprintWidgetTree, FCompilerResultsLog& CompileLog) const {};
 #endif
 
 	static UUserWidget* CreateWidgetInstance(UWidget& OwningWidget, TSubclassOf<UUserWidget> UserWidgetClass, FName WidgetName);
diff --git a/Engine/Source/Runtime/UMG/Public/Blueprint/WidgetTree.h b/Engine/Source/Runtime/UMG/Public/Blueprint/WidgetTree.h
--- a/Engine/Source/Runtime/UMG/Public/Blueprint/WidgetTree.h
+++ b/Engine/Source/Runtime/UMG/Public/Blueprint/WidgetTree.h
@@ -43,11 +43,23 @@ public:
 	bool RemoveWidget(UWidget* Widget);
 
 	/** Gets the parent widget of a given widget, and potentially the child index. */
-	class UPanelWidget* FindWidgetParent(UWidget* Widget, int32& OutChildIndex);
+	static class UPanelWidget* FindWidgetParent(UWidget* Widget, int32& OutChildIndex);
+
+	/** 
+	 * Searches recursively through the children of the given ParentWidget to find a child widget of the given name.
+	 * If successful, also gets the index the child ultimately occupies within the starting ParentWidget (INDEX_NONE otherwise)
+	 */
+	static UWidget* FindWidgetChild(UPanelWidget* ParentWidget, FName ChildWidgetName, int32& OutChildIndex);
+
+	/**
+	 * Determines the child index of the given ParentWidget that the given ChildWidget ultimately occupies, accounting for nesting
+	 * @return The child slot index within ParentWidget that ChildWidget ultimately occupies (INDEX_NONE if ChildWidget is not within ParentWidget at any level)
+	 */
+	static int32 FindChildIndex(const UPanelWidget* ParentWidget, const UWidget* ChildWidget);
 
 	/** Gathers all the widgets in the tree recursively */
 	void GetAllWidgets(TArray<UWidget*>& Widgets) const;
-
+	
 	/** Gathers descendant child widgets of a parent widget. */
 	static void GetChildWidgets(UWidget* Parent, TArray<UWidget*>& Widgets);
 
