Removed garbage YOURS & THEIRS cruft inserted by P4

diff --git a/Engine/Source/Runtime/UMG/Private/WidgetTree.cpp b/Engine/Source/Runtime/UMG/Private/WidgetTree.cpp
--- a/Engine/Source/Runtime/UMG/Private/WidgetTree.cpp
+++ b/Engine/Source/Runtime/UMG/Private/WidgetTree.cpp
@@ -67,8 +67,6 @@ UPanelWidget* UWidgetTree::FindWidgetParent(UWidget* Widget, int32& OutChildInde
 	return Parent;
 }
 
->>>> ORIGINAL //UE4/Main/Engine/Source/Runtime/UMG/Private/WidgetTree.cpp#none
-==== THEIRS //UE4/Main/Engine/Source/Runtime/UMG/Private/WidgetTree.cpp@=4753498
 UWidget* UWidgetTree::FindWidgetChild(UPanelWidget* ParentWidget, FName ChildWidgetName, int32& OutChildIndex)
 {
 	OutChildIndex = INDEX_NONE;
@@ -122,8 +120,6 @@ int32 UWidgetTree::FindChildIndex(const UPanelWidget* ParentWidget, const UWidge
 	return INDEX_NONE;
 }
 
-==== YOURS //DanH_FortSpecial/Engine/Source/Runtime/UMG/Private/WidgetTree.cpp
-<<<<
 bool UWidgetTree::RemoveWidget(UWidget* InRemovedWidget)
 {
 	bool bRemoved = false;
