UMG: Fix deletion of rotation/angle track

diff --git a/Engine/Source/Runtime/UMG/Private/Animation/MovieScene2DTransformSection.cpp b/Engine/Source/Runtime/UMG/Private/Animation/MovieScene2DTransformSection.cpp
--- a/Engine/Source/Runtime/UMG/Private/Animation/MovieScene2DTransformSection.cpp
+++ b/Engine/Source/Runtime/UMG/Private/Animation/MovieScene2DTransformSection.cpp
@@ -184,7 +184,7 @@ FMovieScene2DTransformMask UMovieScene2DTransformSection::GetMaskByName(const FN
 	{
 		return EMovieScene2DTransformChannel::TranslationY;
 	}
-	else if (InName == TEXT("Angle"))
+	else if (InName == TEXT("Angle") || InName == TEXT("Rotation"))
 	{
 		return EMovieScene2DTransformChannel::Rotation;
 	}
