Merging

//UE4/Main/Engine/Source/Runtime/UMG/Public/Blueprint/WidgetTree.h

to //UE4/dev-sequencer/Engine/Source/Runtime/UMG/Public/Blueprint/WidgetTree.h

diff --git a/Engine/Source/Runtime/UMG/Public/Blueprint/WidgetTree.h b/Engine/Source/Runtime/UMG/Public/Blueprint/WidgetTree.h
--- a/Engine/Source/Runtime/UMG/Public/Blueprint/WidgetTree.h
+++ b/Engine/Source/Runtime/UMG/Public/Blueprint/WidgetTree.h
@@ -43,11 +43,23 @@ public:
 	bool RemoveWidget(UWidget* Widget);
 
 	/** Gets the parent widget of a given widget, and potentially the child index. */
-	class UPanelWidget* FindWidgetParent(UWidget* Widget, int32& OutChildIndex);
+	static class UPanelWidget* FindWidgetParent(UWidget* Widget, int32& OutChildIndex);
+
+	/** 
+	 * Searches recursively through the children of the given ParentWidget to find a child widget of the given name.
+	 * If successful, also gets the index the child ultimately occupies within the starting ParentWidget (INDEX_NONE otherwise)
+	 */
+	static UWidget* FindWidgetChild(UPanelWidget* ParentWidget, FName ChildWidgetName, int32& OutChildIndex);
+
+	/**
+	 * Determines the child index of the given ParentWidget that the given ChildWidget ultimately occupies, accounting for nesting
+	 * @return The child slot index within ParentWidget that ChildWidget ultimately occupies (INDEX_NONE if ChildWidget is not within ParentWidget at any level)
+	 */
+	static int32 FindChildIndex(const UPanelWidget* ParentWidget, const UWidget* ChildWidget);
 
 	/** Gathers all the widgets in the tree recursively */
 	void GetAllWidgets(TArray<UWidget*>& Widgets) const;
-
+	
 	/** Gathers descendant child widgets of a parent widget. */
 	static void GetChildWidgets(UWidget* Parent, TArray<UWidget*>& Widgets);
 
