Force remerge a few files that were not edited in dev-sequencer

diff --git a/Engine/Source/Runtime/UMG/Private/Components/Widget.cpp b/Engine/Source/Runtime/UMG/Private/Components/Widget.cpp
--- a/Engine/Source/Runtime/UMG/Private/Components/Widget.cpp
+++ b/Engine/Source/Runtime/UMG/Private/Components/Widget.cpp
@@ -1046,7 +1046,8 @@ bool UWidget::Modify(bool bAlwaysMarkDirty)
 
 	if ( Slot )
 	{
-		Modified &= Slot->Modify(bAlwaysMarkDirty);
+		Slot->SetFlags(RF_Transactional);
+		Modified |= Slot->Modify(bAlwaysMarkDirty);
 	}
 
 	return Modified;
diff --git a/Engine/Source/Runtime/UMG/Public/Blueprint/UserWidget.h b/Engine/Source/Runtime/UMG/Public/Blueprint/UserWidget.h
--- a/Engine/Source/Runtime/UMG/Public/Blueprint/UserWidget.h
+++ b/Engine/Source/Runtime/UMG/Public/Blueprint/UserWidget.h
@@ -1081,7 +1081,7 @@ public:
 	 * @see ValidateCompiledDefaults
 	 * @see ValidateCompiledWidgetTree
 	 */
-	void ValidateBlueprint(const UWidgetTree& BlueprintWidgetTree, FCompilerResultsLog& OutErrors) const;
+	void ValidateBlueprint(const UWidgetTree& BlueprintWidgetTree, FCompilerResultsLog& CompileLog) const;
 
 	/**
 	 * Override to perform any custom inspections of the default widget tree at the end of compilation.
@@ -1091,7 +1091,7 @@ public:
 	 *
 	 * Tip: If you need to validate properties of BindWidget members, you can search for them by property name within the widget tree.
 	 */
-	virtual void ValidateCompiledWidgetTree(const UWidgetTree& BlueprintWidgetTree, FCompilerResultsLog& OutErrors) const {};
+	virtual void ValidateCompiledWidgetTree(const UWidgetTree& BlueprintWidgetTree, FCompilerResultsLog& CompileLog) const {};
 #endif
 
 	static UUserWidget* CreateWidgetInstance(UWidget& OwningWidget, TSubclassOf<UUserWidget> UserWidgetClass, FName WidgetName);
