RE-re-resolving borked merge from yesterday, should fix compile

diff --git a/Engine/Source/Runtime/UMG/Private/UserWidget.cpp b/Engine/Source/Runtime/UMG/Private/UserWidget.cpp
--- a/Engine/Source/Runtime/UMG/Private/UserWidget.cpp
+++ b/Engine/Source/Runtime/UMG/Private/UserWidget.cpp
@@ -2186,7 +2186,7 @@ UUserWidget* UUserWidget::CreateInstanceInternal(UObject* Outer, TSubclassOf<UUs
 #endif
 
 			FObjectInstancingGraph ObjectInstancingGraph;
-			NewWidget = NewObject<UUserWidget>(Outer, UserWidgetClass, InstanceName, RF_NoFlags, Template, false, &ObjectInstancingGraph);
+			NewWidget = NewObject<UUserWidget>(Outer, UserWidgetClass, InstanceName, RF_Transactional, Template, false, &ObjectInstancingGraph);
 		}
 #if !WITH_EDITOR && (UE_BUILD_DEBUG || UE_BUILD_DEVELOPMENT)
 		else
@@ -2210,7 +2210,7 @@ UUserWidget* UUserWidget::CreateInstanceInternal(UObject* Outer, TSubclassOf<UUs
 
 	if (!NewWidget)
 	{
-		NewWidget = NewObject<UUserWidget>(Outer, UserWidgetClass, InstanceName, RF_NoFlags);
+		NewWidget = NewObject<UUserWidget>(Outer, UserWidgetClass, InstanceName, RF_Transactional);
 	}
 	
 	if (LocalPlayer)
@@ -2276,4 +2276,4 @@ bool CreateWidgetHelpers::ValidateUserWidgetClass(const UClass* UserWidgetClass)
 	return true;
 }
 
-#undef LOCTEXT_NAMESPACE
\ No newline at end of file
+#undef LOCTEXT_NAMESPACE
diff --git a/Engine/Source/Runtime/UMG/Private/WidgetTree.cpp b/Engine/Source/Runtime/UMG/Private/WidgetTree.cpp
--- a/Engine/Source/Runtime/UMG/Private/WidgetTree.cpp
+++ b/Engine/Source/Runtime/UMG/Private/WidgetTree.cpp
@@ -67,6 +67,63 @@ UPanelWidget* UWidgetTree::FindWidgetParent(UWidget* Widget, int32& OutChildInde
 	return Parent;
 }
 
+>>>> ORIGINAL //UE4/Main/Engine/Source/Runtime/UMG/Private/WidgetTree.cpp#none
+==== THEIRS //UE4/Main/Engine/Source/Runtime/UMG/Private/WidgetTree.cpp@=4753498
+UWidget* UWidgetTree::FindWidgetChild(UPanelWidget* ParentWidget, FName ChildWidgetName, int32& OutChildIndex)
+{
+	OutChildIndex = INDEX_NONE;
+	UWidget* FoundChild = nullptr;
+
+	if (ParentWidget)
+	{
+		const auto CheckWidgetFunc = [&FoundChild, ChildWidgetName](UWidget* Widget) 
+			{
+				if (!FoundChild && Widget->GetFName() == ChildWidgetName)
+				{
+					FoundChild = Widget;
+				}
+			};
+
+		// Check all the children of the given ParentWidget, but only track the index at the top level
+		for (int32 ChildIdx = 0; ChildIdx < ParentWidget->GetChildrenCount(); ++ChildIdx)
+		{
+			CheckWidgetFunc(ParentWidget->GetChildAt(ChildIdx));
+
+			if (!FoundChild)
+			{
+				ForWidgetAndChildren(ParentWidget->GetChildAt(ChildIdx), CheckWidgetFunc);
+			}
+
+			if (FoundChild)
+			{
+				OutChildIndex = ChildIdx;
+				break;
+			}
+		}
+	}
+
+	return FoundChild;
+}
+
+int32 UWidgetTree::FindChildIndex(const UPanelWidget* ParentWidget, const UWidget* ChildWidget)
+{
+	const UWidget* CurrentWidget = ChildWidget;
+	while (CurrentWidget)
+	{
+		const UPanelWidget* NextParent = CurrentWidget->Slot ? CurrentWidget->Slot->Parent : nullptr;
+		if (NextParent && NextParent == ParentWidget)
+		{
+			return NextParent->GetChildIndex(CurrentWidget);
+		}
+
+		CurrentWidget = NextParent;
+	}
+
+	return INDEX_NONE;
+}
+
+==== YOURS //DanH_FortSpecial/Engine/Source/Runtime/UMG/Private/WidgetTree.cpp
+<<<<
 bool UWidgetTree::RemoveWidget(UWidget* InRemovedWidget)
 {
 	bool bRemoved = false;
