Fixup build issues from bad robomerge into Dev-Networking

diff --git a/Engine/Source/Editor/UMGEditor/Private/WidgetBlueprintEditor.cpp b/Engine/Source/Editor/UMGEditor/Private/WidgetBlueprintEditor.cpp
--- a/Engine/Source/Editor/UMGEditor/Private/WidgetBlueprintEditor.cpp
+++ b/Engine/Source/Editor/UMGEditor/Private/WidgetBlueprintEditor.cpp
@@ -588,6 +588,7 @@ static bool MigratePropertyValue(UObject* SourceObject, UObject* DestinationObje
 		{
 			if (DestinationObject)
 			{
+				DestinationObject->SetFlags(RF_Transactional);
 				DestinationObject->Modify();
 			}
 			return true;
@@ -665,6 +666,7 @@ void FWidgetBlueprintEditor::PostUndo(bool bSuccessful)
 	FBlueprintEditor::PostUndo(bSuccessful);
 
 	OnWidgetBlueprintTransaction.Broadcast();
+	RefreshPreview();
 }
 
 void FWidgetBlueprintEditor::PostRedo(bool bSuccessful)
@@ -672,6 +674,7 @@ void FWidgetBlueprintEditor::PostRedo(bool bSuccessful)
 	FBlueprintEditor::PostRedo(bSuccessful);
 
 	OnWidgetBlueprintTransaction.Broadcast();
+	RefreshPreview();
 }
 
 TSharedRef<SWidget> FWidgetBlueprintEditor::CreateSequencerWidget()
diff --git a/Engine/Source/Runtime/UMG/Private/Components/Widget.cpp b/Engine/Source/Runtime/UMG/Private/Components/Widget.cpp
--- a/Engine/Source/Runtime/UMG/Private/Components/Widget.cpp
+++ b/Engine/Source/Runtime/UMG/Private/Components/Widget.cpp
@@ -1046,7 +1046,8 @@ bool UWidget::Modify(bool bAlwaysMarkDirty)
 
 	if ( Slot )
 	{
-		Modified &= Slot->Modify(bAlwaysMarkDirty);
+		Slot->SetFlags(RF_Transactional);
+		Modified |= Slot->Modify(bAlwaysMarkDirty);
 	}
 
 	return Modified;
diff --git a/Engine/Source/Runtime/UMG/Public/Blueprint/UserWidget.h b/Engine/Source/Runtime/UMG/Public/Blueprint/UserWidget.h
--- a/Engine/Source/Runtime/UMG/Public/Blueprint/UserWidget.h
+++ b/Engine/Source/Runtime/UMG/Public/Blueprint/UserWidget.h
@@ -1081,7 +1081,7 @@ public:
 	 * @see ValidateCompiledDefaults
 	 * @see ValidateCompiledWidgetTree
 	 */
-	void ValidateBlueprint(const UWidgetTree& BlueprintWidgetTree, FCompilerResultsLog& OutErrors) const;
+	void ValidateBlueprint(const UWidgetTree& BlueprintWidgetTree, FCompilerResultsLog& CompileLog) const;
 
 	/**
 	 * Override to perform any custom inspections of the default widget tree at the end of compilation.
@@ -1091,7 +1091,7 @@ public:
 	 *
 	 * Tip: If you need to validate properties of BindWidget members, you can search for them by property name within the widget tree.
 	 */
-	virtual void ValidateCompiledWidgetTree(const UWidgetTree& BlueprintWidgetTree, FCompilerResultsLog& OutErrors) const {};
+	virtual void ValidateCompiledWidgetTree(const UWidgetTree& BlueprintWidgetTree, FCompilerResultsLog& CompileLog) const {};
 #endif
 
 	static UUserWidget* CreateWidgetInstance(UWidget& OwningWidget, TSubclassOf<UUserWidget> UserWidgetClass, FName WidgetName);
diff --git a/Engine/Source/Runtime/UMG/Public/Blueprint/WidgetTree.h b/Engine/Source/Runtime/UMG/Public/Blueprint/WidgetTree.h
--- a/Engine/Source/Runtime/UMG/Public/Blueprint/WidgetTree.h
+++ b/Engine/Source/Runtime/UMG/Public/Blueprint/WidgetTree.h
@@ -43,11 +43,23 @@ public:
 	bool RemoveWidget(UWidget* Widget);
 
 	/** Gets the parent widget of a given widget, and potentially the child index. */
-	class UPanelWidget* FindWidgetParent(UWidget* Widget, int32& OutChildIndex);
+	static class UPanelWidget* FindWidgetParent(UWidget* Widget, int32& OutChildIndex);
+
+	/** 
+	 * Searches recursively through the children of the given ParentWidget to find a child widget of the given name.
+	 * If successful, also gets the index the child ultimately occupies within the starting ParentWidget (INDEX_NONE otherwise)
+	 */
+	static UWidget* FindWidgetChild(UPanelWidget* ParentWidget, FName ChildWidgetName, int32& OutChildIndex);
+
+	/**
+	 * Determines the child index of the given ParentWidget that the given ChildWidget ultimately occupies, accounting for nesting
+	 * @return The child slot index within ParentWidget that ChildWidget ultimately occupies (INDEX_NONE if ChildWidget is not within ParentWidget at any level)
+	 */
+	static int32 FindChildIndex(const UPanelWidget* ParentWidget, const UWidget* ChildWidget);
 
 	/** Gathers all the widgets in the tree recursively */
 	void GetAllWidgets(TArray<UWidget*>& Widgets) const;
-
+	
 	/** Gathers descendant child widgets of a parent widget. */
 	static void GetChildWidgets(UWidget* Parent, TArray<UWidget*>& Widgets);
 
