Ensure UserWidget bindings are created in packaged builds for child Blueprints

diff --git a/Engine/Source/Runtime/UMG/Private/WidgetBlueprintGeneratedClass.cpp b/Engine/Source/Runtime/UMG/Private/WidgetBlueprintGeneratedClass.cpp
--- a/Engine/Source/Runtime/UMG/Private/WidgetBlueprintGeneratedClass.cpp
+++ b/Engine/Source/Runtime/UMG/Private/WidgetBlueprintGeneratedClass.cpp
@@ -133,47 +133,53 @@ void UWidgetBlueprintGeneratedClass::InitializeBindingsStatic(UUserWidget* UserW
 {
 	ensure(!UserWidget->HasAnyFlags(RF_ArchetypeObject));
 
-	// For each property binding that we're given, find the corresponding field, and setup the delegate binding on the widget.
-	for ( const FDelegateRuntimeBinding& Binding : InBindings )
+	// Iterate all generated classes in the widget's class hierarchy and initialize bindings found on each one.
+	UClass* CurrentClass = UserWidget->GetClass();
+	while (UWidgetBlueprintGeneratedClass* WBPGC = Cast<UWidgetBlueprintGeneratedClass>(CurrentClass))
 	{
-		UObjectProperty* WidgetProperty = FindField<UObjectProperty>(UserWidget->GetClass(), *Binding.ObjectName);
-		if ( WidgetProperty == nullptr )
+		// For each property binding that we're given, find the corresponding field, and setup the delegate binding on the widget.
+		for (const FDelegateRuntimeBinding& Binding : WBPGC->Bindings)
 		{
-			continue;
-		}
-
-		UWidget* Widget = Cast<UWidget>(WidgetProperty->GetObjectPropertyValue_InContainer(UserWidget));
-
-		if ( Widget )
-		{
-			UDelegateProperty* DelegateProperty = FindField<UDelegateProperty>(Widget->GetClass(), FName(*( Binding.PropertyName.ToString() + TEXT("Delegate") )));
-			if ( !DelegateProperty )
+			UObjectProperty* WidgetProperty = FindField<UObjectProperty>(WBPGC, *Binding.ObjectName);
+			if (WidgetProperty == nullptr)
 			{
-				DelegateProperty = FindField<UDelegateProperty>(Widget->GetClass(), Binding.PropertyName);
+				continue;
 			}
 
-			if ( DelegateProperty )
-			{
-				bool bSourcePathBound = false;
+			UWidget* Widget = Cast<UWidget>(WidgetProperty->GetObjectPropertyValue_InContainer(UserWidget));
 
-				if ( Binding.SourcePath.IsValid() )
+			if (Widget)
+			{
+				UDelegateProperty* DelegateProperty = FindField<UDelegateProperty>(Widget->GetClass(), FName(*(Binding.PropertyName.ToString() + TEXT("Delegate"))));
+				if (!DelegateProperty)
 				{
-					bSourcePathBound = Widget->AddBinding(DelegateProperty, UserWidget, Binding.SourcePath);
+					DelegateProperty = FindField<UDelegateProperty>(Widget->GetClass(), Binding.PropertyName);
 				}
 
-				// If no native binder is found then the only possibility is that the binding is for
-				// a delegate that doesn't match the known native binders available and so we
-				// fallback to just attempting to bind to the function directly.
-				if ( bSourcePathBound == false )
+				if (DelegateProperty)
 				{
-					FScriptDelegate* ScriptDelegate = DelegateProperty->GetPropertyValuePtr_InContainer(Widget);
-					if ( ScriptDelegate )
+					bool bSourcePathBound = false;
+
+					if (Binding.SourcePath.IsValid())
+					{
+						bSourcePathBound = Widget->AddBinding(DelegateProperty, UserWidget, Binding.SourcePath);
+					}
+
+					// If no native binder is found then the only possibility is that the binding is for
+					// a delegate that doesn't match the known native binders available and so we
+					// fallback to just attempting to bind to the function directly.
+					if (bSourcePathBound == false)
 					{
-						ScriptDelegate->BindUFunction(UserWidget, Binding.FunctionName);
+						FScriptDelegate* ScriptDelegate = DelegateProperty->GetPropertyValuePtr_InContainer(Widget);
+						if (ScriptDelegate)
+						{
+							ScriptDelegate->BindUFunction(UserWidget, Binding.FunctionName);
+						}
 					}
 				}
 			}
 		}
+		CurrentClass = CurrentClass->GetSuperClass();
 	}
 }
 
