Fix clang cis

diff --git a/Engine/Source/Runtime/UMG/Private/Slate/WidgetRenderer.cpp b/Engine/Source/Runtime/UMG/Private/Slate/WidgetRenderer.cpp
--- a/Engine/Source/Runtime/UMG/Private/Slate/WidgetRenderer.cpp
+++ b/Engine/Source/Runtime/UMG/Private/Slate/WidgetRenderer.cpp
@@ -186,9 +186,9 @@ void FWidgetRenderer::DrawWindow(
 		FRenderThreadUpdateContext Context =
 		{
 			&DrawBuffer,
-			FApp::GetCurrentTime() - GStartTime,
-			FApp::GetDeltaTime(),
-			FPlatformTime::Seconds() - GStartTime,
+			static_cast<float>(FApp::GetCurrentTime() - GStartTime),
+			static_cast<float>(FApp::GetDeltaTime()),
+			static_cast<float>(FPlatformTime::Seconds() - GStartTime),
 			RenderTarget->GameThread_GetRenderTargetResource(),
 			Renderer.Get(),
 			bClearTarget
