PR #5375: Added support to update consume mouse wheel for ScrollBox (Contributed by [REDACTED_PERSON])

diff --git a/Engine/Source/Runtime/Slate/Private/Widgets/Layout/SScrollBox.cpp b/Engine/Source/Runtime/Slate/Private/Widgets/Layout/SScrollBox.cpp
--- a/Engine/Source/Runtime/Slate/Private/Widgets/Layout/SScrollBox.cpp
+++ b/Engine/Source/Runtime/Slate/Private/Widgets/Layout/SScrollBox.cpp
@@ -491,6 +491,11 @@ EOrientation SScrollBox::GetOrientation()
 	return Orientation;
 }
 
+void SScrollBox::SetConsumeMouseWheel(EConsumeMouseWheel NewConsumeMouseWheel)
+{
+	ConsumeMouseWheel = NewConsumeMouseWheel;
+}
+
 void SScrollBox::SetOrientation(EOrientation InOrientation)
 {
 	if (Orientation != InOrientation)
diff --git a/Engine/Source/Runtime/Slate/Public/Widgets/Layout/SScrollBox.h b/Engine/Source/Runtime/Slate/Public/Widgets/Layout/SScrollBox.h
--- a/Engine/Source/Runtime/Slate/Public/Widgets/Layout/SScrollBox.h
+++ b/Engine/Source/Runtime/Slate/Public/Widgets/Layout/SScrollBox.h
@@ -185,6 +185,8 @@ public:
 	/** Get the current orientation of the scrollbox. */
 	EOrientation GetOrientation();
 
+	void SetConsumeMouseWheel(EConsumeMouseWheel NewConsumeMouseWheel);
+
 	/** Sets the current orientation of the scrollbox and updates the layout */
 	void SetOrientation(EOrientation InOrientation);
 
diff --git a/Engine/Source/Runtime/UMG/Private/Components/ScrollBox.cpp b/Engine/Source/Runtime/UMG/Private/Components/ScrollBox.cpp
--- a/Engine/Source/Runtime/UMG/Private/Components/ScrollBox.cpp
+++ b/Engine/Source/Runtime/UMG/Private/Components/ScrollBox.cpp
@@ -99,6 +99,7 @@ void UScrollBox::SynchronizeProperties()
 	MyScrollBox->SetScrollBarAlwaysVisible(AlwaysShowScrollbar);
 	MyScrollBox->SetAllowOverscroll(AllowOverscroll ? EAllowOverscroll::Yes : EAllowOverscroll::No);
 	MyScrollBox->SetScrollBarRightClickDragAllowed(bAllowRightClickDragScrolling);
+	MyScrollBox->SetConsumeMouseWheel(ConsumeMouseWheel);
 }
 
 float UScrollBox::GetScrollOffset() const
@@ -193,6 +194,16 @@ void UScrollBox::PostLoad()
 	}
 }
 
+void UScrollBox::SetConsumeMouseWheel(EConsumeMouseWheel NewConsumeMouseWheel)
+{
+	ConsumeMouseWheel = NewConsumeMouseWheel;
+
+	if (MyScrollBox.IsValid())
+	{
+		MyScrollBox->SetConsumeMouseWheel(NewConsumeMouseWheel);
+	}
+}
+
 void UScrollBox::SetOrientation(EOrientation NewOrientation)
 {
 	Orientation = NewOrientation;
diff --git a/Engine/Source/Runtime/UMG/Public/Components/ScrollBox.h b/Engine/Source/Runtime/UMG/Public/Components/ScrollBox.h
--- a/Engine/Source/Runtime/UMG/Public/Components/ScrollBox.h
+++ b/Engine/Source/Runtime/UMG/Public/Components/ScrollBox.h
@@ -75,6 +75,10 @@ public:
 	UPROPERTY(EditAnywhere, BlueprintReadOnly, Category = "Scroll")
 	bool bAllowRightClickDragScrolling;
 
+
+	UFUNCTION(BlueprintCallable, Category = "Scroll")
+	void SetConsumeMouseWheel(EConsumeMouseWheel NewConsumeMouseWheel);
+
 	UFUNCTION(BlueprintCallable, Category = "Scroll")
 	void SetOrientation(EOrientation NewOrientation);
 
