PR #5380: Fix problem of infinity recursive call. (Contributed by [REDACTED_PERSON])

diff --git a/Engine/Source/Runtime/UMG/Private/UserWidget.cpp b/Engine/Source/Runtime/UMG/Private/UserWidget.cpp
--- a/Engine/Source/Runtime/UMG/Private/UserWidget.cpp
+++ b/Engine/Source/Runtime/UMG/Private/UserWidget.cpp
@@ -1976,7 +1976,7 @@ bool UUserWidget::ShouldSerializeWidgetTree(const ITargetPlatform* TargetPlatfor
 
 	// We preserve widget trees if you're a sub-object of an archetype that is going to serialize it's
 	// widget tree.
-	for (const UObject* It = this; It; It = It->GetOuter())
+	for (const UObject* It = GetOuter(); It; It = It->GetOuter())
 	{
 		if (It->HasAllFlags(RF_ArchetypeObject))
 		{
