PR #5443: Improved `GetAllWidgetsWithInterface` (Contributed by [REDACTED_PERSON])

diff --git a/Engine/Source/Runtime/UMG/Private/WidgetBlueprintLibrary.cpp b/Engine/Source/Runtime/UMG/Private/WidgetBlueprintLibrary.cpp
--- a/Engine/Source/Runtime/UMG/Private/WidgetBlueprintLibrary.cpp
+++ b/Engine/Source/Runtime/UMG/Private/WidgetBlueprintLibrary.cpp
@@ -503,7 +503,7 @@ void UWidgetBlueprintLibrary::GetAllWidgetsOfClass(UObject* WorldContextObject,
 		return;
 	}
 	 
-	const UWorld* World = GEngine->GetWorldFromContextObject(WorldContextObject, EGetWorldErrorMode::LogAndReturnNull);
+	const UWorld* const World = GEngine->GetWorldFromContextObject(WorldContextObject, EGetWorldErrorMode::LogAndReturnNull);
 	if ( !World )
 	{
 		return;
@@ -513,14 +513,8 @@ void UWidgetBlueprintLibrary::GetAllWidgetsOfClass(UObject* WorldContextObject,
 	{
 		UUserWidget* LiveWidget = *Itr;
 
-		// Skip any widget that's not in the current world context.
-		if ( LiveWidget->GetWorld() != World )
-		{
-			continue;
-		}
-
-		// Skip any widget that is not a child of the class specified.
-		if ( !LiveWidget->GetClass()->IsChildOf(WidgetClass) )
+		// Skip any widget that's not in the current world context or that is not a child of the class specified.
+		if (LiveWidget->GetWorld() != World || !LiveWidget->GetClass()->IsChildOf(WidgetClass))
 		{
 			continue;
 		}
@@ -539,16 +533,17 @@ void UWidgetBlueprintLibrary::GetAllWidgetsOfClass(UObject* WorldContextObject,
 	}
 }
 
-void UWidgetBlueprintLibrary::GetAllWidgetsWithInterface(UObject* WorldContextObject, TSubclassOf<UInterface> Interface, TArray<UUserWidget*>& FoundWidgets, bool TopLevelOnly)
+void UWidgetBlueprintLibrary::GetAllWidgetsWithInterface(UObject* WorldContextObject, TArray<UUserWidget*>& FoundWidgets, TSubclassOf<UInterface> Interface, bool TopLevelOnly)
 {
+	//Prevent possibility of an ever-growing array if user uses this in a loop
 	FoundWidgets.Empty();
 
 	if (!Interface || !WorldContextObject)
 	{
 		return;
 	}
-
-	const UWorld* World = GEngine->GetWorldFromContextObject(WorldContextObject, EGetWorldErrorMode::LogAndReturnNull);
+	 
+	const UWorld* const World = GEngine->GetWorldFromContextObject(WorldContextObject, EGetWorldErrorMode::LogAndReturnNull);
 	if (!World)
 	{
 		return;
@@ -558,19 +553,20 @@ void UWidgetBlueprintLibrary::GetAllWidgetsWithInterface(UObject* WorldContextOb
 	{
 		UUserWidget* LiveWidget = *Itr;
 
-		if (LiveWidget->GetWorld() != World)
+		// Skip any widget that's not in the current world context or that is not a child of the class specified.
+		if (LiveWidget->GetWorld() != World || !LiveWidget->GetClass()->ImplementsInterface(Interface))
 		{
 			continue;
 		}
 
 		if (TopLevelOnly)
 		{
-			if (LiveWidget->IsInViewport() && LiveWidget->GetClass()->ImplementsInterface(Interface))
+			if (LiveWidget->IsInViewport())
 			{
 				FoundWidgets.Add(LiveWidget);
 			}
 		}
-		else if (LiveWidget->GetClass()->ImplementsInterface(Interface))
+		else
 		{
 			FoundWidgets.Add(LiveWidget);
 		}
diff --git a/Engine/Source/Runtime/UMG/Public/Blueprint/WidgetBlueprintLibrary.h b/Engine/Source/Runtime/UMG/Public/Blueprint/WidgetBlueprintLibrary.h
--- a/Engine/Source/Runtime/UMG/Public/Blueprint/WidgetBlueprintLibrary.h
+++ b/Engine/Source/Runtime/UMG/Public/Blueprint/WidgetBlueprintLibrary.h
@@ -297,7 +297,7 @@ public:
 	* @param TopLevelOnly Only the widgets that are direct children of the viewport will be returned.
 	*/
 	UFUNCTION(BlueprintCallable, BlueprintCosmetic, Category="Widget", meta = (WorldContext = "WorldContextObject", DeterminesOutputType = "WidgetClass", DynamicOutputParam = "FoundWidgets"))
-	static void GetAllWidgetsWithInterface(UObject* WorldContextObject, TSubclassOf<UInterface> Interface, TArray<UUserWidget*>& FoundWidgets, bool TopLevelOnly);
+	static void GetAllWidgetsWithInterface(UObject* WorldContextObject, TArray<UUserWidget*>& FoundWidgets, TSubclassOf<UInterface> Interface, bool TopLevelOnly);
 
 	UFUNCTION(BlueprintPure, Category="Widget", meta = (CompactNodeTitle = "->", BlueprintAutocast))
 	static FInputEvent GetInputEventFromKeyEvent(const FKeyEvent& Event);
