PR #5438: Trigger Cancel if the dragged widget was dropped with different pointer index (Contributed by [REDACTED_PERSON])

diff --git a/Engine/Source/Runtime/UMG/Private/Slate/UMGDragDropOp.cpp b/Engine/Source/Runtime/UMG/Private/Slate/UMGDragDropOp.cpp
--- a/Engine/Source/Runtime/UMG/Private/Slate/UMGDragDropOp.cpp
+++ b/Engine/Source/Runtime/UMG/Private/Slate/UMGDragDropOp.cpp
@@ -39,9 +39,9 @@ bool FUMGDragDropOp::AffectedByPointerEvent(const FPointerEvent& PointerEvent)
 
 void FUMGDragDropOp::OnDrop( bool bDropWasHandled, const FPointerEvent& MouseEvent )
 {
-	if (AffectedByPointerEvent(MouseEvent))
+	if ( DragOperation )
 	{
-		if ( bDropWasHandled )
+		if ( bDropWasHandled && MouseEvent.GetPointerIndex() == PointerIndex )
 		{
 			DragOperation->Drop(MouseEvent);
 		}
