PR #3890: Fixed Trace returning wrong widget when stacked (Contributed by [REDACTED_PERSON])

diff --git a/Engine/Source/Runtime/UMG/Private/Components/WidgetInteractionComponent.cpp b/Engine/Source/Runtime/UMG/Private/Components/WidgetInteractionComponent.cpp
--- a/Engine/Source/Runtime/UMG/Private/Components/WidgetInteractionComponent.cpp
+++ b/Engine/Source/Runtime/UMG/Private/Components/WidgetInteractionComponent.cpp
@@ -204,6 +204,7 @@ UWidgetInteractionComponent::FWidgetTraceResult UWidgetInteractionComponent::Per
 				{
 					TraceResult.bWasHit = true;
 					TraceResult.HitResult = HitResult;
+					break;
 				}
 			}
 			else
