Adding more getters and setters for widget component properties

diff --git a/Engine/Source/Runtime/UMG/Public/Components/WidgetComponent.h b/Engine/Source/Runtime/UMG/Public/Components/WidgetComponent.h
--- a/Engine/Source/Runtime/UMG/Public/Components/WidgetComponent.h
+++ b/Engine/Source/Runtime/UMG/Public/Components/WidgetComponent.h
@@ -46,7 +46,7 @@ enum class EWidgetBlendMode : uint8
 	Transparent
 };
 
-UENUM()
+UENUM(BlueprintType)
 enum class EWidgetGeometryMode : uint8
 {
 	/** The widget is mapped onto a plane */
@@ -186,6 +186,13 @@ public:
 	UFUNCTION(BlueprintCallable, Category=UserInterface)
 	void SetOwnerPlayer(ULocalPlayer* LocalPlayer);
 
+	/** @see bManuallyRedraw */
+	UFUNCTION(BlueprintCallable, Category = UserInterface)
+	bool GetManuallyRedraw() const
+	{
+		return bManuallyRedraw;
+	};
+
 	/** @see bManuallyRedraw */
 	UFUNCTION(BlueprintCallable, Category=UserInterface)
 	void SetManuallyRedraw(bool bUseManualRedraw);
@@ -216,9 +223,31 @@ public:
 	/** Sets the blend mode to use for this widget */
 	void SetBlendMode( const EWidgetBlendMode NewBlendMode );
 
+	/** Gets whether the widget is two-sided or not */
+	UFUNCTION(BlueprintCallable, Category = Rendering)
+	bool GetTwoSided() const
+	{
+		return bIsTwoSided;
+	};
+
 	/** Sets whether the widget is two-sided or not */
+	UFUNCTION(BlueprintCallable, Category = Rendering)
 	void SetTwoSided( const bool bWantTwoSided );
 
+	/** Gets whether the widget ticks when offscreen or not */
+	UFUNCTION(BlueprintCallable, Category = Animation)
+	bool GetTickWhenOffscreen() const
+	{
+		return TickWhenOffscreen;
+	};
+
+	/** Sets whether the widget ticks when offscreen or not */
+	UFUNCTION(BlueprintCallable, Category = Animation)
+	void SetTickWhenOffscreen(const bool bWantTickWhenOffscreen)
+	{
+		TickWhenOffscreen = bWantTickWhenOffscreen;
+	};
+
 	/** Sets the background color and opacityscale for this widget */
 	UFUNCTION(BlueprintCallable, Category=UserInterface)
 	void SetBackgroundColor( const FLinearColor NewBackgroundColor );
@@ -231,21 +260,27 @@ public:
 	void SetOpacityFromTexture( const float NewOpacityFromTexture );
 
 	/** @return The pivot point where the UI is rendered about the origin. */
+	UFUNCTION(BlueprintCallable, Category = UserInterface)
 	FVector2D GetPivot() const { return Pivot; }
 
 	/**  */
+	UFUNCTION(BlueprintCallable, Category = UserInterface)
 	void SetPivot( const FVector2D& InPivot ) { Pivot = InPivot; }
 
 	/**  */
+	UFUNCTION(BlueprintCallable, Category = UserInterface)
 	bool GetDrawAtDesiredSize() const { return bDrawAtDesiredSize; }
 
 	/**  */
+	UFUNCTION(BlueprintCallable, Category = UserInterface)
 	void SetDrawAtDesiredSize(bool InbDrawAtDesiredSize) { bDrawAtDesiredSize = InbDrawAtDesiredSize; }
 
 	/**  */
+	UFUNCTION(BlueprintCallable, Category = UserInterface)
 	float GetRedrawTime() const { return RedrawTime; }
 
 	/**  */
+	UFUNCTION(BlueprintCallable, Category = UserInterface)
 	void SetRedrawTime(float bInRedrawTime) { RedrawTime = bInRedrawTime; }
 
 	/** Get the fake window we create for widgets displayed in the world. */
@@ -257,8 +292,10 @@ public:
 	/** Sets the widget class used to generate the widget for this component */
 	void SetWidgetClass(TSubclassOf<UUserWidget> InWidgetClass);
 
+	UFUNCTION(BlueprintCallable, Category = UserInterface)
 	EWidgetSpace GetWidgetSpace() const { return Space; }
 
+	UFUNCTION(BlueprintCallable, Category = UserInterface)
 	void SetWidgetSpace( EWidgetSpace NewSpace ) { Space = NewSpace; }
 
 	bool GetEditTimeUsable() const { return bEditTimeUsable; }
@@ -266,12 +303,31 @@ public:
 	void SetEditTimeUsable(bool Value) { bEditTimeUsable = Value; }
 
 	/** @see EWidgetGeometryMode, @see GetCylinderArcAngle() */
-	EWidgetGeometryMode GetGeometryMode() const { return GeometryMode; }
+	UFUNCTION(BlueprintCallable, Category = UserInterface)
+	EWidgetGeometryMode GetGeometryMode() const 
+	{
+		return GeometryMode; 
+	}
+
+	UFUNCTION(BlueprintCallable, Category = UserInterface)
+	void SetGeometryMode(EWidgetGeometryMode InGeometryMode) 
+	{ 
+		GeometryMode = InGeometryMode; 
+	}
 
 	bool GetReceiveHardwareInput() const { return bReceiveHardwareInput; }
 
 	/** Defines the curvature of the widget component when using EWidgetGeometryMode::Cylinder; ignored otherwise.  */
+	UFUNCTION(BlueprintCallable, Category = UserInterface)
 	float GetCylinderArcAngle() const { return CylinderArcAngle; }
+
+	/** Defines the curvature of the widget component when using EWidgetGeometryMode::Cylinder; ignored otherwise.  */
+	UFUNCTION(BlueprintCallable, Category = UserInterface)
+	void SetCylinderArcAngle(const float InCylinderArcAngle) 
+	{ 
+		CylinderArcAngle = InCylinderArcAngle; 
+	}
+
 	
 	/** Sets shared layer name used when this widget is initialized */
 	void SetInitialSharedLayerName(FName NewSharedLayerName) { SharedLayerName = NewSharedLayerName; }
@@ -279,6 +335,20 @@ public:
 	/** Sets layer z order used when this widget is initialized */
 	void SetInitialLayerZOrder(int32 NewLayerZOrder) { LayerZOrder = NewLayerZOrder; }
 
+	/** @see bWindowFocusable */
+	UFUNCTION(BlueprintCallable, Category = UserInterface)
+	bool GetWindowFocusable() const
+	{
+		return bWindowFocusable;
+	};
+
+	/** @see bWindowFocusable */
+	UFUNCTION(BlueprintCallable, Category = UserInterface)
+	void SetWindowFocusable(bool bInWindowFocusable)
+	{
+		bWindowFocusable = bInWindowFocusable;
+	};
+
 protected:
 	/** Just because the user attempts to receive hardware input does not mean it's possible. */
 	bool CanReceiveHardwareInput() const;
