Call correct super in CanvasPanelSlot

diff --git a/Engine/Source/Runtime/UMG/Private/Components/CanvasPanelSlot.cpp b/Engine/Source/Runtime/UMG/Private/Components/CanvasPanelSlot.cpp
--- a/Engine/Source/Runtime/UMG/Private/Components/CanvasPanelSlot.cpp
+++ b/Engine/Source/Runtime/UMG/Private/Components/CanvasPanelSlot.cpp
@@ -331,7 +331,7 @@ void UCanvasPanelSlot::PostEditChangeChainProperty(struct FPropertyChangedChainE
 		}
 	}
 
-	Super::PostEditChangeProperty(PropertyChangedEvent);
+	Super::PostEditChangeChainProperty(PropertyChangedEvent);
 }
 
 void UCanvasPanelSlot::SaveBaseLayout()
