PR #5436: Add `GetSelectedIndex` to `ComboBox` (Contributed by [REDACTED_PERSON])

diff --git a/Engine/Source/Runtime/UMG/Private/Components/ComboBoxString.cpp b/Engine/Source/Runtime/UMG/Private/Components/ComboBoxString.cpp
--- a/Engine/Source/Runtime/UMG/Private/Components/ComboBoxString.cpp
+++ b/Engine/Source/Runtime/UMG/Private/Components/ComboBoxString.cpp
@@ -215,10 +215,15 @@ void UComboBoxString::RefreshOptions()
 void UComboBoxString::SetSelectedOption(FString Option)
 {
 	int32 InitialIndex = FindOptionIndex(Option);
-	if (InitialIndex != -1)
+	SetSelectedIndex(InitialIndex);
+}
+
+void UComboBoxString::SetSelectedIndex(const int32 Index)
+{
+	if (Options.IsValidIndex(Index))
 	{
-		CurrentOptionPtr = Options[InitialIndex];
-		SelectedOption = Option;
+		CurrentOptionPtr = Options[Index];
+		SelectedOption = *CurrentOptionPtr;
 
 		if ( ComboBoxContent.IsValid() )
 		{
@@ -237,6 +242,22 @@ FString UComboBoxString::GetSelectedOption() const
 	return FString();
 }
 
+int32 UComboBoxString::GetSelectedIndex() const
+{
+	if (CurrentOptionPtr.IsValid())
+	{
+		for (int32 OptionIndex = 0; OptionIndex < Options.Num(); ++OptionIndex)
+		{
+			if (Options[OptionIndex] == CurrentOptionPtr)
+			{
+				return OptionIndex;
+			}
+		}
+	}
+
+	return -1;
+}
+
 int32 UComboBoxString::GetOptionCount() const
 {
 	return Options.Num();
diff --git a/Engine/Source/Runtime/UMG/Public/Components/ComboBoxString.h b/Engine/Source/Runtime/UMG/Public/Components/ComboBoxString.h
--- a/Engine/Source/Runtime/UMG/Public/Components/ComboBoxString.h
+++ b/Engine/Source/Runtime/UMG/Public/Components/ComboBoxString.h
@@ -124,9 +124,15 @@ public:
 	UFUNCTION(BlueprintCallable, Category="ComboBox")
 	void SetSelectedOption(FString Option);
 
+	UFUNCTION(BlueprintCallable, Category = "ComboBox")
+	void SetSelectedIndex(const int32 Index);
+
 	UFUNCTION(BlueprintCallable, Category="ComboBox")
 	FString GetSelectedOption() const;
 
+	UFUNCTION(BlueprintCallable, Category="ComboBox")
+	int32 GetSelectedIndex() const;
+
 	/** @return The number of options */
 	UFUNCTION(BlueprintCallable, Category="ComboBox")
 	int32 GetOptionCount() const;
