DR - Fix more merge issues
#rnx

diff --git a/Engine/Source/Runtime/D3D12RHI/Private/D3D12UAV.cpp b/Engine/Source/Runtime/D3D12RHI/Private/D3D12UAV.cpp
--- a/Engine/Source/Runtime/D3D12RHI/Private/D3D12UAV.cpp
+++ b/Engine/Source/Runtime/D3D12RHI/Private/D3D12UAV.cpp
@@ -262,14 +262,6 @@ FUnorderedAccessViewRHIRef FD3D12DynamicRHI::RHICreateUnorderedAccessView_Render
 	return RHICreateUnorderedAccessView(VertexBufferRHI, Format);
 }
 
-FD3D12StagingBuffer::~FD3D12StagingBuffer()
-{
-	if (StagedRead)
-	{
-		StagedRead.SafeRelease();
-	}
-}
-
 void* FD3D12StagingBuffer::Lock(uint32 Offset, uint32 NumBytes)
 {
 	check(!bIsLocked);
diff --git a/Engine/Source/Runtime/D3D12RHI/Public/D3D12Resources.h b/Engine/Source/Runtime/D3D12RHI/Public/D3D12Resources.h
--- a/Engine/Source/Runtime/D3D12RHI/Public/D3D12Resources.h
+++ b/Engine/Source/Runtime/D3D12RHI/Public/D3D12Resources.h
@@ -965,15 +965,9 @@ public:
 	}
 	void SafeRelease()
 	{
-		if (StagedRead)
-		{
-			StagedRead->Release();
-			StagedRead = nullptr;
-		}
+		StagedRead.SafeRelease();
 	}
 
-	virtual ~FD3D12StagingBuffer() final override;
-
 	virtual void* Lock(uint32 Offset, uint32 NumBytes) final override;
 	virtual void Unlock() final override;
 
diff --git a/Engine/Source/Runtime/UMG/Public/Blueprint/UserWidget.h b/Engine/Source/Runtime/UMG/Public/Blueprint/UserWidget.h
--- a/Engine/Source/Runtime/UMG/Public/Blueprint/UserWidget.h
+++ b/Engine/Source/Runtime/UMG/Public/Blueprint/UserWidget.h
@@ -1081,7 +1081,7 @@ public:
 	 * @see ValidateCompiledDefaults
 	 * @see ValidateCompiledWidgetTree
 	 */
-	void ValidateBlueprint(const UWidgetTree& BlueprintWidgetTree, FCompilerResultsLog& OutErrors) const;
+	void ValidateBlueprint(const UWidgetTree& BlueprintWidgetTree, FCompilerResultsLog& CompileLog) const;
 
 	/**
 	 * Override to perform any custom inspections of the default widget tree at the end of compilation.
@@ -1091,7 +1091,7 @@ public:
 	 *
 	 * Tip: If you need to validate properties of BindWidget members, you can search for them by property name within the widget tree.
 	 */
-	virtual void ValidateCompiledWidgetTree(const UWidgetTree& BlueprintWidgetTree, FCompilerResultsLog& OutErrors) const {};
+	virtual void ValidateCompiledWidgetTree(const UWidgetTree& BlueprintWidgetTree, FCompilerResultsLog& CompileLog) const {};
 #endif
 
 	static UUserWidget* CreateWidgetInstance(UWidget& OwningWidget, TSubclassOf<UUserWidget> UserWidgetClass, FName WidgetName);
diff --git a/Engine/Source/Runtime/UMG/Public/Blueprint/WidgetTree.h b/Engine/Source/Runtime/UMG/Public/Blueprint/WidgetTree.h
--- a/Engine/Source/Runtime/UMG/Public/Blueprint/WidgetTree.h
+++ b/Engine/Source/Runtime/UMG/Public/Blueprint/WidgetTree.h
@@ -43,7 +43,19 @@ public:
 	bool RemoveWidget(UWidget* Widget);
 
 	/** Gets the parent widget of a given widget, and potentially the child index. */
-	class UPanelWidget* FindWidgetParent(UWidget* Widget, int32& OutChildIndex);
+	static class UPanelWidget* FindWidgetParent(UWidget* Widget, int32& OutChildIndex);
+
+	/** 
+	 * Searches recursively through the children of the given ParentWidget to find a child widget of the given name.
+	 * If successful, also gets the index the child ultimately occupies within the starting ParentWidget (INDEX_NONE otherwise)
+	 */
+	static UWidget* FindWidgetChild(UPanelWidget* ParentWidget, FName ChildWidgetName, int32& OutChildIndex);
+
+	/**
+	 * Determines the child index of the given ParentWidget that the given ChildWidget ultimately occupies, accounting for nesting
+	 * @return The child slot index within ParentWidget that ChildWidget ultimately occupies (INDEX_NONE if ChildWidget is not within ParentWidget at any level)
+	 */
+	static int32 FindChildIndex(const UPanelWidget* ParentWidget, const UWidget* ChildWidget);
 
 	/** Gathers all the widgets in the tree recursively */
 	void GetAllWidgets(TArray<UWidget*>& Widgets) const;
