Backing out CL#4763485 to force reintegration. We have a ticket open with Perforce support to track down why this integration is causing problems.

diff --git a/Engine/Source/Runtime/UMG/Private/UserWidget.cpp b/Engine/Source/Runtime/UMG/Private/UserWidget.cpp
--- a/Engine/Source/Runtime/UMG/Private/UserWidget.cpp
+++ b/Engine/Source/Runtime/UMG/Private/UserWidget.cpp
@@ -2187,7 +2187,7 @@ UUserWidget* UUserWidget::CreateInstanceInternal(UObject* Outer, TSubclassOf<UUs
 #endif
 
 			FObjectInstancingGraph ObjectInstancingGraph;
-			NewWidget = NewObject<UUserWidget>(Outer, UserWidgetClass, InstanceName, RF_Transactional, Template, false, &ObjectInstancingGraph);
+			NewWidget = NewObject<UUserWidget>(Outer, UserWidgetClass, InstanceName, RF_NoFlags, Template, false, &ObjectInstancingGraph);
 		}
 #if !WITH_EDITOR && (UE_BUILD_DEBUG || UE_BUILD_DEVELOPMENT)
 		else
@@ -2211,7 +2211,7 @@ UUserWidget* UUserWidget::CreateInstanceInternal(UObject* Outer, TSubclassOf<UUs
 
 	if (!NewWidget)
 	{
-		NewWidget = NewObject<UUserWidget>(Outer, UserWidgetClass, InstanceName, RF_Transactional);
+		NewWidget = NewObject<UUserWidget>(Outer, UserWidgetClass, InstanceName, RF_NoFlags);
 	}
 	
 	if (LocalPlayer)
@@ -2277,4 +2277,4 @@ bool CreateWidgetHelpers::ValidateUserWidgetClass(const UClass* UserWidgetClass)
 	return true;
 }
 
-#undef LOCTEXT_NAMESPACE
+#undef LOCTEXT_NAMESPACE
\ No newline at end of file
diff --git a/Engine/Source/Runtime/UMG/Private/WidgetTree.cpp b/Engine/Source/Runtime/UMG/Private/WidgetTree.cpp
--- a/Engine/Source/Runtime/UMG/Private/WidgetTree.cpp
+++ b/Engine/Source/Runtime/UMG/Private/WidgetTree.cpp
@@ -67,59 +67,6 @@ UPanelWidget* UWidgetTree::FindWidgetParent(UWidget* Widget, int32& OutChildInde
 	return Parent;
 }
 
-UWidget* UWidgetTree::FindWidgetChild(UPanelWidget* ParentWidget, FName ChildWidgetName, int32& OutChildIndex)
-{
-	OutChildIndex = INDEX_NONE;
-	UWidget* FoundChild = nullptr;
-
-	if (ParentWidget)
-	{
-		const auto CheckWidgetFunc = [&FoundChild, ChildWidgetName](UWidget* Widget) 
-			{
-				if (!FoundChild && Widget->GetFName() == ChildWidgetName)
-				{
-					FoundChild = Widget;
-				}
-			};
-
-		// Check all the children of the given ParentWidget, but only track the index at the top level
-		for (int32 ChildIdx = 0; ChildIdx < ParentWidget->GetChildrenCount(); ++ChildIdx)
-		{
-			CheckWidgetFunc(ParentWidget->GetChildAt(ChildIdx));
-
-			if (!FoundChild)
-			{
-				ForWidgetAndChildren(ParentWidget->GetChildAt(ChildIdx), CheckWidgetFunc);
-			}
-
-			if (FoundChild)
-			{
-				OutChildIndex = ChildIdx;
-				break;
-			}
-		}
-	}
-
-	return FoundChild;
-}
-
-int32 UWidgetTree::FindChildIndex(const UPanelWidget* ParentWidget, const UWidget* ChildWidget)
-{
-	const UWidget* CurrentWidget = ChildWidget;
-	while (CurrentWidget)
-	{
-		const UPanelWidget* NextParent = CurrentWidget->Slot ? CurrentWidget->Slot->Parent : nullptr;
-		if (NextParent && NextParent == ParentWidget)
-		{
-			return NextParent->GetChildIndex(CurrentWidget);
-		}
-
-		CurrentWidget = NextParent;
-	}
-
-	return INDEX_NONE;
-}
-
 bool UWidgetTree::RemoveWidget(UWidget* InRemovedWidget)
 {
 	bool bRemoved = false;
diff --git a/Engine/Source/Runtime/UMG/Public/Blueprint/UserWidget.h b/Engine/Source/Runtime/UMG/Public/Blueprint/UserWidget.h
--- a/Engine/Source/Runtime/UMG/Public/Blueprint/UserWidget.h
+++ b/Engine/Source/Runtime/UMG/Public/Blueprint/UserWidget.h
@@ -1081,7 +1081,7 @@ public:
 	 * @see ValidateCompiledDefaults
 	 * @see ValidateCompiledWidgetTree
 	 */
-	void ValidateBlueprint(const UWidgetTree& BlueprintWidgetTree, class IWidgetCompilerLog& CompileLog) const;
+	void ValidateBlueprint(const UWidgetTree& BlueprintWidgetTree, FCompilerResultsLog& CompileLog) const;
 
 	/**
 	 * Override to perform any custom inspections of the default widget tree at the end of compilation.
@@ -1091,7 +1091,7 @@ public:
 	 *
 	 * Tip: If you need to validate properties of BindWidget members, you can search for them by property name within the widget tree.
 	 */
-	virtual void ValidateCompiledWidgetTree(const UWidgetTree& BlueprintWidgetTree, class IWidgetCompilerLog& CompileLog) const {};
+	virtual void ValidateCompiledWidgetTree(const UWidgetTree& BlueprintWidgetTree, FCompilerResultsLog& CompileLog) const {};
 #endif
 
 	static UUserWidget* CreateWidgetInstance(UWidget& OwningWidget, TSubclassOf<UUserWidget> UserWidgetClass, FName WidgetName);
diff --git a/Engine/Source/Runtime/UMG/Public/Blueprint/WidgetTree.h b/Engine/Source/Runtime/UMG/Public/Blueprint/WidgetTree.h
--- a/Engine/Source/Runtime/UMG/Public/Blueprint/WidgetTree.h
+++ b/Engine/Source/Runtime/UMG/Public/Blueprint/WidgetTree.h
@@ -43,19 +43,7 @@ public:
 	bool RemoveWidget(UWidget* Widget);
 
 	/** Gets the parent widget of a given widget, and potentially the child index. */
-	static class UPanelWidget* FindWidgetParent(UWidget* Widget, int32& OutChildIndex);
-
-	/** 
-	 * Searches recursively through the children of the given ParentWidget to find a child widget of the given name.
-	 * If successful, also gets the index the child ultimately occupies within the starting ParentWidget (INDEX_NONE otherwise)
-	 */
-	static UWidget* FindWidgetChild(UPanelWidget* ParentWidget, FName ChildWidgetName, int32& OutChildIndex);
-
-	/**
-	 * Determines the child index of the given ParentWidget that the given ChildWidget ultimately occupies, accounting for nesting
-	 * @return The child slot index within ParentWidget that ChildWidget ultimately occupies (INDEX_NONE if ChildWidget is not within ParentWidget at any level)
-	 */
-	static int32 FindChildIndex(const UPanelWidget* ParentWidget, const UWidget* ChildWidget);
+	class UPanelWidget* FindWidgetParent(UWidget* Widget, int32& OutChildIndex);
 
 	/** Gathers all the widgets in the tree recursively */
 	void GetAllWidgets(TArray<UWidget*>& Widgets) const;
