UMG - Making text labels that include newlines strip them out when making the labels used in the widget designer hierarchy.

diff --git a/Engine/Source/Runtime/UMG/Private/Components/TextBlock.cpp b/Engine/Source/Runtime/UMG/Private/Components/TextBlock.cpp
--- a/Engine/Source/Runtime/UMG/Private/Components/TextBlock.cpp
+++ b/Engine/Source/Runtime/UMG/Private/Components/TextBlock.cpp
@@ -305,7 +305,7 @@ FString UTextBlock::GetLabelMetadata() const
 {
 	const int32 MaxSampleLength = 15;
 
-	FString TextStr = Text.ToString();
+	FString TextStr = Text.ToString().Replace(TEXT("\n"), TEXT(" "));
 	TextStr = TextStr.Len() <= MaxSampleLength ? TextStr : TextStr.Left(MaxSampleLength - 2) + TEXT("..");
 	return TEXT(" \"") + TextStr + TEXT("\"");
 }
@@ -329,7 +329,7 @@ void UTextBlock::OnCreationFromPalette()
 
 bool UTextBlock::CanEditChange(const UProperty* InProperty) const
 {
-	if(bSimpleTextMode && InProperty)
+	if (bSimpleTextMode && InProperty)
 	{
 		static TArray<FName> InvalidPropertiesInSimpleMode =
 		{
@@ -346,7 +346,7 @@ bool UTextBlock::CanEditChange(const UProperty* InProperty) const
 		return !InvalidPropertiesInSimpleMode.Contains(InProperty->GetFName());
 	}
 
-	return true;
+	return Super::CanEditChange(InProperty);
 }
 
 #endif
