Exposing an override default text style to RichTextBlock
- Can be specified inline in the designer or modified at runtime
- Also exposed modifiers for properties inherited from UTextLayoutWidget

diff --git a/Engine/Source/Runtime/UMG/Private/Components/RichTextBlock.cpp b/Engine/Source/Runtime/UMG/Private/Components/RichTextBlock.cpp
--- a/Engine/Source/Runtime/UMG/Private/Components/RichTextBlock.cpp
+++ b/Engine/Source/Runtime/UMG/Private/Components/RichTextBlock.cpp
@@ -41,7 +41,7 @@ TSharedRef<SWidget> URichTextBlock::RebuildWidget()
 
 	MyRichTextBlock =
 		SNew(SRichTextBlock)
-		.TextStyle(&DefaultTextStyle)
+		.TextStyle(bOverrideDefaultStyle ? &DefaultTextStyleOverride : &DefaultTextStyle)
 		.Marshaller(Marshaller);
 	
 	return MyRichTextBlock.ToSharedRef();
@@ -112,6 +112,12 @@ const FTextBlockStyle& URichTextBlock::GetDefaultTextStyle() const
 	return DefaultTextStyle;
 }
 
+const FTextBlockStyle& URichTextBlock::GetCurrentDefaultTextStyle() const
+{
+	ensure(StyleInstance.IsValid());
+	return bOverrideDefaultStyle ? DefaultTextStyleOverride : DefaultTextStyle;
+}
+
 URichTextBlockDecorator* URichTextBlock::GetDecoratorByClass(TSubclassOf<URichTextBlockDecorator> DecoratorClass)
 {
 	for (URichTextBlockDecorator* Decorator : InstanceDecorators)
@@ -150,6 +156,16 @@ TSharedPtr< IRichTextMarkupWriter > URichTextBlock::CreateMarkupWriter()
 	return FDefaultRichTextMarkupWriter::Create();
 }
 
+void URichTextBlock::BeginDefaultStyleOverride()
+{
+	if (!bOverrideDefaultStyle)
+	{
+		// If we aren't already overriding, make sure override style starts off matching the existing default
+		bOverrideDefaultStyle = true;
+		DefaultTextStyleOverride = DefaultTextStyle;
+	}
+}
+
 #if WITH_EDITOR
 
 const FText URichTextBlock::GetPaletteCategory()
@@ -164,6 +180,93 @@ void URichTextBlock::OnCreationFromPalette()
 
 #endif
 
+void URichTextBlock::SetDefaultTextStyle(const FTextBlockStyle& InDefaultTextStyle)
+{
+	BeginDefaultStyleOverride();
+	DefaultTextStyleOverride = InDefaultTextStyle;
+	ApplyUpdatedDefaultTextStyle();
+}
+
+void URichTextBlock::ClearAllDefaultStyleOverrides()
+{
+	if (bOverrideDefaultStyle)
+	{
+		bOverrideDefaultStyle = false;
+		ApplyUpdatedDefaultTextStyle();
+	}
+}
+
+void URichTextBlock::SetDefaultColorAndOpacity(FSlateColor InColorAndOpacity)
+{
+	BeginDefaultStyleOverride();
+	DefaultTextStyleOverride.ColorAndOpacity = InColorAndOpacity;
+	ApplyUpdatedDefaultTextStyle();
+}
+
+
+void URichTextBlock::SetDefaultShadowColorAndOpacity(FLinearColor InShadowColorAndOpacity)
+{
+	BeginDefaultStyleOverride();
+	DefaultTextStyleOverride.ShadowColorAndOpacity = InShadowColorAndOpacity;
+	ApplyUpdatedDefaultTextStyle();
+}
+
+void URichTextBlock::SetDefaultShadowOffset(FVector2D InShadowOffset)
+{
+	BeginDefaultStyleOverride();
+	DefaultTextStyleOverride.ShadowOffset = InShadowOffset;
+	ApplyUpdatedDefaultTextStyle();
+}
+
+void URichTextBlock::SetDefaultFont(FSlateFontInfo InFontInfo)
+{
+	BeginDefaultStyleOverride();
+	DefaultTextStyleOverride.Font = InFontInfo;
+	ApplyUpdatedDefaultTextStyle();
+}
+
+void URichTextBlock::SetDefaultStrikeBrush(FSlateBrush& InStrikeBrush)
+{
+	BeginDefaultStyleOverride();
+	DefaultTextStyleOverride.StrikeBrush = InStrikeBrush;
+	ApplyUpdatedDefaultTextStyle();
+}
+
+void URichTextBlock::SetJustification(ETextJustify::Type InJustification)
+{
+	Justification = InJustification;
+	if (MyRichTextBlock.IsValid())
+	{
+		MyRichTextBlock->SetJustification(InJustification);
+	}
+}
+
+void URichTextBlock::SetMinDesiredWidth(float InMinDesiredWidth)
+{
+	MinDesiredWidth = InMinDesiredWidth;
+	if (MyRichTextBlock.IsValid())
+	{
+		MyRichTextBlock->SetMinDesiredWidth(InMinDesiredWidth);
+	}
+}
+
+void URichTextBlock::SetAutoWrapText(bool InAutoTextWrap)
+{
+	AutoWrapText = InAutoTextWrap;
+	if (MyRichTextBlock.IsValid())
+	{
+		MyRichTextBlock->SetAutoWrapText(InAutoTextWrap);
+	}
+}
+
+void URichTextBlock::ApplyUpdatedDefaultTextStyle()
+{
+	if (MyRichTextBlock.IsValid())
+	{
+		MyRichTextBlock->SetTextStyle(bOverrideDefaultStyle ? DefaultTextStyleOverride : DefaultTextStyle);
+	}
+}
+
 /////////////////////////////////////////////////////
 
 #undef LOCTEXT_NAMESPACE
diff --git a/Engine/Source/Runtime/UMG/Public/Components/RichTextBlock.h b/Engine/Source/Runtime/UMG/Public/Components/RichTextBlock.h
--- a/Engine/Source/Runtime/UMG/Public/Components/RichTextBlock.h
+++ b/Engine/Source/Runtime/UMG/Public/Components/RichTextBlock.h
@@ -37,6 +37,75 @@ class UMG_API URichTextBlock : public UTextLayoutWidget
 {
 	GENERATED_BODY()
 
+public:
+	/**  
+	 * Sets the color and opacity of the default text in this rich text block
+	 * @param InColorAndOpacity		The new text color and opacity
+	 */
+	UFUNCTION(BlueprintCallable, Category="Appearance")
+	void SetDefaultColorAndOpacity(FSlateColor InColorAndOpacity);
+
+	/**  
+	 * Sets the color and opacity of the default text drop shadow
+	 * Note: if opacity is zero no shadow will be drawn
+	 * @param InShadowColorAndOpacity		The new drop shadow color and opacity
+	 */
+	UFUNCTION(BlueprintCallable, Category="Appearance")
+	void SetDefaultShadowColorAndOpacity(FLinearColor InShadowColorAndOpacity);
+
+	/**  
+	 * Sets the offset that the default text drop shadow should be drawn at
+	 * @param InShadowOffset		The new offset
+	 */
+	UFUNCTION(BlueprintCallable, Category="Appearance")
+	void SetDefaultShadowOffset(FVector2D InShadowOffset);
+
+	/**
+	 * Dynamically set the default font info for this rich text block
+	 * @param InFontInfo The new font info
+	 */
+	UFUNCTION(BlueprintCallable, Category = "Appearance")
+	void SetDefaultFont(FSlateFontInfo InFontInfo);
+
+	/**
+	 * Dynamically set the default strike brush for this rich text block
+	 * @param InStrikeBrush The new brush to use to strike through text
+	 */
+	UFUNCTION(BlueprintCallable, Category = "Appearance")
+	void SetDefaultStrikeBrush(FSlateBrush& InStrikeBrush);
+
+	/**
+	 *  Set the text justification for this rich text block
+	 *  @param InJustification new justification
+	 */
+	UFUNCTION(BlueprintCallable, Category = "Appearance")
+	void SetJustification(ETextJustify::Type InJustification);
+	
+	/**
+	 *  Set the minimum desired width for this rich text block
+	 *  @param InMinDesiredWidth new minimum desired width
+	 */
+	UFUNCTION(BlueprintCallable, Category = "Appearance")
+	void SetMinDesiredWidth(float InMinDesiredWidth);
+	
+	/**
+	 * Set the auto wrap for this rich text block
+	 * @param InAutoTextWrap to turn wrap on or off
+	 */
+	UFUNCTION(BlueprintCallable, Category = "Appearance")
+	void SetAutoWrapText(bool InAutoTextWrap);
+
+	/** 
+	 * Wholesale override of the currently established default text style
+	 * @param InDefaultTextStyle The new text style to apply to all default (i.e. undecorated) text in the block
+	 */
+	UFUNCTION(BlueprintCallable, Category = Appearance)
+	void SetDefaultTextStyle(const FTextBlockStyle& InDefaultTextStyle);
+
+	/** Remove all overrides made to the default text style and return to the style specified in the style set data table */
+	UFUNCTION()
+	void ClearAllDefaultStyleOverrides();
+
 public:
 	URichTextBlock(const FObjectInitializer& ObjectInitializer);
 	
@@ -55,16 +124,32 @@ public:
 	// End UWidget interface
 #endif
 
-	/**  */
+	/**
+	 * Directly sets the widget text.
+	 * Warning: This will wipe any binding created for the Text property!
+	 * @param InText The text to assign to the widget
+	 */
 	UFUNCTION(BlueprintCallable, Category = "Widget")
 	void SetText(const FText& InText);
 
 	const FTextBlockStyle& GetDefaultTextStyle() const;
+	const FTextBlockStyle& GetCurrentDefaultTextStyle() const;
 
 	/**  */
 	UFUNCTION(BlueprintCallable, Category = "Widget")
 	URichTextBlockDecorator* GetDecoratorByClass(TSubclassOf<URichTextBlockDecorator> DecoratorClass);
 
+protected:
+	virtual TSharedRef<SWidget> RebuildWidget() override;
+	
+	virtual void UpdateStyleData();
+	virtual void CreateDecorators(TArray< TSharedRef< class ITextDecorator > >& OutDecorators);
+	virtual TSharedPtr< class IRichTextMarkupParser > CreateMarkupParser();
+	virtual TSharedPtr< class IRichTextMarkupWriter > CreateMarkupWriter();
+
+	void BeginDefaultStyleOverride();
+	void ApplyUpdatedDefaultTextStyle();
+
 protected:
 	/** The text to display */
 	UPROPERTY(EditAnywhere, Category=Content, meta=( MultiLine="true" ))
@@ -78,23 +163,22 @@ protected:
 	UPROPERTY(EditAnywhere, Category=Appearance)
 	TArray<TSubclassOf<URichTextBlockDecorator>> DecoratorClasses;
 
-protected:
-	FTextBlockStyle DefaultTextStyle;
+	/** True to specify the default text style for this rich text inline, overriding any default provided in the style set table */
+	UPROPERTY(EditAnywhere, Category = Appearance)
+	bool bOverrideDefaultStyle = false;
 
-	TSharedPtr<class FSlateStyleSet> StyleInstance;
+	/** Text style to apply by default to text in this block */
+	UPROPERTY(EditAnywhere, Category=Appearance, meta = (EditCondition = bOverrideDefaultStyle))
+	FTextBlockStyle DefaultTextStyleOverride;
+
+	/** The minimum desired size for the text */
+	UPROPERTY(EditAnywhere, BlueprintReadOnly, Category=Appearance)
+	float MinDesiredWidth;
+
+	FTextBlockStyle DefaultTextStyle;
 
 	UPROPERTY(Transient)
 	TArray<URichTextBlockDecorator*> InstanceDecorators;
-
-	/** Native Slate Widget */
+	TSharedPtr<class FSlateStyleSet> StyleInstance;
 	TSharedPtr<SRichTextBlock> MyRichTextBlock;
-
-	// UWidget interface
-	virtual TSharedRef<SWidget> RebuildWidget() override;
-	// End of UWidget interface
-
-	virtual void UpdateStyleData();
-	virtual void CreateDecorators(TArray< TSharedRef< class ITextDecorator > >& OutDecorators);
-	virtual TSharedPtr< class IRichTextMarkupParser > CreateMarkupParser();
-	virtual TSharedPtr< class IRichTextMarkupWriter > CreateMarkupWriter();
 };
