Fixed SetStrikeBrush showing as an output parameter in BP

#rnx

diff --git a/Engine/Source/Runtime/UMG/Private/Components/TextBlock.cpp b/Engine/Source/Runtime/UMG/Private/Components/TextBlock.cpp
--- a/Engine/Source/Runtime/UMG/Private/Components/TextBlock.cpp
+++ b/Engine/Source/Runtime/UMG/Private/Components/TextBlock.cpp
@@ -94,7 +94,7 @@ void UTextBlock::SetFont(FSlateFontInfo InFontInfo)
 	}
 }
 
-void UTextBlock::SetStrikeBrush(FSlateBrush& InStrikeBrush)
+void UTextBlock::SetStrikeBrush(FSlateBrush InStrikeBrush)
 {
 	StrikeBrush = InStrikeBrush;
 	if (MyTextBlock.IsValid())
diff --git a/Engine/Source/Runtime/UMG/Public/Components/TextBlock.h b/Engine/Source/Runtime/UMG/Public/Components/TextBlock.h
--- a/Engine/Source/Runtime/UMG/Public/Components/TextBlock.h
+++ b/Engine/Source/Runtime/UMG/Public/Components/TextBlock.h
@@ -71,7 +71,7 @@ public:
 	 * @param InStrikeBrush The new brush to use to strike through text
 	 */
 	UFUNCTION(BlueprintCallable, Category = "Appearance")
-	void SetStrikeBrush(FSlateBrush& InStrikeBrush);
+	void SetStrikeBrush(FSlateBrush InStrikeBrush);
 
 	/**
 	 *  Set the text justification for this text block
