UMG - Adding SetBrushResourceObject to make it possible to directly set the resource object of the brush, without needing to go through the complexities of getting the current brush, copying all but a single property, and replacing that one in blueprints.

Chris.Gagnon, Vincent.Gauthier

diff --git a/Engine/Source/Runtime/UMG/Private/Components/Image.cpp b/Engine/Source/Runtime/UMG/Private/Components/Image.cpp
--- a/Engine/Source/Runtime/UMG/Private/Components/Image.cpp
+++ b/Engine/Source/Runtime/UMG/Private/Components/Image.cpp
@@ -133,6 +133,20 @@ void UImage::SetBrushTintColor(FSlateColor TintColor)
 	}
 }
 
+void UImage::SetBrushResourceObject(UObject* ResourceObject)
+{
+	if (Brush.GetResourceObject() != ResourceObject)
+	{
+		Brush.SetResourceObject(ResourceObject);
+
+		if (MyImage.IsValid())
+		{
+			MyImage->SetImage(&Brush);
+			MyImage->Invalidate(EInvalidateWidget::PaintAndVolatility);
+		}
+	}
+}
+
 void UImage::SetBrushFromAsset(USlateBrushAsset* Asset)
 {
 	if(!Asset || Brush != Asset->Brush)
diff --git a/Engine/Source/Runtime/UMG/Public/Components/Image.h b/Engine/Source/Runtime/UMG/Public/Components/Image.h
--- a/Engine/Source/Runtime/UMG/Public/Components/Image.h
+++ b/Engine/Source/Runtime/UMG/Public/Components/Image.h
@@ -81,6 +81,10 @@ public:
 	/**  */
 	UFUNCTION(BlueprintCallable, Category = "Appearance")
 	void SetBrushTintColor(FSlateColor TintColor);
+
+	/**  */
+	UFUNCTION(BlueprintCallable, Category = "Appearance")
+	void SetBrushResourceObject(UObject* ResourceObject);
 	
 	/**  */
 	UFUNCTION(BlueprintCallable, Category="Appearance")
