UMG - Adding a method to allow the user widget pool to function outside of a widget, by giving it a world to outer new widgets to.  Adding an ensure to make sure it's not being called before it has been initialized.

Dan.Hertzka

diff --git a/Engine/Source/Runtime/UMG/Private/Blueprint/UserWidgetPool.cpp b/Engine/Source/Runtime/UMG/Private/Blueprint/UserWidgetPool.cpp
--- a/Engine/Source/Runtime/UMG/Private/Blueprint/UserWidgetPool.cpp
+++ b/Engine/Source/Runtime/UMG/Private/Blueprint/UserWidgetPool.cpp
@@ -25,6 +25,11 @@ FUserWidgetPool::~FUserWidgetPool()
 	ResetPool();
 }
 
+void FUserWidgetPool::SetWorld(UWorld* InOwningWorld)
+{
+	OwningWorld = InOwningWorld;
+}
+
 void FUserWidgetPool::AddReferencedObjects(FReferenceCollector& Collector)
 {
 	Collector.AddReferencedObjects<UUserWidget>(ActiveWidgets, OwningWidget.Get());
diff --git a/Engine/Source/Runtime/UMG/Public/Blueprint/UserWidgetPool.h b/Engine/Source/Runtime/UMG/Public/Blueprint/UserWidgetPool.h
--- a/Engine/Source/Runtime/UMG/Public/Blueprint/UserWidgetPool.h
+++ b/Engine/Source/Runtime/UMG/Public/Blueprint/UserWidgetPool.h
@@ -25,9 +25,12 @@ public:
 	FUserWidgetPool& operator=(FUserWidgetPool&& Other);
 	~FUserWidgetPool();
 
+	/** In the case that you don't have an owner widget, you should set a world to your pool, or it won't be able to construct widgets. */
+	void SetWorld(UWorld* OwningWorld);
+
 	virtual void AddReferencedObjects(FReferenceCollector& Collector) override;
 
-	bool IsInitialized() const { return OwningWidget.IsValid(); }
+	bool IsInitialized() const { return OwningWidget.IsValid() || OwningWorld.IsValid(); }
 	const TArray<UUserWidget*>& GetActiveWidgets() const { return ActiveWidgets; }
 
 	/**
@@ -67,7 +70,7 @@ private:
 	template <typename UserWidgetT = UUserWidget>
 	UserWidgetT* AddActiveWidgetInternal(TSubclassOf<UserWidgetT> WidgetClass, WidgetConstructFunc ConstructWidgetFunc)
 	{
-		if (!IsInitialized())
+		if (!ensure(IsInitialized()))
 		{
 			return nullptr;
 		}
@@ -85,7 +88,14 @@ private:
 
 		if (!WidgetInstance)
 		{
-			WidgetInstance = CreateWidget(OwningWidget.Get(), WidgetClass);
+			if (UWidget* OwnerWidgetPtr = OwningWidget.Get())
+			{
+				WidgetInstance = CreateWidget(OwnerWidgetPtr, WidgetClass);
+			}
+			else
+			{
+				WidgetInstance = CreateWidget(OwningWorld.Get(), WidgetClass);
+			}
 		}
 
 		if (WidgetInstance)
@@ -102,6 +112,7 @@ private:
 	}
 
 	TWeakObjectPtr<UWidget> OwningWidget;
+	TWeakObjectPtr<UWorld> OwningWorld;
 	
 	TArray<UUserWidget*> ActiveWidgets;
 	TArray<UUserWidget*> InactiveWidgets;
