[UE-73025] - Fixing UserListEntry interface structure screwup, inheriting directly from IUserListEntry in a BP will no longer crash
- Removed INativeUserListEntry and moved the functionality into native functions on IUserListEntry and a static BP library
- UUserListEntry is now marked as a BlueprintType, so it'll appear as a Cast option in BP (gotta disable "context sensitive" in the context menu)

Fixed unnecessary requirement that IUserObjectListEntry children cache and return their assigned list item in GetListItemObject()
- The existing design of SListView has the list tracking the entry-item association and the entries accessing their item via the owning list, so the UMG version now follows suit
- GetListItemObject is now a static BP method, the previously existing GetListItem<T>() remains as before
- Renamed SetListItemObjectInternal to NativeOnListItemObjectSet, as it isn't required to function as a literal setter anymore, so we can keep to the naming scheme used elsewhere in UMG for this scenario

Fixing half-implemented option to change selection mode of a UListView
- Now properly forwards the setting to the SListView to affect future selection behavior

Note: Redirects are set up for the functions that became static BP functions and work in situ for cases where they were being called on "self". Cases where an external object called these methods on another objects will point to the new function, but will require manual node fixup

diff --git a/Engine/Config/BaseEngine.ini b/Engine/Config/BaseEngine.ini
--- a/Engine/Config/BaseEngine.ini
+++ b/Engine/Config/BaseEngine.ini
@@ -663,6 +663,14 @@ bUseStreamingPause=false
 +ClassRedirects=(OldName="EditorAutomationActorComponent",NewName="EditorUtilityActorComponent")
 +ClassRedirects=(OldName="EditorAutomationObject",NewName="EditorUtilityObject")
 
+; These aren't fully "clean" redirects, but the previous setup was broken and the necessary fix cannot be redirected to cleanly in all cases
+; - If calling these on "self" in a BP where the self implements IUserListEntry, the redirect works cleanly
+; - If calling these on an external object (rare), the redirect will point users to the new function to use, but the nodes will need manual fixup
++FunctionRedirects=(OldName="NativeUserListEntry.IsListItemSelected", NewName="UserListEntryLibrary.IsListItemSelected")
++FunctionRedirects=(OldName="NativeUserListEntry.IsListItemExpanded", NewName="UserListEntryLibrary.IsListItemExpanded")
++FunctionRedirects=(OldName="NativeUserListEntry.GetOwningListView", NewName="UserListEntryLibrary.GetOwningListView")
++FunctionRedirects=(OldName="UserObjectListEntry.GetListItemObject", NewName="UserObjectListEntryLibrary.GetListItemObject")
+
 [CoreUObject.Metadata]
 ; Note: UnrealHeaderTool should be rerun after making changes to MetadataRedirects to catch any keys specified in class headers
 +MetadataRedirects=(OldKey="K2Protected", NewKey="BlueprintProtected")
diff --git a/Engine/Source/Runtime/Slate/Public/Widgets/Views/SListView.h b/Engine/Source/Runtime/Slate/Public/Widgets/Views/SListView.h
--- a/Engine/Source/Runtime/Slate/Public/Widgets/Views/SListView.h
+++ b/Engine/Source/Runtime/Slate/Public/Widgets/Views/SListView.h
@@ -1296,6 +1296,39 @@ public:
 		SetItemSelection( SoleSelectedItem, true, SelectInfo );
 	}
 
+	/** 
+	 * Set the current selection mode of the list.
+	 * If going from multi-select to a type of single-select and one item is selected, it will be maintained (otherwise all will be cleared).
+	 * If disabling selection, any current selections will be cleared.
+	 */
+	void SetSelectionMode(const TAttribute<ESelectionMode::Type>& NewSelectionMode)
+	{
+		const ESelectionMode::Type PreviousMode = SelectionMode.Get();
+		SelectionMode = NewSelectionMode;
+		const ESelectionMode::Type NewMode = NewSelectionMode.Get();
+		if (PreviousMode != NewMode)
+		{
+			if (NewMode == ESelectionMode::None)
+			{
+				ClearSelection();
+			}
+			else if (PreviousMode == ESelectionMode::Multi)
+			{
+				// We've gone to a single-selection mode, so if we already had a single item selected, preserve it
+				if (SelectedItems.Num() == 1)
+				{
+					SetSelection(*SelectedItems.CreateIterator());
+				}
+				else
+				{
+					// Otherwise, there's no way to know accurately which item was selected most recently, so just wipe it all
+					// The caller responsible for changing the mode can decide themselves which item they want to be selected
+					ClearSelection();
+				}
+			}
+		}
+	}
+
 	/**
 	 * Find a widget for this item if it has already been constructed.
 	 *
diff --git a/Engine/Source/Runtime/UMG/Private/Blueprint/IUserListEntry.cpp b/Engine/Source/Runtime/UMG/Private/Blueprint/IUserListEntry.cpp
--- a/Engine/Source/Runtime/UMG/Private/Blueprint/IUserListEntry.cpp
+++ b/Engine/Source/Runtime/UMG/Private/Blueprint/IUserListEntry.cpp
@@ -8,11 +8,6 @@
 
 TMap<TWeakObjectPtr<const UUserWidget>, TWeakPtr<const IObjectTableRow>> IObjectTableRow::ObjectRowsByUserWidget;
 
-UNativeUserListEntry::UNativeUserListEntry(const FObjectInitializer& Initializer)
-	: Super(Initializer)
-{
-}
-
 UUserListEntry::UUserListEntry(const FObjectInitializer& Initializer)
 	: Super(Initializer)
 {
@@ -56,32 +51,17 @@ void IUserListEntry::UpdateItemExpansion(UUserWidget& ListEntryWidget, bool bIsE
 
 bool IUserListEntry::IsListItemSelected() const
 {
-	TSharedPtr<const IObjectTableRow> SlateRow = IObjectTableRow::ObjectRowFromUserWidget(CastChecked<const UUserWidget>(this, ECastCheckedType::NullAllowed));
-	if (SlateRow.IsValid())
-	{
-		return SlateRow->IsItemSelected();
-	}
-	return false;
+	return UUserListEntryLibrary::IsListItemSelected(Cast<UUserWidget>(const_cast<IUserListEntry*>(this)));
 }
 
 bool IUserListEntry::IsListItemExpanded() const
 {
-	TSharedPtr<const IObjectTableRow> SlateRow = IObjectTableRow::ObjectRowFromUserWidget(CastChecked<const UUserWidget>(this, ECastCheckedType::NullAllowed));
-	if (SlateRow.IsValid())
-	{
-		return SlateRow->IsItemExpanded();
-	}
-	return false;
+	return UUserListEntryLibrary::IsListItemExpanded(Cast<UUserWidget>(const_cast<IUserListEntry*>(this)));
 }
 
 UListViewBase* IUserListEntry::GetOwningListView() const
 {
-	TSharedPtr<const IObjectTableRow> SlateRow = IObjectTableRow::ObjectRowFromUserWidget(CastChecked<const UUserWidget>(this, ECastCheckedType::NullAllowed));
-	if (SlateRow.IsValid())
-	{
-		return SlateRow->GetOwningListView();
-	}
-	return nullptr;
+	return UUserListEntryLibrary::GetOwningListView(Cast<UUserWidget>(const_cast<IUserListEntry*>(this)));
 }
 
 void IUserListEntry::NativeOnEntryReleased()
@@ -97,4 +77,34 @@ void IUserListEntry::NativeOnItemSelectionChanged(bool bIsSelected)
 void IUserListEntry::NativeOnItemExpansionChanged(bool bIsExpanded)
 {
 	Execute_BP_OnItemExpansionChanged(Cast<UObject>(this), bIsExpanded);
-}
\ No newline at end of file
+}
+
+bool UUserListEntryLibrary::IsListItemSelected(TScriptInterface<IUserListEntry> UserListEntry)
+{
+	TSharedPtr<const IObjectTableRow> SlateRow = IObjectTableRow::ObjectRowFromUserWidget(Cast<UUserWidget>(UserListEntry.GetObject()));
+	if (SlateRow.IsValid())
+	{
+		return SlateRow->IsItemSelected();
+	}
+	return false;
+}
+
+bool UUserListEntryLibrary::IsListItemExpanded(TScriptInterface<IUserListEntry> UserListEntry)
+{
+	TSharedPtr<const IObjectTableRow> SlateRow = IObjectTableRow::ObjectRowFromUserWidget(Cast<UUserWidget>(UserListEntry.GetObject()));
+	if (SlateRow.IsValid())
+	{
+		return SlateRow->IsItemExpanded();
+	}
+	return false;
+}
+
+UListViewBase* UUserListEntryLibrary::GetOwningListView(TScriptInterface<IUserListEntry> UserListEntry)
+{
+	TSharedPtr<const IObjectTableRow> SlateRow = IObjectTableRow::ObjectRowFromUserWidget(Cast<UUserWidget>(UserListEntry.GetObject()));
+	if (SlateRow.IsValid())
+	{
+		return SlateRow->GetOwningListView();
+	}
+	return nullptr;
+}
diff --git a/Engine/Source/Runtime/UMG/Private/Blueprint/IUserObjectListEntry.cpp b/Engine/Source/Runtime/UMG/Private/Blueprint/IUserObjectListEntry.cpp
--- a/Engine/Source/Runtime/UMG/Private/Blueprint/IUserObjectListEntry.cpp
+++ b/Engine/Source/Runtime/UMG/Private/Blueprint/IUserObjectListEntry.cpp
@@ -2,30 +2,49 @@
 
 #include "Blueprint/IUserObjectListEntry.h"
 #include "Blueprint/UserWidget.h"
+#include "Components/ListView.h"
 
 UUserObjectListEntry::UUserObjectListEntry(const FObjectInitializer& Initializer)
 	: Super(Initializer)
 {
 }
 
-UObject* IUserObjectListEntry::GetListItem(UUserWidget& EntryWidget)
+void IUserObjectListEntry::NativeOnListItemObjectSet(UObject* ListItemObject)
 {
-	if (EntryWidget.Implements<UUserObjectListEntry>())
-	{
-		return Execute_GetListItemObject(&EntryWidget);
-	}
-	return nullptr;
+PRAGMA_DISABLE_DEPRECATION_WARNINGS
+	// Making sure the deprecated path still functions, remove this when the old method is deleted
+	SetListItemObjectInternal(ListItemObject);	
+PRAGMA_ENABLE_DEPRECATION_WARNINGS
+
+	Execute_OnListItemObjectSet(Cast<UObject>(this), ListItemObject);
+}
+
+UObject* IUserObjectListEntry::GetListItemObjectInternal() const
+{
+	return UUserObjectListEntryLibrary::GetListItemObject(Cast<UUserWidget>(const_cast<IUserObjectListEntry*>(this)));
 }
 
 void IUserObjectListEntry::SetListItemObject(UUserWidget& ListEntryWidget, UObject* ListItemObject)
 {
 	if (IUserObjectListEntry* NativeImplementation = Cast<IUserObjectListEntry>(&ListEntryWidget))
 	{
-		NativeImplementation->SetListItemObjectInternal(ListItemObject);
-		Execute_OnListItemObjectSet(&ListEntryWidget, ListItemObject);
+		NativeImplementation->NativeOnListItemObjectSet(ListItemObject);
 	}
 	else if (ListEntryWidget.Implements<UUserObjectListEntry>())
 	{
 		Execute_OnListItemObjectSet(&ListEntryWidget, ListItemObject);
 	}
-}
\ No newline at end of file
+}
+
+UObject* UUserObjectListEntryLibrary::GetListItemObject(TScriptInterface<IUserObjectListEntry> UserObjectListEntry)
+{
+	if (UUserWidget* EntryWidget = Cast<UUserWidget>(UserObjectListEntry.GetObject()))
+	{
+		const UListView* OwningListView = Cast<UListView>(UUserListEntryLibrary::GetOwningListView(EntryWidget));
+		if (UObject* const* ListItem = OwningListView ? OwningListView->ItemFromEntryWidget(*EntryWidget) : nullptr)
+		{
+			return *ListItem;
+		}
+	}
+	return nullptr;
+}
diff --git a/Engine/Source/Runtime/UMG/Private/Components/ListView.cpp b/Engine/Source/Runtime/UMG/Private/Components/ListView.cpp
--- a/Engine/Source/Runtime/UMG/Private/Components/ListView.cpp
+++ b/Engine/Source/Runtime/UMG/Private/Components/ListView.cpp
@@ -84,25 +84,10 @@ void UListView::ClearListItems()
 
 void UListView::SetSelectionMode(TEnumAsByte<ESelectionMode::Type> InSelectionMode)
 {
-	if (InSelectionMode != SelectionMode)
+	SelectionMode = InSelectionMode;
+	if (MyListView)
 	{
-		SelectionMode = InSelectionMode;
-
-		if (InSelectionMode == ESelectionMode::None)
-		{
-			ClearSelection();
-		}
-		else if (InSelectionMode == ESelectionMode::Single || InSelectionMode == ESelectionMode::SingleToggle)
-		{
-			// Try to preserve the last selected item.
-			TArray<UObject*> CurrentlySelectedItems;
-			GetSelectedItems(CurrentlySelectedItems);
-			UObject* const LastSelectedItem = CurrentlySelectedItems.Num() > 0 ? CurrentlySelectedItems.Last(0) : nullptr;
-			if (LastSelectedItem)
-			{
-				SetSelectedItem(LastSelectedItem);
-			}
-		}
+		MyListView->SetSelectionMode(InSelectionMode);
 	}
 }
 
@@ -220,19 +205,19 @@ TSharedRef<STableViewBase> UListView::RebuildListWidget()
 
 void UListView::HandleListEntryHovered(UUserWidget& EntryWidget)
 {
-	if (UObject* ListItem = IUserObjectListEntry::GetListItem(EntryWidget))
+	if (UObject* const* ListItem = ItemFromEntryWidget(EntryWidget))
 	{
-		OnItemIsHoveredChanged().Broadcast(ListItem, true);
-		BP_OnItemIsHoveredChanged.Broadcast(ListItem, true);
+		OnItemIsHoveredChanged().Broadcast(*ListItem, true);
+		BP_OnItemIsHoveredChanged.Broadcast(*ListItem, true);
 	}
 }
 
 void UListView::HandleListEntryUnhovered(UUserWidget& EntryWidget)
 {
-	if (UObject* ListItem = IUserObjectListEntry::GetListItem(EntryWidget))
+	if (UObject* const* ListItem = ItemFromEntryWidget(EntryWidget))
 	{
-		OnItemIsHoveredChanged().Broadcast(ListItem, false);
-		BP_OnItemIsHoveredChanged.Broadcast(ListItem, false);
+		OnItemIsHoveredChanged().Broadcast(*ListItem, false);
+		BP_OnItemIsHoveredChanged.Broadcast(*ListItem, false);
 	}
 }
 
diff --git a/Engine/Source/Runtime/UMG/Public/Blueprint/IUserListEntry.h b/Engine/Source/Runtime/UMG/Public/Blueprint/IUserListEntry.h
--- a/Engine/Source/Runtime/UMG/Public/Blueprint/IUserListEntry.h
+++ b/Engine/Source/Runtime/UMG/Public/Blueprint/IUserListEntry.h
@@ -8,59 +8,51 @@
 class UUserWidget;
 class IObjectTableRow;
 
-UINTERFACE(meta = (CannotImplementInterfaceInBlueprint))
-class UMG_API UNativeUserListEntry : public UInterface
+/**
+ * Required interface for any UUserWidget class to be usable as entry widget in a UListViewBase
+ * Provides access to getters and events for changes to the status of the widget as an entry that represents an item in a list.
+ * @see UListViewBase @see ITypedUMGListView
+ *
+ * Note: To be usable as an entry for UListView, UTileView, or UTreeView, implement IUserObjectListEntry instead
+ */
+UINTERFACE(BlueprintType)
+class UMG_API UUserListEntry : public UInterface
 {
 	GENERATED_UINTERFACE_BODY()
 };
 
-class UMG_API INativeUserListEntry : public IInterface
+class UMG_API IUserListEntry : public IInterface
 {
 	GENERATED_IINTERFACE_BODY()
 
 public:
-	/** Returns true if the item represented by this entry is currently selected. */
-	UFUNCTION(BlueprintCallable, Category = UserListEntry)
-	virtual bool IsListItemSelected() const = 0;
-
+	/** Returns true if the item represented by this entry is currently selected in the owning list view. */
+	bool IsListItemSelected() const;
+	
 	/** Returns true if the item represented by this entry is currently expanded and showing its children. Tree view entries only. */
-	UFUNCTION(BlueprintCallable, Category = UserListEntry)
-	virtual bool IsListItemExpanded() const = 0;
-
+	bool IsListItemExpanded() const;
+	
 	/** Returns the list view that contains this entry. */
-	UFUNCTION(BlueprintCallable, Category = UserListEntry)
-	virtual UListViewBase* GetOwningListView() const = 0;
+	UListViewBase* GetOwningListView() const;
 
-	/** 
-	 * Advanced native-only option for specific rows to preclude themselves from any kind of selection. 
+	/**
+	 * Advanced native-only option for specific rows to preclude themselves from any kind of selection.
 	 * Intended primarily for category separators and the like.
 	 * Note that this is only relevant when the row is in a list that allows selection in the first place.
 	 */
 	virtual bool IsListItemSelectable() const { return true; }
-};
-
-UINTERFACE()
-class UMG_API UUserListEntry : public UNativeUserListEntry
-{
-	GENERATED_UINTERFACE_BODY()
-};
-
-class UMG_API IUserListEntry : public INativeUserListEntry
-{
-	GENERATED_IINTERFACE_BODY()
 
 public:
+	/** Functionality largely for "internal" use/plumbing - see SObjectTableRow's usage. You shouldn't have any cause to call these directly. */
 	static void ReleaseEntry(UUserWidget& ListEntryWidget);
 	static void UpdateItemSelection(UUserWidget& ListEntryWidget, bool bIsSelected);
 	static void UpdateItemExpansion(UUserWidget& ListEntryWidget, bool bIsExpanded);
-	
-	virtual bool IsListItemSelected() const override final;
-	virtual bool IsListItemExpanded() const override final;
-
-	virtual UListViewBase* GetOwningListView() const override final;
 
 protected:
-	/** These follow the same pattern as the NativeOn[X] methods in UUserWidget - super calls are expected in order to route the event to BP. */
+	/** 
+	 * These follow the same pattern as the NativeOn[X] methods in UUserWidget - super calls are expected in order to route the event to BP.
+	 * See the BP events below for descriptions.
+	 */
 	virtual void NativeOnItemSelectionChanged(bool bIsSelected);
 	virtual void NativeOnItemExpansionChanged(bool bIsExpanded);
 	virtual void NativeOnEntryReleased();
@@ -76,4 +68,33 @@ protected:
 	/** Called when this entry is released from the owning table and no longer represents any list item */
 	UFUNCTION(BlueprintImplementableEvent, Category = UserListEntry, meta = (DisplayName = "On Entry Released"))
 	void BP_OnEntryReleased();
+};
+
+/** Static library to supply "for free" functionality to widgets that implement IUserListEntry */
+UCLASS()
+class UMG_API UUserListEntryLibrary : public UBlueprintFunctionLibrary
+{
+	GENERATED_BODY()
+
+public:
+	/** 
+	 * Returns true if the item represented by this entry is currently selected in the owning list view. 
+	 * @param UserListEntry Note: Visually not transmitted, but this defaults to "self". No need to hook up if calling internally.
+	 */
+	UFUNCTION(BlueprintPure, Category = UserListEntry, meta = (DefaultToSelf = UserListEntry))
+	static bool IsListItemSelected(TScriptInterface<IUserListEntry> UserListEntry);
+
+	/** 
+	 * Returns true if the item represented by this entry is currently expanded and showing its children. Tree view entries only.
+	 * @param UserListEntry Note: Visually not transmitted, but this defaults to "self". No need to hook up if calling internally.
+	 */
+	UFUNCTION(BlueprintPure, Category = UserListEntry, meta = (DefaultToSelf = UserListEntry))
+	static bool IsListItemExpanded(TScriptInterface<IUserListEntry> UserListEntry);
+
+	/** 
+	 * Returns the list view that contains this entry.
+	 * @param UserListEntry Note: Visually not transmitted, but this defaults to "self". No need to hook up if calling internally.
+	 */
+	UFUNCTION(BlueprintPure, Category = UserListEntry, meta = (DefaultToSelf = UserListEntry))
+	static UListViewBase* GetOwningListView(TScriptInterface<IUserListEntry> UserListEntry);
 };
\ No newline at end of file
diff --git a/Engine/Source/Runtime/UMG/Public/Blueprint/IUserObjectListEntry.h b/Engine/Source/Runtime/UMG/Public/Blueprint/IUserObjectListEntry.h
--- a/Engine/Source/Runtime/UMG/Public/Blueprint/IUserObjectListEntry.h
+++ b/Engine/Source/Runtime/UMG/Public/Blueprint/IUserObjectListEntry.h
@@ -3,11 +3,12 @@
 #pragma once
 
 #include "IUserListEntry.h"
+#include "Kismet/BlueprintFunctionLibrary.h"
 #include "IUserObjectListEntry.generated.h"
 
 /**
- * Implement for list entry widgets that represent UObject items.
- * Required for a widget to be usable as an entry in UMG lists (ListView, TileView, and TreeView)
+ * Required interface for any UUserWidget class to be usable as entry widget in a stock UMG list view - ListView, TileView, and TreeView
+ * Provides a change event and getter for the object item the entry is assigned to represent by the owning list view (in addition to functionality from IUserListEntry)
  */
 UINTERFACE()
 class UMG_API UUserObjectListEntry : public UUserListEntry
@@ -20,27 +21,46 @@ class UMG_API IUserObjectListEntry : public IUserListEntry
 	GENERATED_IINTERFACE_BODY()
 
 public:
-	static UObject* GetListItem(UUserWidget& EntryWidget);
-
+	/** Returns the UObject item in the owning UListView that this entry currently represents */
 	template <typename ItemObjectT = UObject>
 	ItemObjectT* GetListItem() const
 	{
 		static_assert(TIsDerivedFrom<ItemObjectT, UObject>::IsDerived, "Items represented by an ObjectListEntry are always expected to be UObjects.");
-		return Cast<ItemObjectT>(Execute_GetListItemObject(Cast<UObject>(this)));
+		return Cast<ItemObjectT>(GetListItemObjectInternal());
 	}
 
 protected:
-	virtual void SetListItemObjectInternal(UObject* InObject) {}
+	/** Follows the same pattern as the NativeOn[X] methods in UUserWidget - super calls are expected in order to route the event to BP. */
+	virtual void NativeOnListItemObjectSet(UObject* ListItemObject);
 	
-	/** Returns the item object that this entry currently represents */
-	UFUNCTION(BlueprintNativeEvent, BlueprintCallable, Category = ObjectListEntry, meta = (AllowPrivateAccess = true))
-	UObject* GetListItemObject() const;
-
-	/** Called when the item object represented by this entry is established */
+	/** Called when this entry is assigned a new item object to represent by the owning list view */
 	UFUNCTION(BlueprintImplementableEvent, Category = ObjectListEntry)
 	void OnListItemObjectSet(UObject* ListItemObject);
 
+	UE_DEPRECATED(4.23, "Implementers of IUserObjectListEntry no longer need to manually implement the GetListItemObject method.")
+	virtual UObject* GetListItemObject_Implementation() const { return nullptr; }
+
+	UE_DEPRECATED(4.23, "Renamed to NativeOnListItemObjectSet, as caching the item object is no longer a hard requirement.")
+	virtual void SetListItemObjectInternal(UObject* InObject) {}
+
 private:
+	UObject* GetListItemObjectInternal() const;
+	
 	template <typename> friend class SObjectTableRow;
 	static void SetListItemObject(UUserWidget& ListEntryWidget, UObject* ListItemObject);
+};
+
+/** Static library to supply "for free" functionality to widgets that implement IUserListEntry */
+UCLASS()
+class UMG_API UUserObjectListEntryLibrary : public UBlueprintFunctionLibrary
+{
+	GENERATED_BODY()
+
+public:
+	/** 
+	 * Returns the item in the owning list view that this entry is currently assigned to represent. 
+	 * @param UserObjectListEntry Note: Visually not transmitted, but this defaults to "self". No need to hook up if calling internally.
+	 */
+	UFUNCTION(BlueprintPure, Category = UserObjectListEntry, meta = (DefaultToSelf = UserObjectListEntry))
+	static UObject* GetListItemObject(TScriptInterface<IUserObjectListEntry> UserObjectListEntry);
 };
\ No newline at end of file
