Expose GridSlot's SetNudge function to blueprints

diff --git a/Engine/Source/Runtime/UMG/Public/Components/GridSlot.h b/Engine/Source/Runtime/UMG/Public/Components/GridSlot.h
--- a/Engine/Source/Runtime/UMG/Public/Components/GridSlot.h
+++ b/Engine/Source/Runtime/UMG/Public/Components/GridSlot.h
@@ -84,6 +84,7 @@ public:
 	void SetLayer(int32 InLayer);
 
 	/** Sets the offset for this slot's content by some amount; positive values offset to lower right*/
+	UFUNCTION(BlueprintCallable, Category = "Layout|Grid Slot")
 	void SetNudge(FVector2D InNudge);
 
 	/**  */
