PR #4454: Added ScrollBoxSlot Function (Contributed by [REDACTED_PERSON])

diff --git a/Engine/Source/Runtime/UMG/Private/WidgetLayoutLibrary.cpp b/Engine/Source/Runtime/UMG/Private/WidgetLayoutLibrary.cpp
--- a/Engine/Source/Runtime/UMG/Private/WidgetLayoutLibrary.cpp
+++ b/Engine/Source/Runtime/UMG/Private/WidgetLayoutLibrary.cpp
@@ -9,6 +9,7 @@
 #include "Components/CanvasPanelSlot.h"
 #include "Components/HorizontalBoxSlot.h"
 #include "Components/VerticalBoxSlot.h"
+#include "Components/ScrollBoxSlot.h"
 #include "Components/UniformGridSlot.h"
 #include "Components/GridSlot.h"
 #include "Components/OverlaySlot.h"
@@ -280,6 +281,16 @@ UVerticalBoxSlot* UWidgetLayoutLibrary::SlotAsVerticalBoxSlot(UWidget* Widget)
 	return nullptr;
 }
 
+UScrollBoxSlot* UWidgetLayoutLibrary::SlotAsScrollBoxSlot(UWidget* Widget)
+{
+	if (Widget)
+	{
+		return Cast<UScrollBoxSlot>(Widget->Slot);
+	}
+
+	return nullptr;
+}
+
 void UWidgetLayoutLibrary::RemoveAllWidgets(UObject* WorldContextObject)
 {
 	UWorld* World = GEngine->GetWorldFromContextObject(WorldContextObject, EGetWorldErrorMode::LogAndReturnNull);
diff --git a/Engine/Source/Runtime/UMG/Public/Blueprint/WidgetLayoutLibrary.h b/Engine/Source/Runtime/UMG/Public/Blueprint/WidgetLayoutLibrary.h
--- a/Engine/Source/Runtime/UMG/Public/Blueprint/WidgetLayoutLibrary.h
+++ b/Engine/Source/Runtime/UMG/Public/Blueprint/WidgetLayoutLibrary.h
@@ -15,6 +15,7 @@ class UHorizontalBoxSlot;
 class UOverlaySlot;
 class UUniformGridSlot;
 class UVerticalBoxSlot;
+class UScrollBoxSlot;
 class UWidget;
 
 UCLASS()
@@ -147,6 +148,15 @@ public:
 	UFUNCTION(BlueprintPure, Category = "Slot")
 	static UVerticalBoxSlot* SlotAsVerticalBoxSlot(UWidget* Widget);
 
+
+	/**
+	* Gets the slot object on the child widget as a Scroll Box Slot, allowing you to manipulate its information.
+	* @param Widget The child widget of a Scroll Box.
+	*/
+	UFUNCTION(BlueprintPure, Category = "Slot")
+	static UScrollBoxSlot* SlotAsScrollBoxSlot(UWidget* Widget);
+
+
 	/**
 	 * Removes all widgets from the viewport.
 	 */
