PR #3683: Add GetRenderAngle() function (Contributed by [REDACTED_PERSON])

diff --git a/Engine/Config/BaseEngine.ini b/Engine/Config/BaseEngine.ini
--- a/Engine/Config/BaseEngine.ini
+++ b/Engine/Config/BaseEngine.ini
@@ -644,6 +644,9 @@ bUseStreamingPause=false
 
 +FunctionRedirects=(OldName="PlayerController.ClientPlayForceFeedback",NewName="PlayerController.K2_ClientPlayForceFeedback")
 
+; 4.23
++FunctionRedirects=(OldName="Widget.SetRenderAngle", NewName="Widget.SetRenderTransformAngle")
+
 [CoreUObject.Metadata]
 ; Note: UnrealHeaderTool should be rerun after making changes to MetadataRedirects to catch any keys specified in class headers
 +MetadataRedirects=(OldKey="K2Protected", NewKey="BlueprintProtected")
diff --git a/Engine/Source/Runtime/UMG/Private/Components/Widget.cpp b/Engine/Source/Runtime/UMG/Private/Components/Widget.cpp
--- a/Engine/Source/Runtime/UMG/Private/Components/Widget.cpp
+++ b/Engine/Source/Runtime/UMG/Private/Components/Widget.cpp
@@ -187,12 +187,17 @@ void UWidget::SetRenderShear(FVector2D Shear)
 	UpdateRenderTransform();
 }
 
-void UWidget::SetRenderAngle(float Angle)
+void UWidget::SetRenderTransformAngle(float Angle)
 {
 	RenderTransform.Angle = Angle;
 	UpdateRenderTransform();
 }
 
+float UWidget::GetRenderTransformAngle() const
+{
+	return RenderTransform.Angle;
+}
+
 void UWidget::SetRenderTranslation(FVector2D Translation)
 {
 	RenderTransform.Translation = Translation;
diff --git a/Engine/Source/Runtime/UMG/Public/Components/Widget.h b/Engine/Source/Runtime/UMG/Public/Components/Widget.h
--- a/Engine/Source/Runtime/UMG/Public/Components/Widget.h
+++ b/Engine/Source/Runtime/UMG/Public/Components/Widget.h
@@ -419,7 +419,11 @@ public:
 
 	/** */
 	UFUNCTION(BlueprintCallable, Category="Widget|Transform")
-	void SetRenderAngle(float Angle);
+	void SetRenderTransformAngle(float Angle);
+	
+	/** */
+	UFUNCTION(BlueprintCallable, Category = "Widget|Transform")
+	float GetRenderTransformAngle() const;
 	
 	/** */
 	UFUNCTION(BlueprintCallable, Category="Widget|Transform")
