Don't re-create default STextBlock in UComboBoxString when changing selection

diff --git a/Engine/Source/Runtime/UMG/Private/Components/ComboBoxString.cpp b/Engine/Source/Runtime/UMG/Private/Components/ComboBoxString.cpp
--- a/Engine/Source/Runtime/UMG/Private/Components/ComboBoxString.cpp
+++ b/Engine/Source/Runtime/UMG/Private/Components/ComboBoxString.cpp
@@ -119,7 +119,7 @@ TSharedRef<SWidget> UComboBoxString::RebuildWidget()
 	if ( InitialIndex != -1 )
 	{
 		// Generate the widget for the initially selected widget if needed
-		ComboBoxContent->SetContent(HandleGenerateWidget(CurrentOptionPtr));
+		UpdateOrGenerateWidget(CurrentOptionPtr);
 	}
 
 	return MyComboBox.ToSharedRef();
@@ -220,7 +220,8 @@ void UComboBoxString::SetSelectedOption(FString Option)
 
 void UComboBoxString::SetSelectedIndex(const int32 Index)
 {
-	if (Options.IsValidIndex(Index))
+	// Don't select item if its already selected
+	if (Options.IsValidIndex(Index) && SelectedOption != *CurrentOptionPtr)
 	{
 		CurrentOptionPtr = Options[Index];
 		SelectedOption = *CurrentOptionPtr;
@@ -228,7 +229,7 @@ void UComboBoxString::SetSelectedIndex(const int32 Index)
 		if ( ComboBoxContent.IsValid() )
 		{
 			MyComboBox->SetSelectedItem(CurrentOptionPtr);
-			ComboBoxContent->SetContent(HandleGenerateWidget(CurrentOptionPtr));
+			UpdateOrGenerateWidget(CurrentOptionPtr);
 		}
 	}
 }
@@ -263,6 +264,22 @@ int32 UComboBoxString::GetOptionCount() const
 	return Options.Num();
 }
 
+void UComboBoxString::UpdateOrGenerateWidget(TSharedPtr<FString> Item)
+{
+	// If no custom widget was supplied and the default STextBlock already exists,
+	// just update its text instead of rebuilding the widget.
+	if (DefaultComboBoxContent.IsValid() && (IsDesignTime() || OnGenerateWidgetEvent.IsBound()))
+	{
+		const FString StringItem = Item.IsValid() ? *Item : FString();
+		DefaultComboBoxContent.Pin()->SetText(FText::FromString(StringItem));
+	}
+	else
+	{
+		DefaultComboBoxContent.Reset();
+		ComboBoxContent->SetContent(HandleGenerateWidget(Item));
+	}
+}
+
 TSharedRef<SWidget> UComboBoxString::HandleGenerateWidget(TSharedPtr<FString> Item) const
 {
 	FString StringItem = Item.IsValid() ? *Item : FString();
@@ -291,7 +308,7 @@ void UComboBoxString::HandleSelectionChanged(TSharedPtr<FString> Item, ESelectIn
 	// When the selection changes we always generate another widget to represent the content area of the combobox.
 	if ( ComboBoxContent.IsValid() )
 	{
-		ComboBoxContent->SetContent(HandleGenerateWidget(CurrentOptionPtr));
+		UpdateOrGenerateWidget(CurrentOptionPtr);
 	}
 
 	if ( !IsDesignTime() )
diff --git a/Engine/Source/Runtime/UMG/Public/Components/ComboBoxString.h b/Engine/Source/Runtime/UMG/Public/Components/ComboBoxString.h
--- a/Engine/Source/Runtime/UMG/Public/Components/ComboBoxString.h
+++ b/Engine/Source/Runtime/UMG/Public/Components/ComboBoxString.h
@@ -152,6 +152,9 @@ public:
 #endif
 
 protected:
+	/** Refresh ComboBoxContent with the correct widget/data when the selected option changes */
+	void UpdateOrGenerateWidget(TSharedPtr<FString> Item);
+
 	/** Called by slate when it needs to generate a new item for the combobox */
 	virtual TSharedRef<SWidget> HandleGenerateWidget(TSharedPtr<FString> Item) const;
 
@@ -175,6 +178,9 @@ protected:
 	/** A shared pointer to a container that holds the combobox content that is selected */
 	TSharedPtr< SBox > ComboBoxContent;
 
+	/** If OnGenerateWidgetEvent is not bound, this will store the default STextBlock generated */
+	TWeakPtr<STextBlock> DefaultComboBoxContent;
+
 	/** A shared pointer to the current selected string */
 	TSharedPtr<FString> CurrentOptionPtr;
 };
