Remove 'b' prefix from float variable

diff --git a/Engine/Source/Runtime/UMG/Public/Components/WidgetComponent.h b/Engine/Source/Runtime/UMG/Public/Components/WidgetComponent.h
--- a/Engine/Source/Runtime/UMG/Public/Components/WidgetComponent.h
+++ b/Engine/Source/Runtime/UMG/Public/Components/WidgetComponent.h
@@ -281,7 +281,7 @@ public:
 
 	/**  */
 	UFUNCTION(BlueprintCallable, Category = UserInterface)
-	void SetRedrawTime(float bInRedrawTime) { RedrawTime = bInRedrawTime; }
+	void SetRedrawTime(float InRedrawTime) { RedrawTime = InRedrawTime; }
 
 	/** Get the fake window we create for widgets displayed in the world. */
 	TSharedPtr< SWindow > GetVirtualWindow() const;
