Fix 'b' prefix in variable name

diff --git a/Engine/Source/Runtime/UMG/Public/Components/WidgetComponent.h b/Engine/Source/Runtime/UMG/Public/Components/WidgetComponent.h
--- a/Engine/Source/Runtime/UMG/Public/Components/WidgetComponent.h
+++ b/Engine/Source/Runtime/UMG/Public/Components/WidgetComponent.h
@@ -273,7 +273,7 @@ public:
 
 	/**  */
 	UFUNCTION(BlueprintCallable, Category = UserInterface)
-	void SetDrawAtDesiredSize(bool InbDrawAtDesiredSize) { bDrawAtDesiredSize = InbDrawAtDesiredSize; }
+	void SetDrawAtDesiredSize(bool bInDrawAtDesiredSize) { bDrawAtDesiredSize = bInDrawAtDesiredSize; }
 
 	/**  */
 	UFUNCTION(BlueprintCallable, Category = UserInterface)
