Fix a problem when ComboBox::SetSelectedIndex is called while the CurrentOption is not set yet.

diff --git a/Engine/Source/Runtime/UMG/Private/Components/ComboBoxString.cpp b/Engine/Source/Runtime/UMG/Private/Components/ComboBoxString.cpp
--- a/Engine/Source/Runtime/UMG/Private/Components/ComboBoxString.cpp
+++ b/Engine/Source/Runtime/UMG/Private/Components/ComboBoxString.cpp
@@ -220,20 +220,22 @@ void UComboBoxString::SetSelectedOption(FString Option)
 
 void UComboBoxString::SetSelectedIndex(const int32 Index)
 {
-	// Don't select item if its already selected
-	if (Options.IsValidIndex(Index) && SelectedOption != *CurrentOptionPtr)
+	if (Options.IsValidIndex(Index))
 	{
 		CurrentOptionPtr = Options[Index];
-		SelectedOption = *CurrentOptionPtr;
-
-		if ( ComboBoxContent.IsValid() )
+		// Don't select item if its already selected
+		if (SelectedOption != *CurrentOptionPtr)
 		{
-			MyComboBox->SetSelectedItem(CurrentOptionPtr);
-			UpdateOrGenerateWidget(CurrentOptionPtr);
+			SelectedOption = *CurrentOptionPtr;
+
+			if (ComboBoxContent.IsValid())
+			{
+				MyComboBox->SetSelectedItem(CurrentOptionPtr);
+				UpdateOrGenerateWidget(CurrentOptionPtr);
+			}
 		}
 	}
 }
-
 FString UComboBoxString::GetSelectedOption() const
 {
 	if (CurrentOptionPtr.IsValid())
