Fix events not firing on first frame of UMG animations

diff --git a/Engine/Source/Runtime/UMG/Private/Animation/UMGSequencePlayer.cpp b/Engine/Source/Runtime/UMG/Private/Animation/UMGSequencePlayer.cpp
--- a/Engine/Source/Runtime/UMG/Private/Animation/UMGSequencePlayer.cpp
+++ b/Engine/Source/Runtime/UMG/Private/Animation/UMGSequencePlayer.cpp
@@ -166,16 +166,17 @@ void UUMGSequencePlayer::PlayInternal(double StartAtTime, double EndAtTime, int3
 	NumLoopsCompleted = 0;
 	bIsPlayingForward = InPlayMode != EUMGSequencePlayMode::Reverse;
 
+	PlayerStatus = EMovieScenePlayerStatus::Playing;
+
 	// Immediately evaluate the first frame of the animation so that if tick has already occurred, the widget is setup correctly and ready to be
 	// rendered using the first frames data, otherwise you may see a *pop* due to a widget being constructed with a default different than the
 	// first frame of the animation.
+	// Playback assumes the start frame has already been evaulated, so we also want to evaluate any events on the start frame here.
 	if (RootTemplateInstance.IsValid())
 	{
 		const FMovieSceneContext Context(FMovieSceneEvaluationRange(AbsolutePlaybackStart + TimeCursorPosition, AbsolutePlaybackStart + TimeCursorPosition, AnimationResolution), PlayerStatus);
 		RootTemplateInstance.Evaluate(Context, *this);
 	}
-
-	PlayerStatus = EMovieScenePlayerStatus::Playing;
 }
 
 void UUMGSequencePlayer::Play(float StartAtTime, int32 InNumLoopsToPlay, EUMGSequencePlayMode::Type InPlayMode, float InPlaybackSpeed)
