UMG: Set has jumped when crossing playback boundaries so that audio can reset

diff --git a/Engine/Source/Runtime/UMG/Private/Animation/UMGSequencePlayer.cpp b/Engine/Source/Runtime/UMG/Private/Animation/UMGSequencePlayer.cpp
--- a/Engine/Source/Runtime/UMG/Private/Animation/UMGSequencePlayer.cpp
+++ b/Engine/Source/Runtime/UMG/Private/Animation/UMGSequencePlayer.cpp
@@ -113,7 +113,8 @@ void UUMGSequencePlayer::Tick(float DeltaTime)
 
 			bIsEvaluating = true;
 
-			const FMovieSceneContext Context(FMovieSceneEvaluationRange(AbsolutePlaybackStart + TimeCursorPosition, AbsolutePlaybackStart + LastTimePosition, AnimationResolution), PlayerStatus);
+			FMovieSceneContext Context(FMovieSceneEvaluationRange(AbsolutePlaybackStart + TimeCursorPosition, AbsolutePlaybackStart + LastTimePosition, AnimationResolution), PlayerStatus);
+			Context.SetHasJumped(bCrossedLowerBound || bCrossedUpperBound || bCrossedEndTime);
 			RootTemplateInstance.Evaluate(Context, *this);
 
 			bIsEvaluating = false;
