- adding in support for a UImage to set a texture material param from a TSoftObjectPtr<UTexture2D>
- letting the user pass in a texture to apply while we are loading the TSoftObjectPtr
- letting the user pass in a param name we set to 1 while loading and 0 when finished

diff --git a/Engine/Source/Runtime/UMG/Private/Components/Image.cpp b/Engine/Source/Runtime/UMG/Private/Components/Image.cpp
--- a/Engine/Source/Runtime/UMG/Private/Components/Image.cpp
+++ b/Engine/Source/Runtime/UMG/Private/Components/Image.cpp
@@ -354,6 +354,30 @@ void UImage::SetBrushFromSoftMaterial(TSoftObjectPtr<UMaterialInterface> SoftMat
 	);
 }
 
+void UImage::SetMaterialParamFromSoftTexture(UTexture2D* LoadingTexture, TSoftObjectPtr<UTexture2D> SoftTexture, FName TextureParamName, FName LoadingParamName)
+{
+	if (UMaterialInstanceDynamic* BrushMid = GetDynamicMaterial())
+	{
+		BrushMid->SetTextureParameterValue(TextureParamName, LoadingTexture);
+		BrushMid->SetScalarParameterValue(LoadingParamName, 1.0f);
+
+		TWeakObjectPtr<UImage> WeakThis(this); // using weak ptr in case 'this' has gone out of scope by the time this lambda is called
+
+		RequestAsyncLoad(SoftTexture,
+			[WeakThis, SoftTexture, TextureParamName, LoadingParamName]() {
+			if (UImage* StrongThis = WeakThis.Get())
+			{
+				ensureMsgf(SoftTexture.Get(), TEXT("Failed to load %s"), *SoftTexture.ToSoftObjectPath().ToString());
+				if (UMaterialInstanceDynamic* ImageMID = StrongThis->GetDynamicMaterial())
+				{
+					ImageMID->SetTextureParameterValue(TextureParamName, SoftTexture.Get());
+					ImageMID->SetScalarParameterValue(LoadingParamName, 0.0f);
+				}
+			}
+		});
+	}
+}
+
 UMaterialInstanceDynamic* UImage::GetDynamicMaterial()
 {
 	UMaterialInterface* Material = NULL;
diff --git a/Engine/Source/Runtime/UMG/Public/Components/Image.h b/Engine/Source/Runtime/UMG/Public/Components/Image.h
--- a/Engine/Source/Runtime/UMG/Public/Components/Image.h
+++ b/Engine/Source/Runtime/UMG/Public/Components/Image.h
@@ -115,9 +115,13 @@ public:
 	virtual void SetBrushFromSoftTexture(TSoftObjectPtr<UTexture2D> SoftTexture, bool bMatchSize = false);
 
 	/**  */
-	UFUNCTION(BlueprintCallable, Category="Appearance")
+	UFUNCTION(BlueprintCallable, Category = "Appearance")
 	virtual void SetBrushFromSoftMaterial(TSoftObjectPtr<UMaterialInterface> SoftMaterial);
 
+	/** */
+	UFUNCTION(BlueprintCallable, Category = "Appearance")
+	virtual void SetMaterialParamFromSoftTexture(UTexture2D* LoadingTexture, TSoftObjectPtr<UTexture2D> SoftTexture, FName TextureParamName, FName LoadingParamName);
+
 	/**  */
 	UFUNCTION(BlueprintCallable, Category="Appearance")
 	UMaterialInstanceDynamic* GetDynamicMaterial();
