No longer invalidate volatility on widgets which are animating.  There is no reason to do this anymore and is causing some perf issues

diff --git a/Engine/Source/Runtime/UMG/Private/UserWidget.cpp b/Engine/Source/Runtime/UMG/Private/UserWidget.cpp
--- a/Engine/Source/Runtime/UMG/Private/UserWidget.cpp
+++ b/Engine/Source/Runtime/UMG/Private/UserWidget.cpp
@@ -632,8 +632,6 @@ UUMGSequencePlayer* UUserWidget::PlayAnimation(UWidgetAnimation* InAnimation, fl
 	{
 		Player->Play(StartAtTime, NumberOfLoops, PlayMode, PlaybackSpeed);
 
-		Invalidate(EInvalidateWidget::Volatility);
-
 		OnAnimationStartedPlaying(*Player);
 
 		UpdateCanTick();
@@ -651,8 +649,6 @@ UUMGSequencePlayer* UUserWidget::PlayAnimationTimeRange(UWidgetAnimation* InAnim
 	{
 		Player->PlayTo(StartAtTime, EndAtTime, NumberOfLoops, PlayMode, PlaybackSpeed);
 
-		Invalidate(EInvalidateWidget::Volatility);
-
 		OnAnimationStartedPlaying(*Player);
 
 		UpdateCanTick();
@@ -1513,12 +1509,6 @@ void UUserWidget::TickActionsAndAnimation(const FGeometry& MyGeometry, float InD
 
 	StoppedSequencePlayers.Empty();
 
-	// If we're no longer playing animations invalidate layout so that we recache the volatility of the widget.
-	if ( bWasPlayingAnimation && IsPlayingAnimation() == false )
-	{
-		Invalidate(EInvalidateWidget::Volatility);
-	}
-
 	UWorld* World = GetWorld();
 	if (World)
 	{
