Fix VS 2019 warning

diff --git a/Engine/Source/Runtime/UMG/Private/Slate/SRetainerWidget.cpp b/Engine/Source/Runtime/UMG/Private/Slate/SRetainerWidget.cpp
--- a/Engine/Source/Runtime/UMG/Private/Slate/SRetainerWidget.cpp
+++ b/Engine/Source/Runtime/UMG/Private/Slate/SRetainerWidget.cpp
@@ -541,7 +541,7 @@ bool SRetainerWidget::ComputeVolatility() const
 {
 	// We need to be volatile when global invalidation is turned on so that we have a chance to update
 	// @todo instead of being volatile it may be possible to move this to tick
-	return GSlateEnableGlobalInvalidation;
+	return GSlateEnableGlobalInvalidation != 0;
 }
 
 void SRetainerWidget::OnGlobalInvalidationToggled(bool bGlobalInvalidationEnabled)
