The SWorldWidgetScreenLayer never changes desired size.

diff --git a/Engine/Source/Runtime/UMG/Private/Slate/SWorldWidgetScreenLayer.cpp b/Engine/Source/Runtime/UMG/Private/Slate/SWorldWidgetScreenLayer.cpp
--- a/Engine/Source/Runtime/UMG/Private/Slate/SWorldWidgetScreenLayer.cpp
+++ b/Engine/Source/Runtime/UMG/Private/Slate/SWorldWidgetScreenLayer.cpp
@@ -176,4 +176,9 @@ void SWorldWidgetScreenLayer::RemoveEntryFromCanvas(SWorldWidgetScreenLayer::FCo
 	{
 		Canvas->RemoveSlot(ContainerWidget.ToSharedRef());
 	}
-}
\ No newline at end of file
+}
+
+FVector2D SWorldWidgetScreenLayer::ComputeDesiredSize(float) const
+{
+	return FVector2D(0, 0);
+}
diff --git a/Engine/Source/Runtime/UMG/Public/Slate/SWorldWidgetScreenLayer.h b/Engine/Source/Runtime/UMG/Public/Slate/SWorldWidgetScreenLayer.h
--- a/Engine/Source/Runtime/UMG/Public/Slate/SWorldWidgetScreenLayer.h
+++ b/Engine/Source/Runtime/UMG/Public/Slate/SWorldWidgetScreenLayer.h
@@ -24,6 +24,7 @@ public:
 	void Construct(const FArguments& InArgs, const FLocalPlayerContext& InPlayerContext);
 
 	virtual void Tick(const FGeometry& AllottedGeometry, const double InCurrentTime, const float InDeltaTime) override;
+	virtual FVector2D ComputeDesiredSize(float) const override;
 
 	void SetWidgetDrawSize(FVector2D DrawSize);
 
