Slate - Updating the way we hide the slots for the world/actor layers during the fast path, you can't just early out and not draw anything, items needs to be explictly have their visibility enabled or disabled so that the entire chain of widgets are removed from the screen, similar to the problem seen on culling.

diff --git a/Engine/Source/Runtime/UMG/Private/Slate/SWorldWidgetScreenLayer.cpp b/Engine/Source/Runtime/UMG/Private/Slate/SWorldWidgetScreenLayer.cpp
--- a/Engine/Source/Runtime/UMG/Private/Slate/SWorldWidgetScreenLayer.cpp
+++ b/Engine/Source/Runtime/UMG/Private/Slate/SWorldWidgetScreenLayer.cpp
@@ -163,6 +163,19 @@ void SWorldWidgetScreenLayer::Tick(const FGeometry& AllottedGeometry, const doub
 					continue;
 				}
 			}
+
+			// Done
+			return;
+		}
+	}
+
+	if (GSlateIsOnFastUpdatePath)
+	{
+		// Hide everything if we are unable to do any of the work.
+		for (auto It = ComponentMap.CreateIterator(); It; ++It)
+		{
+			FComponentEntry& Entry = It.Value();
+			Entry.ContainerWidget->SetVisibility(EVisibility::Collapsed);
 		}
 	}
 }
