Fixing UListView incorrectly defaulting to horizontal orientation

diff --git a/Engine/Source/Runtime/UMG/Private/Components/ListView.cpp b/Engine/Source/Runtime/UMG/Private/Components/ListView.cpp
--- a/Engine/Source/Runtime/UMG/Private/Components/ListView.cpp
+++ b/Engine/Source/Runtime/UMG/Private/Components/ListView.cpp
@@ -14,6 +14,7 @@
 
 UListView::UListView(const FObjectInitializer& ObjectInitializer)
 	: Super(ObjectInitializer)
+	, Orientation(EOrientation::Orient_Vertical)
 {
 }
 
