UMG - Adding new functions to the userwidgetpool, to allow releasing just inactive slate resources, and changing out the release slate resources function to now mention that it releases them all, since there's now another function that doesn't release them all.

diff --git a/Engine/Source/Runtime/UMG/Private/Blueprint/UserWidgetPool.cpp b/Engine/Source/Runtime/UMG/Private/Blueprint/UserWidgetPool.cpp
--- a/Engine/Source/Runtime/UMG/Private/Blueprint/UserWidgetPool.cpp
+++ b/Engine/Source/Runtime/UMG/Private/Blueprint/UserWidgetPool.cpp
@@ -32,19 +32,30 @@ void FUserWidgetPool::AddReferencedObjects(FReferenceCollector& Collector)
 
 void FUserWidgetPool::Release(UUserWidget* Widget, bool bReleaseSlate)
 {
-	const int32 ActiveWidgetIdx = ActiveWidgets.Find(Widget);
-	if (ActiveWidgetIdx != INDEX_NONE)
+	if (Widget != nullptr)
 	{
-		InactiveWidgets.Push(Widget);
-		ActiveWidgets.RemoveAt(ActiveWidgetIdx);
-
-		if (bReleaseSlate)
+		const int32 ActiveWidgetIdx = ActiveWidgets.Find(Widget);
+		if (ActiveWidgetIdx != INDEX_NONE)
 		{
-			CachedSlateByWidgetObject.Remove(Widget);
+			InactiveWidgets.Push(Widget);
+			ActiveWidgets.RemoveAt(ActiveWidgetIdx);
+
+			if (bReleaseSlate)
+			{
+				CachedSlateByWidgetObject.Remove(Widget);
+			}
 		}
 	}
 }
 
+void FUserWidgetPool::Release(TArray<UUserWidget*> Widgets, bool bReleaseSlate)
+{
+	for (UUserWidget* Widget : Widgets)
+	{
+		Release(Widget, bReleaseSlate);
+	}
+}
+
 void FUserWidgetPool::ReleaseAll(bool bReleaseSlate)
 {
 	InactiveWidgets.Append(ActiveWidgets);
@@ -63,7 +74,15 @@ void FUserWidgetPool::ResetPool()
 	CachedSlateByWidgetObject.Reset();
 }
 
-void FUserWidgetPool::ReleaseSlateResources()
+void FUserWidgetPool::ReleaseInactiveSlateResources()
+{
+	for (UUserWidget* InactiveWidget : InactiveWidgets)
+	{
+		CachedSlateByWidgetObject.Remove(InactiveWidget);
+	}
+}
+
+void FUserWidgetPool::ReleaseAllSlateResources()
 {
 	CachedSlateByWidgetObject.Reset();
 }
diff --git a/Engine/Source/Runtime/UMG/Private/Components/DynamicEntryBoxBase.cpp b/Engine/Source/Runtime/UMG/Private/Components/DynamicEntryBoxBase.cpp
--- a/Engine/Source/Runtime/UMG/Private/Components/DynamicEntryBoxBase.cpp
+++ b/Engine/Source/Runtime/UMG/Private/Components/DynamicEntryBoxBase.cpp
@@ -20,7 +20,7 @@ void UDynamicEntryBoxBase::ReleaseSlateResources(bool bReleaseChildren)
 {
 	Super::ReleaseSlateResources(bReleaseChildren);
 
-	EntryWidgetPool.ReleaseSlateResources();
+	EntryWidgetPool.ReleaseAllSlateResources();
 	MyPanelWidget.Reset();
 }
 
diff --git a/Engine/Source/Runtime/UMG/Public/Blueprint/UserWidgetPool.h b/Engine/Source/Runtime/UMG/Public/Blueprint/UserWidgetPool.h
--- a/Engine/Source/Runtime/UMG/Public/Blueprint/UserWidgetPool.h
+++ b/Engine/Source/Runtime/UMG/Public/Blueprint/UserWidgetPool.h
@@ -69,14 +69,23 @@ public:
 	/** Return a widget object to the pool, allowing it to be reused in the future */
 	void Release(UUserWidget* Widget, bool bReleaseSlate = false);
 
+	/** Return a widget object to the pool, allowing it to be reused in the future */
+	void Release(TArray<UUserWidget*> Widgets, bool bReleaseSlate = false);
+
 	/** Returns all active widget objects to the inactive pool and optionally destroys all cached underlying slate widgets. */
 	void ReleaseAll(bool bReleaseSlate = false);
 
 	/** Full reset of all created widget objects (and any cached underlying slate) */
 	void ResetPool();
 
+	/** Reset of all cached underlying Slate widgets, only for inactive widgets in the pool. */
+	void ReleaseInactiveSlateResources();
+
+	UE_DEPRECATED(4.24, "Use ReleaseAllSlateResources instead.")
+	void ReleaseSlateResources() { ReleaseAllSlateResources(); }
+
 	/** Reset of all cached underlying Slate widgets, but not the active UUserWidget objects */
-	void ReleaseSlateResources();
+	void ReleaseAllSlateResources();
 
 private:
 	template <typename UserWidgetT = UUserWidget>
