Support changing text justification at runtime on our various UMG text widgets.

diff --git a/Engine/Source/Runtime/Slate/Private/Widgets/Text/SlateEditableTextLayout.cpp b/Engine/Source/Runtime/Slate/Private/Widgets/Text/SlateEditableTextLayout.cpp
--- a/Engine/Source/Runtime/Slate/Private/Widgets/Text/SlateEditableTextLayout.cpp
+++ b/Engine/Source/Runtime/Slate/Private/Widgets/Text/SlateEditableTextLayout.cpp
@@ -481,9 +481,12 @@ void FSlateEditableTextLayout::SetMargin(const TAttribute<FMargin>& InMargin)
 
 void FSlateEditableTextLayout::SetJustification(const TAttribute<ETextJustify::Type>& InJustification)
 {
-	Justification = InJustification;
+	if (!Justification.IdenticalTo(InJustification))
+	{
+		Justification = InJustification;
 
-	OwnerWidget->GetSlateWidget()->Invalidate(EInvalidateWidget::LayoutAndVolatility);
+		OwnerWidget->GetSlateWidget()->Invalidate(EInvalidateWidget::LayoutAndVolatility);
+	}
 }
 
 void FSlateEditableTextLayout::SetLineHeightPercentage(const TAttribute<float>& InLineHeightPercentage)
diff --git a/Engine/Source/Runtime/UMG/Private/Components/EditableText.cpp b/Engine/Source/Runtime/UMG/Private/Components/EditableText.cpp
--- a/Engine/Source/Runtime/UMG/Private/Components/EditableText.cpp
+++ b/Engine/Source/Runtime/UMG/Private/Components/EditableText.cpp
@@ -144,6 +144,15 @@ void UEditableText::SetIsReadOnly(bool InbIsReadyOnly)
 	}
 }
 
+void UEditableText::SetJustification(ETextJustify::Type InJustification)
+{
+	Justification = InJustification;
+	if (MyEditableText.IsValid())
+	{
+		MyEditableText->SetJustification(InJustification);
+	}
+}
+
 void UEditableText::SetClearKeyboardFocusOnCommit(bool bInClearKeyboardFocusOnCommit)
 {
 	ClearKeyboardFocusOnCommit = bInClearKeyboardFocusOnCommit;
diff --git a/Engine/Source/Runtime/UMG/Private/Components/EditableTextBox.cpp b/Engine/Source/Runtime/UMG/Private/Components/EditableTextBox.cpp
--- a/Engine/Source/Runtime/UMG/Private/Components/EditableTextBox.cpp
+++ b/Engine/Source/Runtime/UMG/Private/Components/EditableTextBox.cpp
@@ -178,6 +178,15 @@ bool UEditableTextBox::HasError() const
 	return false;
 }
 
+void UEditableTextBox::SetJustification(ETextJustify::Type InJustification)
+{
+	Justification = InJustification;
+	if (MyEditableTextBlock.IsValid())
+	{
+		MyEditableTextBlock->SetJustification(InJustification);
+	}
+}
+
 void UEditableTextBox::HandleOnTextChanged(const FText& InText)
 {
 	Text = InText;
diff --git a/Engine/Source/Runtime/UMG/Private/Components/MultiLineEditableText.cpp b/Engine/Source/Runtime/UMG/Private/Components/MultiLineEditableText.cpp
--- a/Engine/Source/Runtime/UMG/Private/Components/MultiLineEditableText.cpp
+++ b/Engine/Source/Runtime/UMG/Private/Components/MultiLineEditableText.cpp
@@ -103,6 +103,17 @@ void UMultiLineEditableText::SynchronizeProperties()
 	Super::SynchronizeTextLayoutProperties(*MyMultiLineEditableText);
 }
 
+
+void UMultiLineEditableText::SetJustification(ETextJustify::Type InJustification)
+{
+	Super::SetJustification(InJustification);
+
+	if (MyMultiLineEditableText.IsValid())
+	{
+		MyMultiLineEditableText->SetJustification(InJustification);
+	}
+}
+
 FText UMultiLineEditableText::GetText() const
 {
 	if ( MyMultiLineEditableText.IsValid() )
diff --git a/Engine/Source/Runtime/UMG/Private/Components/MultiLineEditableTextBox.cpp b/Engine/Source/Runtime/UMG/Private/Components/MultiLineEditableTextBox.cpp
--- a/Engine/Source/Runtime/UMG/Private/Components/MultiLineEditableTextBox.cpp
+++ b/Engine/Source/Runtime/UMG/Private/Components/MultiLineEditableTextBox.cpp
@@ -110,6 +110,16 @@ void UMultiLineEditableTextBox::SynchronizeProperties()
 	Super::SynchronizeTextLayoutProperties(*MyEditableTextBlock);
 }
 
+void UMultiLineEditableTextBox::SetJustification(ETextJustify::Type InJustification)
+{
+	Super::SetJustification(InJustification);
+
+	if (MyEditableTextBlock.IsValid())
+	{
+		MyEditableTextBlock->SetJustification(InJustification);
+	}
+}
+
 FText UMultiLineEditableTextBox::GetText() const
 {
 	if ( MyEditableTextBlock.IsValid() )
diff --git a/Engine/Source/Runtime/UMG/Private/Components/RichTextBlock.cpp b/Engine/Source/Runtime/UMG/Private/Components/RichTextBlock.cpp
--- a/Engine/Source/Runtime/UMG/Private/Components/RichTextBlock.cpp
+++ b/Engine/Source/Runtime/UMG/Private/Components/RichTextBlock.cpp
@@ -302,7 +302,7 @@ void URichTextBlock::SetDefaultStrikeBrush(FSlateBrush& InStrikeBrush)
 
 void URichTextBlock::SetJustification(ETextJustify::Type InJustification)
 {
-	Justification = InJustification;
+	Super::SetJustification(InJustification);
 	if (MyRichTextBlock.IsValid())
 	{
 		MyRichTextBlock->SetJustification(InJustification);
diff --git a/Engine/Source/Runtime/UMG/Private/Components/TextBlock.cpp b/Engine/Source/Runtime/UMG/Private/Components/TextBlock.cpp
--- a/Engine/Source/Runtime/UMG/Private/Components/TextBlock.cpp
+++ b/Engine/Source/Runtime/UMG/Private/Components/TextBlock.cpp
@@ -106,12 +106,12 @@ void UTextBlock::SetStrikeBrush(FSlateBrush InStrikeBrush)
 	}
 }
 
-void UTextBlock::SetJustification( ETextJustify::Type InJustification )
+void UTextBlock::SetJustification(ETextJustify::Type InJustification)
 {
-	Justification = InJustification;
-	if ( MyTextBlock.IsValid() )
+	Super::SetJustification(InJustification);
+	if (MyTextBlock.IsValid())
 	{
-		MyTextBlock->SetJustification( Justification );
+		MyTextBlock->SetJustification(InJustification);
 	}
 }
 
diff --git a/Engine/Source/Runtime/UMG/Public/Components/EditableText.h b/Engine/Source/Runtime/UMG/Public/Components/EditableText.h
--- a/Engine/Source/Runtime/UMG/Public/Components/EditableText.h
+++ b/Engine/Source/Runtime/UMG/Public/Components/EditableText.h
@@ -126,7 +126,7 @@ public:
 	EVirtualKeyboardDismissAction VirtualKeyboardDismissAction;
 	
 	/** How the text should be aligned with the margin. */
-	UPROPERTY(EditAnywhere, BlueprintReadWrite, Category=Appearance)
+	UPROPERTY(EditAnywhere, BlueprintReadWrite, BlueprintSetter=SetJustification, Category=Appearance)
 	TEnumAsByte<ETextJustify::Type> Justification;
 
 	/** Controls how the text within this widget should be shaped. */
@@ -167,6 +167,9 @@ public:
 	UFUNCTION(BlueprintCallable, Category = "Widget", meta=(DisplayName="SetIsReadOnly (Editable Text)"))
 	void SetIsReadOnly(UPARAM(DisplayName="ReadyOnly") bool InbIsReadyOnly);
 
+	UFUNCTION(BlueprintSetter)
+	void SetJustification(ETextJustify::Type InJustification);
+
 	void SetClearKeyboardFocusOnCommit(bool bInClearKeyboardFocusOnCommit);
 
 public:
diff --git a/Engine/Source/Runtime/UMG/Public/Components/EditableTextBox.h b/Engine/Source/Runtime/UMG/Public/Components/EditableTextBox.h
--- a/Engine/Source/Runtime/UMG/Public/Components/EditableTextBox.h
+++ b/Engine/Source/Runtime/UMG/Public/Components/EditableTextBox.h
@@ -128,7 +128,7 @@ public:
 	EVirtualKeyboardDismissAction VirtualKeyboardDismissAction;
 	
 	/** How the text should be aligned with the margin. */
-	UPROPERTY(EditAnywhere, BlueprintReadWrite, Category=Appearance)
+	UPROPERTY(EditAnywhere, BlueprintReadWrite, BlueprintSetter = SetJustification, Category=Appearance)
 	TEnumAsByte<ETextJustify::Type> Justification;
 
 	/** Controls how the text within this widget should be shaped. */
@@ -176,6 +176,9 @@ public:
 	UFUNCTION(BlueprintCallable, Category = "Widget")
 	bool HasError() const;
 
+	UFUNCTION(BlueprintSetter)
+	void SetJustification(ETextJustify::Type InJustification);
+
 	//~ Begin UWidget Interface
 	virtual void SynchronizeProperties() override;
 	//~ End UWidget Interface
diff --git a/Engine/Source/Runtime/UMG/Public/Components/MultiLineEditableText.h b/Engine/Source/Runtime/UMG/Public/Components/MultiLineEditableText.h
--- a/Engine/Source/Runtime/UMG/Public/Components/MultiLineEditableText.h
+++ b/Engine/Source/Runtime/UMG/Public/Components/MultiLineEditableText.h
@@ -133,6 +133,8 @@ protected:
 	virtual TSharedRef<SWidget> RebuildWidget() override;
 	// End of UWidget
 
+	virtual void SetJustification(ETextJustify::Type InJustification) override;
+
 	void HandleOnTextChanged(const FText& Text);
 	void HandleOnTextCommitted(const FText& Text, ETextCommit::Type CommitMethod);
 
diff --git a/Engine/Source/Runtime/UMG/Public/Components/MultiLineEditableTextBox.h b/Engine/Source/Runtime/UMG/Public/Components/MultiLineEditableTextBox.h
--- a/Engine/Source/Runtime/UMG/Public/Components/MultiLineEditableTextBox.h
+++ b/Engine/Source/Runtime/UMG/Public/Components/MultiLineEditableTextBox.h
@@ -152,6 +152,8 @@ protected:
 	virtual TSharedRef<SWidget> RebuildWidget() override;
 	// End of UWidget
 
+	virtual void SetJustification(ETextJustify::Type InJustification) override;
+
 	void HandleOnTextChanged(const FText& Text);
 	void HandleOnTextCommitted(const FText& Text, ETextCommit::Type CommitMethod);
 
diff --git a/Engine/Source/Runtime/UMG/Public/Components/RichTextBlock.h b/Engine/Source/Runtime/UMG/Public/Components/RichTextBlock.h
--- a/Engine/Source/Runtime/UMG/Public/Components/RichTextBlock.h
+++ b/Engine/Source/Runtime/UMG/Public/Components/RichTextBlock.h
@@ -74,13 +74,6 @@ public:
 	UFUNCTION(BlueprintCallable, Category = "Appearance")
 	void SetDefaultStrikeBrush(FSlateBrush& InStrikeBrush);
 
-	/**
-	 *  Set the text justification for this rich text block
-	 *  @param InJustification new justification
-	 */
-	UFUNCTION(BlueprintCallable, Category = "Appearance")
-	void SetJustification(ETextJustify::Type InJustification);
-	
 	/**
 	 *  Set the minimum desired width for this rich text block
 	 *  @param InMinDesiredWidth new minimum desired width
@@ -151,6 +144,8 @@ public:
 
 protected:
 	virtual TSharedRef<SWidget> RebuildWidget() override;
+
+	virtual void SetJustification(ETextJustify::Type InJustification) override;
 	
 	virtual void UpdateStyleData();
 	void RebuildStyleInstance();
diff --git a/Engine/Source/Runtime/UMG/Public/Components/TextBlock.h b/Engine/Source/Runtime/UMG/Public/Components/TextBlock.h
--- a/Engine/Source/Runtime/UMG/Public/Components/TextBlock.h
+++ b/Engine/Source/Runtime/UMG/Public/Components/TextBlock.h
@@ -73,14 +73,6 @@ public:
 	UFUNCTION(BlueprintCallable, Category = "Appearance")
 	void SetStrikeBrush(FSlateBrush InStrikeBrush);
 
-	/**
-	 *  Set the text justification for this text block
-	 *
-	 *  @param InJustification new justification
-	 */
-	UFUNCTION(BlueprintCallable, Category = "Appearance")
-	void SetJustification(ETextJustify::Type InJustification);
-	
 	/**
 	 *  Set the minimum desired width for this text block
 	 *
@@ -207,6 +199,8 @@ protected:
 	virtual void OnBindingChanged(const FName& Property) override;
 	//~ End UWidget Interface
 
+	virtual void SetJustification(ETextJustify::Type InJustification) override;
+
 	/** Get the text that should be displayed in the internal Slate widget (allows flags to mutate the display text without modifying the persistent designer property data) */
 	virtual TAttribute<FText> GetDisplayText();
 
diff --git a/Engine/Source/Runtime/UMG/Public/Components/TextWidgetTypes.h b/Engine/Source/Runtime/UMG/Public/Components/TextWidgetTypes.h
--- a/Engine/Source/Runtime/UMG/Public/Components/TextWidgetTypes.h
+++ b/Engine/Source/Runtime/UMG/Public/Components/TextWidgetTypes.h
@@ -82,7 +82,7 @@ protected:
 	FShapedTextOptions ShapedTextOptions;
 
 	/** How the text should be aligned with the margin. */
-	UPROPERTY(EditAnywhere, BlueprintReadWrite, Category=Appearance)
+	UPROPERTY(EditAnywhere, BlueprintReadWrite, BlueprintSetter=SetJustification, Category=Appearance)
 	TEnumAsByte<ETextJustify::Type> Justification;
 
 	/** The wrapping policy to use. */
@@ -104,4 +104,7 @@ protected:
 	/** The amount to scale each lines height by. */
 	UPROPERTY(EditAnywhere, BlueprintReadOnly, Category=Appearance, AdvancedDisplay)
 	float LineHeightPercentage;
+
+	UFUNCTION(BlueprintSetter)
+	virtual void SetJustification(ETextJustify::Type InJustification) { Justification = InJustification; }
 };
