Invalidate prepass during UMG animations as animations on widgets can cause desired size impacting effects to their children

diff --git a/Engine/Source/Runtime/UMG/Private/UserWidget.cpp b/Engine/Source/Runtime/UMG/Private/UserWidget.cpp
--- a/Engine/Source/Runtime/UMG/Private/UserWidget.cpp
+++ b/Engine/Source/Runtime/UMG/Private/UserWidget.cpp
@@ -1500,6 +1500,14 @@ void UUserWidget::TickActionsAndAnimation(const FGeometry& MyGeometry, float InD
 	}
 
 	const bool bWasPlayingAnimation = IsPlayingAnimation();
+	if(bWasPlayingAnimation)
+	{ 
+		TSharedPtr<SWidget> CachedWidget = GetCachedWidget();
+		if (CachedWidget.IsValid())
+		{
+			CachedWidget->InvalidatePrepass();
+		}
+	}
 
 	// The process of ticking the players above can stop them so we remove them after all players have ticked
 	for ( UUMGSequencePlayer* StoppedPlayer : StoppedSequencePlayers )
