Fix OrionEditor build (for real this time)

diff --git a/Engine/Source/Runtime/UMG/Public/Components/MultiLineEditableText.h b/Engine/Source/Runtime/UMG/Public/Components/MultiLineEditableText.h
--- a/Engine/Source/Runtime/UMG/Public/Components/MultiLineEditableText.h
+++ b/Engine/Source/Runtime/UMG/Public/Components/MultiLineEditableText.h
@@ -111,7 +111,11 @@ public:
 
 	UFUNCTION(BlueprintSetter)
 	void SetWidgetStyle(const FTextBlockStyle& InWidgetStyle);
-	
+
+	//~ Begin UTextLayoutWidget Interface
+	virtual void SetJustification(ETextJustify::Type InJustification) override;
+	//~ End UTextLayoutWidget Interface
+
 	//~ Begin UWidget Interface
 	virtual void SynchronizeProperties() override;
 	//~ End UWidget Interface
@@ -133,8 +137,6 @@ protected:
 	virtual TSharedRef<SWidget> RebuildWidget() override;
 	// End of UWidget
 
-	virtual void SetJustification(ETextJustify::Type InJustification) override;
-
 	void HandleOnTextChanged(const FText& Text);
 	void HandleOnTextCommitted(const FText& Text, ETextCommit::Type CommitMethod);
 
diff --git a/Engine/Source/Runtime/UMG/Public/Components/MultiLineEditableTextBox.h b/Engine/Source/Runtime/UMG/Public/Components/MultiLineEditableTextBox.h
--- a/Engine/Source/Runtime/UMG/Public/Components/MultiLineEditableTextBox.h
+++ b/Engine/Source/Runtime/UMG/Public/Components/MultiLineEditableTextBox.h
@@ -127,9 +127,11 @@ public:
 	//TODO UMG Add Set BackgroundColor
 	//TODO UMG Add Set ForegroundColor
 	//TODO UMG Add Set Font
-	//TODO UMG Add Set Justification
 
 public:
+	//~ Begin UTextLayoutWidget Interface
+	virtual void SetJustification(ETextJustify::Type InJustification) override;
+	//~ End UTextLayoutWidget Interface
 
 	//~ Begin UWidget Interface
 	virtual void SynchronizeProperties() override;
@@ -152,8 +154,6 @@ protected:
 	virtual TSharedRef<SWidget> RebuildWidget() override;
 	// End of UWidget
 
-	virtual void SetJustification(ETextJustify::Type InJustification) override;
-
 	void HandleOnTextChanged(const FText& Text);
 	void HandleOnTextCommitted(const FText& Text, ETextCommit::Type CommitMethod);
 
diff --git a/Engine/Source/Runtime/UMG/Public/Components/RichTextBlock.h b/Engine/Source/Runtime/UMG/Public/Components/RichTextBlock.h
--- a/Engine/Source/Runtime/UMG/Public/Components/RichTextBlock.h
+++ b/Engine/Source/Runtime/UMG/Public/Components/RichTextBlock.h
@@ -101,7 +101,11 @@ public:
 
 public:
 	URichTextBlock(const FObjectInitializer& ObjectInitializer);
-	
+
+	//~ Begin UTextLayoutWidget Interface
+	virtual void SetJustification(ETextJustify::Type InJustification) override;
+	//~ End UTextLayoutWidget Interface
+
 	// UWidget interface
 	virtual void SynchronizeProperties() override;
 	// End of UWidget interface
@@ -144,8 +148,6 @@ public:
 
 protected:
 	virtual TSharedRef<SWidget> RebuildWidget() override;
-
-	virtual void SetJustification(ETextJustify::Type InJustification) override;
 	
 	virtual void UpdateStyleData();
 	void RebuildStyleInstance();
diff --git a/Engine/Source/Runtime/UMG/Public/Components/TextBlock.h b/Engine/Source/Runtime/UMG/Public/Components/TextBlock.h
--- a/Engine/Source/Runtime/UMG/Public/Components/TextBlock.h
+++ b/Engine/Source/Runtime/UMG/Public/Components/TextBlock.h
@@ -161,6 +161,10 @@ public:
 	UFUNCTION(BlueprintCallable, Category="Widget", meta=(DisplayName="SetText (Text)"))
 	virtual void SetText(FText InText);
 
+	//~ Begin UTextLayoutWidget Interface
+	virtual void SetJustification(ETextJustify::Type InJustification) override;
+	//~ End UTextLayoutWidget Interface
+
 	//~ Begin UWidget Interface
 	virtual void SynchronizeProperties() override;
 	//~ End UWidget Interface
@@ -199,8 +203,6 @@ protected:
 	virtual void OnBindingChanged(const FName& Property) override;
 	//~ End UWidget Interface
 
-	virtual void SetJustification(ETextJustify::Type InJustification) override;
-
 	/** Get the text that should be displayed in the internal Slate widget (allows flags to mutate the display text without modifying the persistent designer property data) */
 	virtual TAttribute<FText> GetDisplayText();
 
