PR #5891: Added missing Widget Layout Library helpers for casting to a specific slot type

diff --git a/Engine/Source/Runtime/UMG/Private/WidgetLayoutLibrary.cpp b/Engine/Source/Runtime/UMG/Private/WidgetLayoutLibrary.cpp
--- a/Engine/Source/Runtime/UMG/Private/WidgetLayoutLibrary.cpp
+++ b/Engine/Source/Runtime/UMG/Private/WidgetLayoutLibrary.cpp
@@ -14,6 +14,12 @@
 #include "Components/GridSlot.h"
 #include "Components/OverlaySlot.h"
 #include "Components/BorderSlot.h"
+#include "Components/SafeZoneSlot.h"
+#include "Components/ScaleBoxSlot.h"
+#include "Components/ScrollBoxSlot.h"
+#include "Components/SizeBoxSlot.h"
+#include "Components/WrapBoxSlot.h"
+#include "Components/WidgetSwitcherSlot.h"
 #include "Engine/UserInterfaceSettings.h"
 #include "Slate/SlateBrushAsset.h"
 #include "Blueprint/WidgetLayoutLibrary.h"
@@ -291,6 +297,56 @@ UScrollBoxSlot* UWidgetLayoutLibrary::SlotAsScrollBoxSlot(UWidget* Widget)
 	return nullptr;
 }
 
+USafeZoneSlot* UWidgetLayoutLibrary::SlotAsSafeBoxSlot(UWidget* Widget)
+{
+	if (Widget)
+	{
+		return Cast<USafeZoneSlot>(Widget->Slot);
+	}
+
+	return nullptr;
+}
+
+UScaleBoxSlot* UWidgetLayoutLibrary::SlotAsScaleBoxSlot(UWidget* Widget)
+{
+	if (Widget)
+	{
+		return Cast<UScaleBoxSlot>(Widget->Slot);
+	}
+
+	return nullptr;
+}
+
+USizeBoxSlot* UWidgetLayoutLibrary::SlotAsSizeBoxSlot(UWidget* Widget)
+{
+	if (Widget)
+	{
+		return Cast<USizeBoxSlot>(Widget->Slot);
+	}
+
+	return nullptr;
+}
+
+UWrapBoxSlot* UWidgetLayoutLibrary::SlotAsWrapBoxSlot(UWidget* Widget)
+{
+	if (Widget)
+	{
+		return Cast<UWrapBoxSlot>(Widget->Slot);
+	}
+
+	return nullptr;
+}
+
+UWidgetSwitcherSlot* UWidgetLayoutLibrary::SlotAsWidgetSwitcherSlot(UWidget* Widget)
+{
+	if (Widget)
+	{
+		return Cast<UWidgetSwitcherSlot>(Widget->Slot);
+	}
+
+	return nullptr;
+}
+
 void UWidgetLayoutLibrary::RemoveAllWidgets(UObject* WorldContextObject)
 {
 	UWorld* World = GEngine->GetWorldFromContextObject(WorldContextObject, EGetWorldErrorMode::LogAndReturnNull);
diff --git a/Engine/Source/Runtime/UMG/Public/Blueprint/WidgetLayoutLibrary.h b/Engine/Source/Runtime/UMG/Public/Blueprint/WidgetLayoutLibrary.h
--- a/Engine/Source/Runtime/UMG/Public/Blueprint/WidgetLayoutLibrary.h
+++ b/Engine/Source/Runtime/UMG/Public/Blueprint/WidgetLayoutLibrary.h
@@ -16,6 +16,11 @@ class UOverlaySlot;
 class UUniformGridSlot;
 class UVerticalBoxSlot;
 class UScrollBoxSlot;
+class USafeZoneSlot;
+class UScaleBoxSlot;
+class USizeBoxSlot;
+class UWrapBoxSlot;
+class UWidgetSwitcherSlot;
 class UWidget;
 
 UCLASS()
@@ -156,6 +161,40 @@ public:
 	UFUNCTION(BlueprintPure, Category = "Slot")
 	static UScrollBoxSlot* SlotAsScrollBoxSlot(UWidget* Widget);
 
+	/**
+	 * Gets the slot object on the child widget as a Safe Box Slot, allowing you to manipulate its information.
+	 * @param Widget The child widget of a Safe Box.
+	 */
+	UFUNCTION(BlueprintPure, Category = "Slot")
+	static USafeZoneSlot* SlotAsSafeBoxSlot(UWidget* Widget);
+
+	/**
+	 * Gets the slot object on the child widget as a Scale Box Slot, allowing you to manipulate its information.
+	 * @param Widget The child widget of a Scale Box.
+	 */
+	UFUNCTION(BlueprintPure, Category = "Slot")
+	static UScaleBoxSlot* SlotAsScaleBoxSlot(UWidget* Widget);
+
+	/**
+	 * Gets the slot object on the child widget as a Size Box Slot, allowing you to manipulate its information.
+	 * @param Widget The child widget of a Size Box.
+	 */
+	UFUNCTION(BlueprintPure, Category = "Slot")
+	static USizeBoxSlot* SlotAsSizeBoxSlot(UWidget* Widget);
+
+	/**
+	 * Gets the slot object on the child widget as a Wrap Box Slot, allowing you to manipulate its information.
+	 * @param Widget The child widget of a Wrap Box.
+	 */
+	UFUNCTION(BlueprintPure, Category = "Slot")
+	static UWrapBoxSlot* SlotAsWrapBoxSlot(UWidget* Widget);
+
+	/**
+	 * Gets the slot object on the child widget as a Widget Switcher Slot, allowing you to manipulate its information.
+	 * @param Widget The child widget of a Widget Switcher Slot.
+	 */
+	UFUNCTION(BlueprintPure, Category = "Slot")
+	static UWidgetSwitcherSlot* SlotAsWidgetSwitcherSlot(UWidget* Widget);
 
 	/**
 	 * Removes all widgets from the viewport.
