Check that CreateWidget class param is not null even in shipping mode so that at least it doesn't crash in the engine layer.

diff --git a/Engine/Source/Runtime/UMG/Private/UserWidget.cpp b/Engine/Source/Runtime/UMG/Private/UserWidget.cpp
--- a/Engine/Source/Runtime/UMG/Private/UserWidget.cpp
+++ b/Engine/Source/Runtime/UMG/Private/UserWidget.cpp
@@ -2136,6 +2136,12 @@ UUserWidget* UUserWidget::CreateInstanceInternal(UObject* Outer, TSubclassOf<UUs
 	{
 		return nullptr;
 	}
+#else
+	if (!UserWidgetClass)
+	{
+		UE_LOG(LogUMG, Error, TEXT("CreateWidget called with a null class."));
+		return nullptr;
+	}
 #endif
 
 #if !UE_BUILD_SHIPPING
