Sequencer: Fix localized GetMaskByName

diff --git a/Engine/Source/Runtime/MovieSceneTracks/Private/Sections/MovieScene3DTransformSection.cpp b/Engine/Source/Runtime/MovieSceneTracks/Private/Sections/MovieScene3DTransformSection.cpp
--- a/Engine/Source/Runtime/MovieSceneTracks/Private/Sections/MovieScene3DTransformSection.cpp
+++ b/Engine/Source/Runtime/MovieSceneTracks/Private/Sections/MovieScene3DTransformSection.cpp
@@ -465,7 +465,7 @@ void UMovieScene3DTransformSection::SetMask(FMovieSceneTransformMask NewMask)
 
 FMovieSceneTransformMask UMovieScene3DTransformSection::GetMaskByName(const FName& InName) const
 {
-	if (InName == TEXT("Location"))
+	if (InName.ToString() == NSLOCTEXT("MovieSceneTransformSection", "Location", "Location").ToString())
 	{
 		return EMovieSceneTransformChannel::Translation;
 	}
@@ -481,7 +481,7 @@ FMovieSceneTransformMask UMovieScene3DTransformSection::GetMaskByName(const FNam
 	{
 		return EMovieSceneTransformChannel::TranslationZ;
 	}
-	else if (InName == TEXT("Rotation"))
+	else if (InName.ToString() == NSLOCTEXT("MovieSceneTransformSection", "Rotation", "Rotation").ToString())
 	{
 		return EMovieSceneTransformChannel::Rotation;
 	}
@@ -497,7 +497,7 @@ FMovieSceneTransformMask UMovieScene3DTransformSection::GetMaskByName(const FNam
 	{
 		return EMovieSceneTransformChannel::RotationZ;
 	}
-	else if (InName == TEXT("Scale"))
+	else if (InName.ToString() == NSLOCTEXT("MovieSceneTransformSection", "Scale", "Scale").ToString())
 	{
 		return EMovieSceneTransformChannel::Scale;
 	}
diff --git a/Engine/Source/Runtime/UMG/Private/Animation/MovieScene2DTransformSection.cpp b/Engine/Source/Runtime/UMG/Private/Animation/MovieScene2DTransformSection.cpp
--- a/Engine/Source/Runtime/UMG/Private/Animation/MovieScene2DTransformSection.cpp
+++ b/Engine/Source/Runtime/UMG/Private/Animation/MovieScene2DTransformSection.cpp
@@ -293,7 +293,7 @@ void UMovieScene2DTransformSection::SetMask(FMovieScene2DTransformMask NewMask)
 
 FMovieScene2DTransformMask UMovieScene2DTransformSection::GetMaskByName(const FName& InName) const
 {
-	if (InName == TEXT("Translation"))
+	if (InName.ToString() == NSLOCTEXT("MovieScene2DTransformSection", "Translation", "Translation").ToString())
 	{
 		return EMovieScene2DTransformChannel::Translation;
 	}
@@ -305,11 +305,12 @@ FMovieScene2DTransformMask UMovieScene2DTransformSection::GetMaskByName(const FN
 	{
 		return EMovieScene2DTransformChannel::TranslationY;
 	}
-	else if (InName == TEXT("Angle") || InName == TEXT("Rotation"))
+	else if (InName.ToString() == NSLOCTEXT("MovieScene2DTransformSection", "Angle", "Angle").ToString() ||
+			InName.ToString() == NSLOCTEXT("MovieScene2DTransformSection", "Rotation", "Rotation").ToString())
 	{
 		return EMovieScene2DTransformChannel::Rotation;
 	}
-	else if (InName == TEXT("Scale"))
+	else if (InName.ToString() == NSLOCTEXT("MovieScene2DTransformSection", "Scale", "Scale").ToString())
 	{
 		return EMovieScene2DTransformChannel::Scale;
 	}
@@ -321,7 +322,7 @@ FMovieScene2DTransformMask UMovieScene2DTransformSection::GetMaskByName(const FN
 	{
 		return EMovieScene2DTransformChannel::ScaleY;
 	}
-	else if (InName == TEXT("Shear"))
+	else if (InName.ToString() == NSLOCTEXT("MovieScene2DTransformSection", "Shear", "Shear").ToString())
 	{
 		return EMovieScene2DTransformChannel::Shear;
 	}
