Fix crash in slate renderer due to a bad FRenderTarget* cast.

diff --git a/Engine/Source/Runtime/UMG/Private/Slate/WidgetRenderer.cpp b/Engine/Source/Runtime/UMG/Private/Slate/WidgetRenderer.cpp
--- a/Engine/Source/Runtime/UMG/Private/Slate/WidgetRenderer.cpp
+++ b/Engine/Source/Runtime/UMG/Private/Slate/WidgetRenderer.cpp
@@ -311,7 +311,7 @@ bool FWidgetRenderer::DrawInvalidationRoot(TSharedRef<SVirtualWindow>& VirtualWi
 				static_cast<float>(FApp::GetCurrentTime() - GStartTime),
 				static_cast<float>(FApp::GetDeltaTime()),
 				static_cast<float>(FPlatformTime::Seconds() - GStartTime),
-				RenderTarget->GameThread_GetRenderTargetResource(),
+				(FRenderTarget*)RenderTarget->GameThread_GetRenderTargetResource(),
 				Renderer.Get(),
 				bClearTarget
 			};
