PR #5932: Add IsOpen function to Combo Box

diff --git a/Engine/Source/Runtime/UMG/Private/Components/ComboBoxString.cpp b/Engine/Source/Runtime/UMG/Private/Components/ComboBoxString.cpp
--- a/Engine/Source/Runtime/UMG/Private/Components/ComboBoxString.cpp
+++ b/Engine/Source/Runtime/UMG/Private/Components/ComboBoxString.cpp
@@ -290,6 +290,11 @@ int32 UComboBoxString::GetOptionCount() const
 	return Options.Num();
 }
 
+bool UComboBoxString::IsOpen() const
+{
+	return MyComboBox.IsValid() && MyComboBox->IsOpen();
+}
+
 void UComboBoxString::UpdateOrGenerateWidget(TSharedPtr<FString> Item)
 {
 	// If no custom widget was supplied and the default STextBlock already exists,
diff --git a/Engine/Source/Runtime/UMG/Public/Components/ComboBoxString.h b/Engine/Source/Runtime/UMG/Public/Components/ComboBoxString.h
--- a/Engine/Source/Runtime/UMG/Public/Components/ComboBoxString.h
+++ b/Engine/Source/Runtime/UMG/Public/Components/ComboBoxString.h
@@ -136,6 +136,9 @@ public:
 	/** Returns the number of options */
 	UFUNCTION(BlueprintCallable, Category="ComboBox")
 	int32 GetOptionCount() const;
+	
+	UFUNCTION(BlueprintCallable, Category="ComboBox", Meta = (ReturnDisplayName = "bOpen"))
+	bool IsOpen() const;
 
 	//~ Begin UVisual Interface
 	virtual void ReleaseSlateResources(bool bReleaseChildren) override;
