UMG - Listview : Allow any desired SObjectPadding to be treated correctly in SObjectTableRows.  We now combine the desired padding from the CDO, with any custom spacing configured on the list.

diff --git a/Engine/Source/Runtime/UMG/Public/Components/ListViewBase.h b/Engine/Source/Runtime/UMG/Public/Components/ListViewBase.h
--- a/Engine/Source/Runtime/UMG/Public/Components/ListViewBase.h
+++ b/Engine/Source/Runtime/UMG/Public/Components/ListViewBase.h
@@ -413,7 +413,11 @@ private:
 		TSubclassOf<UUserWidget> DesiredEntryClass = GetDesiredEntryClassForItem(Item);
 
 		UUserWidget& EntryWidget = OnGenerateEntryWidgetInternal(Item, DesiredEntryClass, OwnerTable);
-		EntryWidget.SetPadding(GetDesiredEntryPadding(Item));
+		
+		// Combine the desired entry padding with the padding the widget wants natively on the CDO.
+		const FMargin DefaultPadding = EntryWidget.GetClass()->GetDefaultObject<UUserWidget>()->Padding;
+		EntryWidget.SetPadding(DefaultPadding + GetDesiredEntryPadding(Item));
+		
 		TSharedPtr<SWidget> CachedWidget = EntryWidget.GetCachedWidget();
 		CachedWidget->SetCanTick(true); // this is a hack to force ticking to true so selection works (which should NOT require ticking! but currently does)
 		return StaticCastSharedPtr<SObjectTableRow<ItemType>>(CachedWidget).ToSharedRef();
