Ignore "SObjectWidget destroyed while collecting garbage" ensure during blueprint compile GC phase.

diff --git a/Engine/Source/Editor/Kismet/Private/BlueprintCompilationManager.cpp b/Engine/Source/Editor/Kismet/Private/BlueprintCompilationManager.cpp
--- a/Engine/Source/Editor/Kismet/Private/BlueprintCompilationManager.cpp
+++ b/Engine/Source/Editor/Kismet/Private/BlueprintCompilationManager.cpp
@@ -274,6 +274,7 @@ void FBlueprintCompilationManagerImpl::CompileSynchronouslyImpl(const FBPCompile
 	{
 		DECLARE_SCOPE_HIERARCHICAL_COUNTER(CollectGarbage)
 
+		TGuardValue<bool> GuardTemplateNameFlag(GIsGCingAfterBlueprintCompile, true);
 		CollectGarbage(GARBAGE_COLLECTION_KEEPFLAGS);
 	}
 
diff --git a/Engine/Source/Runtime/Core/Private/Misc/CoreGlobals.cpp b/Engine/Source/Runtime/Core/Private/Misc/CoreGlobals.cpp
--- a/Engine/Source/Runtime/Core/Private/Misc/CoreGlobals.cpp
+++ b/Engine/Source/Runtime/Core/Private/Misc/CoreGlobals.cpp
@@ -99,6 +99,9 @@ bool GAllowActorScriptExecutionInEditor = false;
 /** Forces use of template names for newly instanced components in a CDO */
 bool GCompilingBlueprint = false;
 
+/** True if we're garbage collecting after a blueprint compilation */
+bool GIsGCingAfterBlueprintCompile = false;
+
 /** True if we're reconstructing blueprint instances. Should never be true on cooked builds */
 bool GIsReconstructingBlueprintInstances = false;
 
diff --git a/Engine/Source/Runtime/Core/Public/CoreGlobals.h b/Engine/Source/Runtime/Core/Public/CoreGlobals.h
--- a/Engine/Source/Runtime/Core/Public/CoreGlobals.h
+++ b/Engine/Source/Runtime/Core/Public/CoreGlobals.h
@@ -106,6 +106,9 @@ extern CORE_API bool GAllowActorScriptExecutionInEditor;
 /** Forces use of template names for newly instanced components in a CDO. */
 extern CORE_API bool GCompilingBlueprint;
 
+/** True if we're garbage collecting after a blueprint compilation */
+extern CORE_API bool GIsGCingAfterBlueprintCompile;
+
 /** True if we're reconstructing blueprint instances. Should never be true on cooked builds */
 extern CORE_API bool GIsReconstructingBlueprintInstances;
 
diff --git a/Engine/Source/Runtime/UMG/Private/Slate/SObjectWidget.cpp b/Engine/Source/Runtime/UMG/Private/Slate/SObjectWidget.cpp
--- a/Engine/Source/Runtime/UMG/Private/Slate/SObjectWidget.cpp
+++ b/Engine/Source/Runtime/UMG/Private/Slate/SObjectWidget.cpp
@@ -35,7 +35,7 @@ SObjectWidget::~SObjectWidget(void)
 #if SLATE_VERBOSE_NAMED_EVENTS
 	// This can happen during blueprint compiling, so just ignore it if it happens then, this is really only a concern
 	// in a running game.
-	if (!GCompilingBlueprint)
+	if (!GCompilingBlueprint && !GIsGCingAfterBlueprintCompile)
 	{
 		// This is only a concern during a running game - design-time instances can be destroyed from GC quite often when recompiling
 		if (!WidgetObject || !WidgetObject->IsDesignTime())
