Fix double selection change events in list view.
Invoke SignalSelectionChanged before OnItemClicked like we do in every other cases (right mouse button up, on touch ended).

diff --git a/Engine/Source/Runtime/UMG/Public/Slate/SObjectTableRow.h b/Engine/Source/Runtime/UMG/Public/Slate/SObjectTableRow.h
--- a/Engine/Source/Runtime/UMG/Public/Slate/SObjectTableRow.h
+++ b/Engine/Source/Runtime/UMG/Public/Slate/SObjectTableRow.h
@@ -347,14 +347,13 @@ public:
 			const ESelectionMode::Type SelectionMode = GetSelectionMode();
 			if (MouseEvent.GetEffectingButton() == EKeys::LeftMouseButton && HasMouseCapture())
 			{
+				bool bSignalSelectionChanged = bChangedSelectionOnMouseDown;
 				if (IsItemSelectable() && MyGeometry.IsUnderLocation(MouseEvent.GetScreenSpacePosition()))
 				{
 					if (SelectionMode == ESelectionMode::SingleToggle)
 					{
 						OwnerTable->Private_ClearSelection();
-						OwnerTable->Private_SignalSelectionChanged(ESelectInfo::OnMouseClick);
-
-						Reply = FReply::Handled();
+						bSignalSelectionChanged = true;
 					}
 					else if (SelectionMode == ESelectionMode::Multi &&
 						OwnerTable->Private_GetNumSelectedItems() > 1 &&
@@ -363,21 +362,19 @@ public:
 						// Releasing mouse over one of the multiple selected items - leave this one as the sole selected item
 						OwnerTable->Private_ClearSelection();
 						OwnerTable->Private_SetItemSelection(*MyItem, true, true);
-						OwnerTable->Private_SignalSelectionChanged(ESelectInfo::OnMouseClick);
-
-						Reply = FReply::Handled();
+						bSignalSelectionChanged = true;
 					}
 				}
 
-				if (OwnerTable->Private_OnItemClicked(*MyItem))
+				if (bSignalSelectionChanged)
 				{
+					OwnerTable->Private_SignalSelectionChanged(ESelectInfo::OnMouseClick);
 					Reply = FReply::Handled();
 				}
 
-				if (bChangedSelectionOnMouseDown)
+				if (OwnerTable->Private_OnItemClicked(*MyItem))
 				{
 					Reply = FReply::Handled();
-					OwnerTable->Private_SignalSelectionChanged(ESelectInfo::OnMouseClick);
 				}
 
 				Reply = Reply.ReleaseMouseCapture();
@@ -468,4 +465,4 @@ inline void SObjectTableRow<UObject*>::InitObjectRowInternal(UUserWidget& ListEn
 	{
 		IUserObjectListEntry::SetListItemObject(*WidgetObject, ListItemObject);
 	}
-}
\ No newline at end of file
+}
