Expose Geometry.LocalTopLeft to Blueprint

diff --git a/Engine/Source/Runtime/UMG/Private/SlateBlueprintLibrary.cpp b/Engine/Source/Runtime/UMG/Private/SlateBlueprintLibrary.cpp
--- a/Engine/Source/Runtime/UMG/Private/SlateBlueprintLibrary.cpp
+++ b/Engine/Source/Runtime/UMG/Private/SlateBlueprintLibrary.cpp
@@ -38,6 +38,12 @@ FVector2D USlateBlueprintLibrary::LocalToAbsolute(const FGeometry& Geometry, FVe
 	return Geometry.LocalToAbsolute(LocalCoordinate);
 }
 
+FVector2D USlateBlueprintLibrary::GetLocalTopLeft(const FGeometry& Geometry)
+{
+	const FVector2D TopLeft(0.0f, 0.0f);
+	return Geometry.GetLocalPositionAtCoordinates(TopLeft);
+}
+
 FVector2D USlateBlueprintLibrary::GetLocalSize(const FGeometry& Geometry)
 {
 	return Geometry.GetLocalSize();
diff --git a/Engine/Source/Runtime/UMG/Public/Blueprint/SlateBlueprintLibrary.h b/Engine/Source/Runtime/UMG/Public/Blueprint/SlateBlueprintLibrary.h
--- a/Engine/Source/Runtime/UMG/Public/Blueprint/SlateBlueprintLibrary.h
+++ b/Engine/Source/Runtime/UMG/Public/Blueprint/SlateBlueprintLibrary.h
@@ -43,6 +43,10 @@ public:
 	UFUNCTION(BlueprintPure, Category="User Interface|Geometry")
 	static FVector2D LocalToAbsolute(const FGeometry& Geometry, FVector2D LocalCoordinate);
 
+	/** Returns the local top/left of the geometry in local space. */
+	UFUNCTION(BlueprintPure, Category = "User Interface|Geometry")
+	static FVector2D GetLocalTopLeft(const FGeometry& Geometry);
+
 	/** Returns the size of the geometry in local space. */
 	UFUNCTION(BlueprintPure, Category="User Interface|Geometry")
 	static FVector2D GetLocalSize(const FGeometry& Geometry);
