Move UMG accessibility properties to a separate UObject in order to reduce memory footprint when using default values.

diff --git a/Engine/Source/Editor/UMGEditor/Private/Customizations/UMGDetailCustomizations.cpp b/Engine/Source/Editor/UMGEditor/Private/Customizations/UMGDetailCustomizations.cpp
--- a/Engine/Source/Editor/UMGEditor/Private/Customizations/UMGDetailCustomizations.cpp
+++ b/Engine/Source/Editor/UMGEditor/Private/Customizations/UMGDetailCustomizations.cpp
@@ -285,11 +285,6 @@ void FBlueprintWidgetCustomization::CustomizeDetails( IDetailLayoutBuilder& Deta
 	static const FName LayoutCategoryKey(TEXT("Layout"));
 	static const FName LocalizationCategoryKey(TEXT("Localization"));
 
-	static const FName AccessibleBehaviorKey(TEXT("AccessibleBehavior"));
-	static const FName AccessibleTextKey(TEXT("AccessibleText"));
-	static const FName AccessibleSummaryBehaviorKey(TEXT("AccessibleSummaryBehavior"));
-	static const FName AccessibleSummaryTextKey(TEXT("AccessibleSummaryText"));
-
 	DetailLayout.EditCategory(LocalizationCategoryKey, FText::GetEmpty(), ECategoryPriority::Uncommon);
 
 	TArray< TWeakObjectPtr<UObject> > OutObjects;
@@ -314,9 +309,7 @@ void FBlueprintWidgetCustomization::CustomizeDetails( IDetailLayoutBuilder& Deta
 		}
 	}
 
-	CustomizeAccessibilityProperty(DetailLayout, AccessibleBehaviorKey, AccessibleTextKey);
-	CustomizeAccessibilityProperty(DetailLayout, AccessibleSummaryBehaviorKey, AccessibleSummaryTextKey);
-
+	PerformAccessibilityCustomization(DetailLayout);
 	PerformBindingCustomization(DetailLayout);
 }
 
@@ -352,24 +345,28 @@ void FBlueprintWidgetCustomization::PerformBindingCustomization(IDetailLayoutBui
 	}
 }
 
-void FBlueprintWidgetCustomization::CustomizeAccessibilityProperty(IDetailLayoutBuilder& DetailLayout, const FName& BehaviorPropertyName, const FName& TextPropertyName)
+void FBlueprintWidgetCustomization::PerformAccessibilityCustomization(IDetailLayoutBuilder& DetailLayout)
 {
-	static const FName AccessibilityCategoryKey(TEXT("Accessibility"));
+	// We have to add these properties even though we're not customizing to preserve UI ordering
+	DetailLayout.EditCategory("Accessibility").AddProperty("bOverrideAccessibleDefaults");
+	DetailLayout.EditCategory("Accessibility").AddProperty("bCanChildrenBeAccessible");
+	CustomizeAccessibilityProperty(DetailLayout, "AccessibleBehavior", "AccessibleText");
+	CustomizeAccessibilityProperty(DetailLayout, "AccessibleSummaryBehavior", "AccessibleSummaryText");
+}
 
-	// Treat the *Behavior property as the "base" property for the row, and then add the *Text properties to the end of it.
-	IDetailCategoryBuilder& AccessibilityCategory = DetailLayout.EditCategory(AccessibilityCategoryKey);
+void FBlueprintWidgetCustomization::CustomizeAccessibilityProperty(IDetailLayoutBuilder& DetailLayout, const FName& BehaviorPropertyName, const FName& TextPropertyName)
+{
+	// Treat AccessibleBehavior as the "base" property for the row, and then add the AccessibleText binding to the end of it.
 	TSharedRef<IPropertyHandle> AccessibleBehaviorPropertyHandle = DetailLayout.GetProperty(BehaviorPropertyName);
-	IDetailPropertyRow& AccessibilityRow = AccessibilityCategory.AddProperty(AccessibleBehaviorPropertyHandle);
+	IDetailPropertyRow& AccessibilityRow = DetailLayout.EditCategory("Accessibility").AddProperty(AccessibleBehaviorPropertyHandle);
 
 	TSharedRef<IPropertyHandle> AccessibleTextPropertyHandle = DetailLayout.GetProperty(TextPropertyName);
-	// Make sure the old *Text properties are hidden so we don't get duplicate widgets
+	// Make sure the old AccessibleText properties are hidden so we don't get duplicate widgets
 	DetailLayout.HideProperty(AccessibleTextPropertyHandle);
-	// The *Text properties are put into an HBox which will consist of the text box and the delegate's combo box.
-	// The combo box may not always be valid for the current selection, as its creation is determined by the ExtensionHandler.
+
 	TSharedRef<SHorizontalBox> CustomTextLayout = SNew(SHorizontalBox)
 	.Visibility(TAttribute<EVisibility>::Create([AccessibleBehaviorPropertyHandle]() -> EVisibility
 	{
-		// Toggle the visibility of the *Text properties on only if the *Behavior property is set to Custom
 		uint8 Behavior = 0;
 		AccessibleBehaviorPropertyHandle->GetValue(Behavior);
 		return (ESlateAccessibleBehavior)Behavior == ESlateAccessibleBehavior::Custom ? EVisibility::Visible : EVisibility::Hidden;
@@ -381,14 +378,12 @@ void FBlueprintWidgetCustomization::CustomizeAccessibilityProperty(IDetailLayout
 	];
 
 	TSharedRef<SWidget> ExtensionWidget = const_cast<IDetailsView*>(DetailLayout.GetDetailsView())->GetExtensionHandler()->GenerateExtensionWidget(UWidget::StaticClass(), AccessibleTextPropertyHandle);
-	if (ExtensionWidget != SNullWidget::NullWidget)
-	{
-		CustomTextLayout->AddSlot()
-		.AutoWidth()
-		[
-			ExtensionWidget
-		];
-	}
+	ensure(ExtensionWidget != SNullWidget::NullWidget);
+	CustomTextLayout->AddSlot()
+	.AutoWidth()
+	[
+		ExtensionWidget
+	];
 
 	TSharedPtr<SWidget> AccessibleBehaviorNameWidget, AccessibleBehaviorValueWidget;
 	AccessibilityRow.GetDefaultWidgets(AccessibleBehaviorNameWidget, AccessibleBehaviorValueWidget);
diff --git a/Engine/Source/Editor/UMGEditor/Private/Customizations/UMGDetailCustomizations.h b/Engine/Source/Editor/UMGEditor/Private/Customizations/UMGDetailCustomizations.h
--- a/Engine/Source/Editor/UMGEditor/Private/Customizations/UMGDetailCustomizations.h
+++ b/Engine/Source/Editor/UMGEditor/Private/Customizations/UMGDetailCustomizations.h
@@ -44,7 +44,7 @@ private:
 
 	int32 HandleAddOrViewIndexForButton(const FName EventName, FName PropertyName) const;
 
-	/** Handle specific customizations for ESlateAccessibleBehavior properties */
+	void PerformAccessibilityCustomization(IDetailLayoutBuilder& DetailLayout);
 	void CustomizeAccessibilityProperty(IDetailLayoutBuilder& DetailLayout, const FName& BehaviorPropertyName, const FName& TextPropertyName);
 private:
 
diff --git a/Engine/Source/Runtime/UMG/Private/Components/Button.cpp b/Engine/Source/Runtime/UMG/Private/Components/Button.cpp
--- a/Engine/Source/Runtime/UMG/Private/Components/Button.cpp
+++ b/Engine/Source/Runtime/UMG/Private/Components/Button.cpp
@@ -35,8 +35,10 @@ UButton::UButton(const FObjectInitializer& ObjectInitializer)
 
 	IsFocusable = true;
 
+#if WITH_EDITORONLY_DATA
 	AccessibleBehavior = ESlateAccessibleBehavior::Summary;
 	bCanChildrenBeAccessible = false;
+#endif
 }
 
 void UButton::ReleaseSlateResources(bool bReleaseChildren)
diff --git a/Engine/Source/Runtime/UMG/Private/Components/CheckBox.cpp b/Engine/Source/Runtime/UMG/Private/Components/CheckBox.cpp
--- a/Engine/Source/Runtime/UMG/Private/Components/CheckBox.cpp
+++ b/Engine/Source/Runtime/UMG/Private/Components/CheckBox.cpp
@@ -35,8 +35,10 @@ UCheckBox::UCheckBox(const FObjectInitializer& ObjectInitializer)
 	BorderBackgroundColor_DEPRECATED = FLinearColor::White;
 
 	IsFocusable = true;
+#if WITH_EDITORONLY_DATA
 	AccessibleBehavior = ESlateAccessibleBehavior::Summary;
 	bCanChildrenBeAccessible = false;
+#endif
 }
 
 void UCheckBox::ReleaseSlateResources(bool bReleaseChildren)
diff --git a/Engine/Source/Runtime/UMG/Private/Components/EditableText.cpp b/Engine/Source/Runtime/UMG/Private/Components/EditableText.cpp
--- a/Engine/Source/Runtime/UMG/Private/Components/EditableText.cpp
+++ b/Engine/Source/Runtime/UMG/Private/Components/EditableText.cpp
@@ -48,8 +48,10 @@ UEditableText::UEditableText(const FObjectInitializer& ObjectInitializer)
 	VirtualKeyboardDismissAction = EVirtualKeyboardDismissAction::TextChangeOnDismiss;
 	Clipping = EWidgetClipping::ClipToBounds;
 
+#if WITH_EDITORONLY_DATA
 	AccessibleBehavior = ESlateAccessibleBehavior::Auto;
 	bCanChildrenBeAccessible = false;
+#endif
 }
 
 void UEditableText::ReleaseSlateResources(bool bReleaseChildren)
diff --git a/Engine/Source/Runtime/UMG/Private/Components/EditableTextBox.cpp b/Engine/Source/Runtime/UMG/Private/Components/EditableTextBox.cpp
--- a/Engine/Source/Runtime/UMG/Private/Components/EditableTextBox.cpp
+++ b/Engine/Source/Runtime/UMG/Private/Components/EditableTextBox.cpp
@@ -49,8 +49,10 @@ UEditableTextBox::UEditableTextBox(const FObjectInitializer& ObjectInitializer)
 
 	WidgetStyle = *DefaultEditableTextBoxStyle;
 
+#if WITH_EDITORONLY_DATA
 	AccessibleBehavior = ESlateAccessibleBehavior::Auto;
 	bCanChildrenBeAccessible = false;
+#endif
 }
 
 void UEditableTextBox::ReleaseSlateResources(bool bReleaseChildren)
diff --git a/Engine/Source/Runtime/UMG/Private/Components/Slider.cpp b/Engine/Source/Runtime/UMG/Private/Components/Slider.cpp
--- a/Engine/Source/Runtime/UMG/Private/Components/Slider.cpp
+++ b/Engine/Source/Runtime/UMG/Private/Components/Slider.cpp
@@ -35,8 +35,10 @@ USlider::USlider(const FObjectInitializer& ObjectInitializer)
 
 	WidgetStyle = *DefaultSliderStyle;
 
+#if WITH_EDITORONLY_DATA
 	AccessibleBehavior = ESlateAccessibleBehavior::Summary;
 	bCanChildrenBeAccessible = false;
+#endif
 }
 
 TSharedRef<SWidget> USlider::RebuildWidget()
diff --git a/Engine/Source/Runtime/UMG/Private/Components/TextBlock.cpp b/Engine/Source/Runtime/UMG/Private/Components/TextBlock.cpp
--- a/Engine/Source/Runtime/UMG/Private/Components/TextBlock.cpp
+++ b/Engine/Source/Runtime/UMG/Private/Components/TextBlock.cpp
@@ -31,8 +31,10 @@ UTextBlock::UTextBlock(const FObjectInitializer& ObjectInitializer)
 		Font = FSlateFontInfo(RobotoFontObj.Object, 24, FName("Bold"));
 	}
 
+#if WITH_EDITORONLY_DATA
 	AccessibleBehavior = ESlateAccessibleBehavior::Auto;
 	bCanChildrenBeAccessible = false;
+#endif
 }
 
 void UTextBlock::PostLoad()
diff --git a/Engine/Source/Runtime/UMG/Private/Components/Widget.cpp b/Engine/Source/Runtime/UMG/Private/Components/Widget.cpp
--- a/Engine/Source/Runtime/UMG/Private/Components/Widget.cpp
+++ b/Engine/Source/Runtime/UMG/Private/Components/Widget.cpp
@@ -161,9 +161,13 @@ UWidget::UWidget(const FObjectInitializer& ObjectInitializer)
 	RenderTransformPivot = FVector2D(0.5f, 0.5f);
 	Cursor = EMouseCursor::Default;
 
+#if WITH_EDITORONLY_DATA
+	bOverrideAccessibleDefaults = false;
 	AccessibleBehavior = ESlateAccessibleBehavior::NotAccessible;
 	AccessibleSummaryBehavior = ESlateAccessibleBehavior::Auto;
 	bCanChildrenBeAccessible = true;
+#endif
+	AccessibleWidgetData = nullptr;
 
 #if WITH_EDITORONLY_DATA
 	{ static const FAutoRegisterLocalizationDataGatheringCallback AutomaticRegistrationOfLocalizationGatherer(UWidget::StaticClass(), &GatherWidgetForLocalization); }
@@ -1072,6 +1076,10 @@ void UWidget::PostEditChangeProperty(FPropertyChangedEvent& PropertyChangedEvent
 	{
 		SynchronizeProperties();
 	}
+	else
+	{
+		SynchronizeAccessibleData();
+	}
 }
 
 void UWidget::SelectByDesigner()
@@ -1102,6 +1110,15 @@ void UWidget::DeselectByDesigner()
 #define LOCTEXT_NAMESPACE "UMG"
 #endif
 
+void UWidget::PreSave(const class ITargetPlatform* TargetPlatform)
+{
+	Super::PreSave(TargetPlatform);
+
+	// This is a failsafe to make sure all the accessibility data is copied over in case
+	// some rare instance isn't handled by SynchronizeProperties. It might not be necessary.
+	SynchronizeAccessibleData();
+}
+
 #if WITH_EDITOR
 bool UWidget::Modify(bool bAlwaysMarkDirty)
 {
@@ -1144,6 +1161,9 @@ void UWidget::SynchronizeProperties()
 	bRoutedSynchronizeProperties = true;
 #endif
 
+	// Always sync accessible data even if the SWidget doesn't exist
+	SynchronizeAccessibleData();
+
 	// We want to apply the bindings to the cached widget, which could be the SWidget, or the SObjectWidget, 
 	// in the case where it's a user widget.  We always want to prefer the SObjectWidget so that bindings to 
 	// visibility and enabled status are not stomping values setup in the root widget in the User Widget.
@@ -1218,12 +1238,40 @@ void UWidget::SynchronizeProperties()
 	}
 
 #if WITH_ACCESSIBILITY
-	TSharedPtr<SWidget> AccessibleWidget = GetAccessibleWidget();
-	if (AccessibleWidget.IsValid())
+	if (AccessibleWidgetData)
+	{
+		TSharedPtr<SWidget> AccessibleWidget = GetAccessibleWidget();
+		if (AccessibleWidget.IsValid())
+		{
+			AccessibleWidget->SetAccessibleBehavior((EAccessibleBehavior)AccessibleWidgetData->AccessibleBehavior, AccessibleWidgetData->CreateAccessibleTextAttribute(), EAccessibleType::Main);
+			AccessibleWidget->SetAccessibleBehavior((EAccessibleBehavior)AccessibleWidgetData->AccessibleSummaryBehavior, AccessibleWidgetData->CreateAccessibleSummaryTextAttribute(), EAccessibleType::Summary);
+			AccessibleWidget->SetCanChildrenBeAccessible(AccessibleWidgetData->bCanChildrenBeAccessible);
+		}
+	}
+#endif
+}
+
+
+void UWidget::SynchronizeAccessibleData()
+{
+#if WITH_EDITORONLY_DATA
+	if (bOverrideAccessibleDefaults)
+	{
+		if (!AccessibleWidgetData)
+		{
+			AccessibleWidgetData = NewObject<USlateAccessibleWidgetData>(this);
+		}
+		AccessibleWidgetData->bCanChildrenBeAccessible = bCanChildrenBeAccessible;
+		AccessibleWidgetData->AccessibleBehavior = AccessibleBehavior;
+		AccessibleWidgetData->AccessibleText = AccessibleText;
+		AccessibleWidgetData->AccessibleTextDelegate = AccessibleTextDelegate;
+		AccessibleWidgetData->AccessibleSummaryBehavior = AccessibleSummaryBehavior;
+		AccessibleWidgetData->AccessibleSummaryText = AccessibleSummaryText;
+		AccessibleWidgetData->AccessibleSummaryTextDelegate = AccessibleSummaryTextDelegate;
+	}
+	else if (AccessibleWidgetData)
 	{
-		AccessibleWidget->SetAccessibleBehavior((EAccessibleBehavior)AccessibleBehavior, PROPERTY_BINDING(FText, AccessibleText), EAccessibleType::Main);
-		AccessibleWidget->SetAccessibleBehavior((EAccessibleBehavior)AccessibleSummaryBehavior, PROPERTY_BINDING(FText, AccessibleSummaryText), EAccessibleType::Summary);
-		AccessibleWidget->SetCanChildrenBeAccessible(bCanChildrenBeAccessible);
+		AccessibleWidgetData = nullptr;
 	}
 #endif
 }
diff --git a/Engine/Source/Runtime/UMG/Public/Components/SlateWrapperTypes.h b/Engine/Source/Runtime/UMG/Public/Components/SlateWrapperTypes.h
--- a/Engine/Source/Runtime/UMG/Public/Components/SlateWrapperTypes.h
+++ b/Engine/Source/Runtime/UMG/Public/Components/SlateWrapperTypes.h
@@ -53,6 +53,61 @@ enum class ESlateAccessibleBehavior : uint8
 	ToolTip
 };
 
+/**
+ * A container for all accessible properties for a UWidget that will be passed to the underlying SWidget.
+ * Any property here must also be added to UWidget.h and synchronized. See UWidget for more information.
+ */
+UCLASS(BlueprintType, EditInlineNew, DefaultToInstanced)
+class USlateAccessibleWidgetData : public UObject
+{
+	GENERATED_BODY()
+
+public:
+	DECLARE_DYNAMIC_DELEGATE_RetVal(FText, FGetText);
+
+	USlateAccessibleWidgetData()
+	{
+		AccessibleBehavior = ESlateAccessibleBehavior::NotAccessible;
+		AccessibleSummaryBehavior = ESlateAccessibleBehavior::Auto;
+		bCanChildrenBeAccessible = true;
+	}
+
+	TAttribute<FText> CreateAccessibleTextAttribute()
+	{
+		return AccessibleTextDelegate.IsBound() ?
+			TAttribute<FText>::Create(AccessibleTextDelegate.GetUObject(), AccessibleTextDelegate.GetFunctionName()) :
+			TAttribute<FText>(AccessibleText);
+	}
+
+	TAttribute<FText> CreateAccessibleSummaryTextAttribute()
+	{
+		return AccessibleSummaryTextDelegate.IsBound() ?
+			TAttribute<FText>::Create(AccessibleSummaryTextDelegate.GetUObject(), AccessibleSummaryTextDelegate.GetFunctionName()) :
+			TAttribute<FText>(AccessibleSummaryText);
+	}
+
+	UPROPERTY()
+	bool bCanChildrenBeAccessible;
+
+	UPROPERTY()
+	ESlateAccessibleBehavior AccessibleBehavior;
+
+	UPROPERTY()
+	ESlateAccessibleBehavior AccessibleSummaryBehavior;
+
+	UPROPERTY()
+	FText AccessibleText;
+
+	UPROPERTY()
+	FGetText AccessibleTextDelegate;
+
+	UPROPERTY()
+	FText AccessibleSummaryText;
+
+	UPROPERTY()
+	FGetText AccessibleSummaryTextDelegate;
+};
+
 /** The sizing options of UWidgets */
 UENUM(BlueprintType)
 namespace ESlateSizeRule
diff --git a/Engine/Source/Runtime/UMG/Public/Components/Widget.h b/Engine/Source/Runtime/UMG/Public/Components/Widget.h
--- a/Engine/Source/Runtime/UMG/Public/Components/Widget.h
+++ b/Engine/Source/Runtime/UMG/Public/Components/Widget.h
@@ -318,31 +318,56 @@ public:
 	UPROPERTY(EditAnywhere, Category="Behavior", meta=(InlineEditConditionToggle))
 	uint8 bOverride_Cursor : 1;
 
+#if WITH_EDITORONLY_DATA
+	// These editor-only properties exist for two reasons:
+	//   1. To make details customization easier to write, specifically in regards to the binding extension widget
+	//   2. To allow subclasses to set their default values without having to subclass USlateAccessibleWidgetData
+	// Every time one of these properties changes, it's data is propagated to AccessibleWidgetData if it exists.
+	// The creations of AccessibleWidgetData is controlled by the details customization through a CheckBox.
+	// The reason this is set up like this is to reduce the memory footprint of UWidget since overriding the default
+	// accessibility rules for a particular widget will be relatively rare. In a shipped game, if no custom rules
+	// are defined, there will only be the memory cost of the UObject pointer.
+	//
+	// IMPORTANT: Any user-editable variables added to USlateAccessibleWidgetData should be duplicated here as well.
+	//            Additionally, its edit condition must be manually assigned in UMGDetailCustomizations.
+
+	/** Override all of the default accessibility behavior and text for this widget. */
+	UPROPERTY(EditAnywhere, Category="Accessibility")
+	uint8 bOverrideAccessibleDefaults : 1;
+
 	/** Whether or not children of this widget can appear as distinct accessible widgets. */
-	UPROPERTY(EditAnywhere, Category = "Accessibility")
+	UPROPERTY(EditAnywhere, Category="Accessibility", meta=(EditCondition="bOverrideAccessibleDefaults"))
 	uint8 bCanChildrenBeAccessible : 1;
 
 	/** Whether or not the widget is accessible, and how to describe it. If set to custom, additional customization options will appear. */
-	UPROPERTY(EditAnywhere, Category = "Accessibility")
+	UPROPERTY(EditAnywhere, Category="Accessibility", meta=(EditCondition="bOverrideAccessibleDefaults"))
 	ESlateAccessibleBehavior AccessibleBehavior;
 
+	/** How to describe this widget when it's being presented through a summary of a parent widget. If set to custom, additional customization options will appear. */
+	UPROPERTY(EditAnywhere, Category="Accessibility", AdvancedDisplay, meta=(EditCondition="bOverrideAccessibleDefaults"))
+	ESlateAccessibleBehavior AccessibleSummaryBehavior;
+
 	/** When AccessibleBehavior is set to Custom, this is the text that will be used to describe the widget. */
-	UPROPERTY(EditAnywhere, BlueprintReadOnly, Category = "Accessibility", meta = (MultiLine = true))
+	UPROPERTY(EditAnywhere, BlueprintReadOnly, Category="Accessibility", meta=(MultiLine=true))
 	FText AccessibleText;
 
+	/** An optional delegate that may be assigned in place of AccessibleText for creating a TAttribute */
 	UPROPERTY()
-	FGetText AccessibleTextDelegate;
-
-	/** How to describe this widget when it's being presented through a summary of a parent widget. If set to custom, additional customization options will appear. */
-	UPROPERTY(EditAnywhere, Category = "Accessibility", AdvancedDisplay)
-	ESlateAccessibleBehavior AccessibleSummaryBehavior;
+	USlateAccessibleWidgetData::FGetText AccessibleTextDelegate;
 
 	/** When AccessibleSummaryBehavior is set to Custom, this is the text that will be used to describe the widget. */
-	UPROPERTY(EditAnywhere, BlueprintReadOnly, Category = "Accessibility", meta = (MultiLine = true), AdvancedDisplay)
+	UPROPERTY(EditAnywhere, BlueprintReadOnly, Category="Accessibility", meta=(MultiLine=true), AdvancedDisplay)
 	FText AccessibleSummaryText;
 
+	/** An optional delegate that may be assigned in place of AccessibleSummaryText for creating a TAttribute */
 	UPROPERTY()
-	FGetText AccessibleSummaryTextDelegate;
+	USlateAccessibleWidgetData::FGetText AccessibleSummaryTextDelegate;
+#endif
+
+private:
+	/** A custom set of accessibility rules for this widget. If null, default rules for the widget are used. */
+	UPROPERTY()
+	USlateAccessibleWidgetData* AccessibleWidgetData;
 
 protected:
 
@@ -853,6 +878,7 @@ public:
 	virtual UWorld* GetWorld() const override;
 	virtual void FinishDestroy() override;
 	virtual bool IsDestructionThreadSafe() const override { return false; }
+	virtual void PreSave(const class ITargetPlatform* TargetPlatform) override;
 	// End UObject
 
 	FORCEINLINE bool CanSafelyRouteEvent()
@@ -974,6 +1000,9 @@ protected:
 	/** Gets the base name used to generate the display label/name of this widget. */
 	FText GetDisplayNameBase() const;
 
+	/** Copy all accessible properties to the AccessibleWidgetData object */
+	void SynchronizeAccessibleData();
+
 protected:
 	//TODO UMG Consider moving conversion functions into another class.
 	// Conversion functions
@@ -1046,6 +1075,8 @@ private:
 private:
 	PROPERTY_BINDING_IMPLEMENTATION(FText, ToolTipText);
 	PROPERTY_BINDING_IMPLEMENTATION(bool, bIsEnabled);
+#if WITH_EDITORONLY_DATA
 	PROPERTY_BINDING_IMPLEMENTATION(FText, AccessibleText);
 	PROPERTY_BINDING_IMPLEMENTATION(FText, AccessibleSummaryText);
+#endif
 };
