Remove unnecessary metadata from accessible data class

diff --git a/Engine/Source/Runtime/UMG/Public/Components/SlateWrapperTypes.h b/Engine/Source/Runtime/UMG/Public/Components/SlateWrapperTypes.h
--- a/Engine/Source/Runtime/UMG/Public/Components/SlateWrapperTypes.h
+++ b/Engine/Source/Runtime/UMG/Public/Components/SlateWrapperTypes.h
@@ -57,7 +57,7 @@ enum class ESlateAccessibleBehavior : uint8
  * A container for all accessible properties for a UWidget that will be passed to the underlying SWidget.
  * Any property here must also be added to UWidget.h and synchronized. See UWidget for more information.
  */
-UCLASS(BlueprintType, EditInlineNew, DefaultToInstanced)
+UCLASS(DefaultToInstanced)
 class USlateAccessibleWidgetData : public UObject
 {
 	GENERATED_BODY()
