ListView now takes its orientation into consideration when getting desired entry padding.

#rnx

diff --git a/Engine/Source/Runtime/UMG/Private/Components/ListView.cpp b/Engine/Source/Runtime/UMG/Private/Components/ListView.cpp
--- a/Engine/Source/Runtime/UMG/Private/Components/ListView.cpp
+++ b/Engine/Source/Runtime/UMG/Private/Components/ListView.cpp
@@ -225,9 +225,18 @@ FMargin UListView::GetDesiredEntryPadding(UObject* Item) const
 {
 	if (ListItems.Num() > 0 && ListItems[0] != Item)
 	{
-		// For all entries after the first one, add the spacing as top padding
-		return FMargin(0.f, EntrySpacing, 0.f, 0.f);
+		if (Orientation == EOrientation::Orient_Horizontal)
+		{
+			// For all entries after the first one, add the spacing as left padding
+			return FMargin(EntrySpacing, 0.f, 0.0f, 0.f);
+		}
+		else
+		{
+			// For all entries after the first one, add the spacing as top padding
+			return FMargin(0.f, EntrySpacing, 0.f, 0.f);
+		}
 	}
+
 	return FMargin(0.f);
 }
 
