Improve cacheline boundary of UHorizontalBoxSlot

#rnx

diff --git a/Engine/Source/Runtime/UMG/Public/Components/HorizontalBoxSlot.h b/Engine/Source/Runtime/UMG/Public/Components/HorizontalBoxSlot.h
--- a/Engine/Source/Runtime/UMG/Public/Components/HorizontalBoxSlot.h
+++ b/Engine/Source/Runtime/UMG/Public/Components/HorizontalBoxSlot.h
@@ -17,6 +17,9 @@ class UMG_API UHorizontalBoxSlot : public UPanelSlot
 {
 	GENERATED_UCLASS_BODY()
 
+private:
+	SHorizontalBox::FSlot* Slot;
+
 public:
 	
 	/** The amount of padding between the slots parent and the content. */
@@ -33,8 +36,6 @@ public:
 	UPROPERTY(EditAnywhere, BlueprintReadOnly, Category="Layout|Horizontal Box Slot")
 	TEnumAsByte<EVerticalAlignment> VerticalAlignment;
 
-public:
-
 	UFUNCTION(BlueprintCallable, Category="Layout|Horizontal Box Slot")
 	void SetPadding(FMargin InPadding);
 
@@ -55,6 +56,4 @@ public:
 
 	virtual void ReleaseSlateResources(bool bReleaseChildren) override;
 
-private:
-	SHorizontalBox::FSlot* Slot;
 };
