Improve cache coherancy of the UVerticalBoxSlot

#rnx

diff --git a/Engine/Source/Runtime/UMG/Public/Components/VerticalBoxSlot.h b/Engine/Source/Runtime/UMG/Public/Components/VerticalBoxSlot.h
--- a/Engine/Source/Runtime/UMG/Public/Components/VerticalBoxSlot.h
+++ b/Engine/Source/Runtime/UMG/Public/Components/VerticalBoxSlot.h
@@ -20,13 +20,21 @@ class UMG_API UVerticalBoxSlot : public UPanelSlot
 
 public:
 	
+	/** How much space this slot should occupy in the direction of the panel. */
+	UPROPERTY(EditAnywhere, BlueprintReadOnly, Category="Layout|Vertical Box Slot", meta = (DisplayAfter = "Padding"))
+	FSlateChildSize Size;
+
 	/** The padding area between the slot and the content it contains. */
 	UPROPERTY(EditAnywhere, BlueprintReadOnly, Category="Layout|Vertical Box Slot")
 	FMargin Padding;
 
-	/** How much space this slot should occupy in the direction of the panel. */
-	UPROPERTY(EditAnywhere, BlueprintReadOnly, Category="Layout|Vertical Box Slot")
-	FSlateChildSize Size;
+private:
+	//TODO UMG Slots should hold weak or shared refs to slots.
+
+	/** A raw pointer to the slot to allow us to adjust the size, padding...etc at runtime. */
+	SVerticalBox::FSlot* Slot;
+
+public:
 
 	/** The alignment of the object horizontally. */
 	UPROPERTY(EditAnywhere, BlueprintReadOnly, Category="Layout|Vertical Box Slot")
@@ -60,10 +68,4 @@ public:
 
 	/** Builds the underlying FSlot for the Slate layout panel. */
 	void BuildSlot(TSharedRef<SVerticalBox> InVerticalBox);
-
-private:
-	//TODO UMG Slots should hold weak or shared refs to slots.
-
-	/** A raw pointer to the slot to allow us to adjust the size, padding...etc at runtime. */
-	SVerticalBox::FSlot* Slot;
 };
