Improve packing of Size Box widgets and BP panels

#rnx

diff --git a/Engine/Source/Runtime/Core/Public/Math/Color.h b/Engine/Source/Runtime/Core/Public/Math/Color.h
--- a/Engine/Source/Runtime/Core/Public/Math/Color.h
+++ b/Engine/Source/Runtime/Core/Public/Math/Color.h
@@ -14,7 +14,7 @@ class FFloat16Color;
 /**
  * Enum for the different kinds of gamma spaces we expect to need to convert from/to.
  */
-enum class EGammaSpace
+enum class EGammaSpace : uint8
 {
 	/** No gamma correction is applied to this space, the incoming colors are assumed to already be in linear space. */
 	Linear,
diff --git a/Engine/Source/Runtime/Engine/Classes/Engine/BlueprintGeneratedClass.h b/Engine/Source/Runtime/Engine/Classes/Engine/BlueprintGeneratedClass.h
--- a/Engine/Source/Runtime/Engine/Classes/Engine/BlueprintGeneratedClass.h
+++ b/Engine/Source/Runtime/Engine/Classes/Engine/BlueprintGeneratedClass.h
@@ -539,10 +539,6 @@ struct ENGINE_API FBlueprintCookedComponentInstancingData
 {
 	GENERATED_USTRUCT_BODY()
 
-	/** Flag indicating whether or not this contains valid cooked data. Note that an empty changed property list can also be a valid template data context. */
-	UPROPERTY()
-	bool bHasValidCookedData;
-
 	/** List of property info records with values that differ between the template and the component class CDO. This list will be generated at cook time. */
 	UPROPERTY()
 	TArray<struct FBlueprintComponentChangedPropertyInfo> ChangedPropertyList;
@@ -554,7 +550,11 @@ struct ENGINE_API FBlueprintCookedComponentInstancingData
 	UClass* ComponentTemplateClass;
 
 	/** Source template object flags (recorded at load time and used for instancing). */
-	EObjectFlags ComponentTemplateFlags;
+	TEnumAsByte<EObjectFlags> ComponentTemplateFlags;
+
+	/** Flag indicating whether or not this contains valid cooked data. Note that an empty changed property list can also be a valid template data context. */
+	UPROPERTY()
+	uint8 bHasValidCookedData : 1;
 
 	/** Default constructor. */
 	FBlueprintCookedComponentInstancingData()
diff --git a/Engine/Source/Runtime/ImageCore/Public/ImageCore.h b/Engine/Source/Runtime/ImageCore/Public/ImageCore.h
--- a/Engine/Source/Runtime/ImageCore/Public/ImageCore.h
+++ b/Engine/Source/Runtime/ImageCore/Public/ImageCore.h
@@ -16,7 +16,7 @@ namespace ERawImageFormat
 	/**
 	 * Enumerates supported raw image formats.
 	 */
-	enum Type
+	enum Type : uint8
 	{
 		G8,
 		BGRA8,
diff --git a/Engine/Source/Runtime/UMG/Public/Components/SizeBox.h b/Engine/Source/Runtime/UMG/Public/Components/SizeBox.h
--- a/Engine/Source/Runtime/UMG/Public/Components/SizeBox.h
+++ b/Engine/Source/Runtime/UMG/Public/Components/SizeBox.h
@@ -22,116 +22,117 @@ class UMG_API USizeBox : public UContentWidget
 {
 	GENERATED_UCLASS_BODY()
 
-public:
-
-	/**  */
-	UPROPERTY(EditAnywhere, Category="Child Layout", meta=(InlineEditConditionToggle))
-		uint32 bOverride_WidthOverride : 1;
-
-	/**  */
-	UPROPERTY(EditAnywhere, Category="Child Layout", meta=(InlineEditConditionToggle))
-		uint32 bOverride_HeightOverride : 1;
-
-	/**  */
-	UPROPERTY(EditAnywhere, Category="Child Layout", meta=(InlineEditConditionToggle))
-		uint32 bOverride_MinDesiredWidth : 1;
-
-	/**  */
-	UPROPERTY(EditAnywhere, Category="Child Layout", meta=(InlineEditConditionToggle))
-		uint32 bOverride_MinDesiredHeight : 1;
-
-	/**  */
-	UPROPERTY(EditAnywhere, Category="Child Layout", meta=(InlineEditConditionToggle))
-		uint32 bOverride_MaxDesiredWidth : 1;
-
-	/**  */
-	UPROPERTY(EditAnywhere, Category="Child Layout", meta=(InlineEditConditionToggle))
-		uint32 bOverride_MaxDesiredHeight : 1;
-
-	/**  */
-	UPROPERTY(EditAnywhere, Category="Child Layout", meta=(InlineEditConditionToggle))
-	uint32 bOverride_MinAspectRatio : 1;
-
-	/**  */
-	UPROPERTY(EditAnywhere, Category = "Child Layout", meta=(InlineEditConditionToggle))
-	uint32 bOverride_MaxAspectRatio : 1;
+protected:
+	TSharedPtr<SBox> MySizeBox;
 
+public:
 
 	/** When specified, ignore the content's desired size and report the WidthOverride as the Box's desired width. */
 	UPROPERTY(EditAnywhere, BlueprintReadOnly, Category="Child Layout", meta=( editcondition="bOverride_WidthOverride" ))
-		float WidthOverride;
+	float WidthOverride;
 
 	/** When specified, ignore the content's desired size and report the HeightOverride as the Box's desired height. */
 	UPROPERTY(EditAnywhere, BlueprintReadOnly, Category="Child Layout", meta=( editcondition="bOverride_HeightOverride" ))
-		float HeightOverride;
+	float HeightOverride;
 
 	/** When specified, will report the MinDesiredWidth if larger than the content's desired width. */
 	UPROPERTY(EditAnywhere, BlueprintReadOnly, Category="Child Layout", meta=( editcondition="bOverride_MinDesiredWidth" ))
-		float MinDesiredWidth;
+	float MinDesiredWidth;
 
 	/** When specified, will report the MinDesiredHeight if larger than the content's desired height. */
 	UPROPERTY(EditAnywhere, BlueprintReadOnly, Category="Child Layout", meta=( editcondition="bOverride_MinDesiredHeight" ))
-		float MinDesiredHeight;
+	float MinDesiredHeight;
 
 	/** When specified, will report the MaxDesiredWidth if smaller than the content's desired width. */
 	UPROPERTY(EditAnywhere, BlueprintReadOnly, Category="Child Layout", meta=( editcondition="bOverride_MaxDesiredWidth" ))
-		float MaxDesiredWidth;
+	float MaxDesiredWidth;
 
 	/** When specified, will report the MaxDesiredHeight if smaller than the content's desired height. */
 	UPROPERTY(EditAnywhere, BlueprintReadOnly, Category="Child Layout", meta=( editcondition="bOverride_MaxDesiredHeight" ))
-		float MaxDesiredHeight;
+	float MaxDesiredHeight;
 
 	/** */
 	UPROPERTY(EditAnywhere, BlueprintReadOnly, Category = "Child Layout", meta = (editcondition = "bOverride_MinAspectRatio"))
-		float MinAspectRatio;
+	float MinAspectRatio;
 
 	/** */
 	UPROPERTY(EditAnywhere, BlueprintReadOnly, Category="Child Layout", meta=( editcondition="bOverride_MaxAspectRatio" ))
-		float MaxAspectRatio;
+	float MaxAspectRatio;
+
+	/**  */
+	UPROPERTY(EditAnywhere, Category="Child Layout", meta=(InlineEditConditionToggle))
+	uint8 bOverride_WidthOverride : 1;
+
+	/**  */
+	UPROPERTY(EditAnywhere, Category="Child Layout", meta=(InlineEditConditionToggle))
+	uint8 bOverride_HeightOverride : 1;
+
+	/**  */
+	UPROPERTY(EditAnywhere, Category="Child Layout", meta=(InlineEditConditionToggle))
+	uint8 bOverride_MinDesiredWidth : 1;
+
+	/**  */
+	UPROPERTY(EditAnywhere, Category="Child Layout", meta=(InlineEditConditionToggle))
+	uint8 bOverride_MinDesiredHeight : 1;
+
+	/**  */
+	UPROPERTY(EditAnywhere, Category="Child Layout", meta=(InlineEditConditionToggle))
+	uint8 bOverride_MaxDesiredWidth : 1;
+
+	/**  */
+	UPROPERTY(EditAnywhere, Category="Child Layout", meta=(InlineEditConditionToggle))
+	uint8 bOverride_MaxDesiredHeight : 1;
+
+	/**  */
+	UPROPERTY(EditAnywhere, Category="Child Layout", meta=(InlineEditConditionToggle))
+	uint8 bOverride_MinAspectRatio : 1;
+
+	/**  */
+	UPROPERTY(EditAnywhere, Category = "Child Layout", meta=(InlineEditConditionToggle))
+	uint8 bOverride_MaxAspectRatio : 1;
 
-public:
 		
 	/** When specified, ignore the content's desired size and report the WidthOverride as the Box's desired width. */
 	UFUNCTION(BlueprintCallable, Category="Layout|Size Box")
-		void SetWidthOverride(float InWidthOverride);
+	void SetWidthOverride(float InWidthOverride);
 
 	UFUNCTION(BlueprintCallable, Category="Layout|Size Box")
-		void ClearWidthOverride();
+	void ClearWidthOverride();
 
 	/** When specified, ignore the content's desired size and report the HeightOverride as the Box's desired height. */
 	UFUNCTION(BlueprintCallable, Category="Layout|Size Box")
-		void SetHeightOverride(float InHeightOverride);
+	void SetHeightOverride(float InHeightOverride);
 
 	UFUNCTION(BlueprintCallable, Category="Layout|Size Box")
-		void ClearHeightOverride();
+	void ClearHeightOverride();
 
 	/** When specified, will report the MinDesiredWidth if larger than the content's desired width. */
 	UFUNCTION(BlueprintCallable, Category="Layout|Size Box")
-		void SetMinDesiredWidth(float InMinDesiredWidth);
+	void SetMinDesiredWidth(float InMinDesiredWidth);
 
 	UFUNCTION(BlueprintCallable, Category="Layout|Size Box")
-		void ClearMinDesiredWidth();
+	void ClearMinDesiredWidth();
 
 	/** When specified, will report the MinDesiredHeight if larger than the content's desired height. */
 	UFUNCTION(BlueprintCallable, Category="Layout|Size Box")
-		void SetMinDesiredHeight(float InMinDesiredHeight);
+	void SetMinDesiredHeight(float InMinDesiredHeight);
 
 	UFUNCTION(BlueprintCallable, Category="Layout|Size Box")
-		void ClearMinDesiredHeight();
+	void ClearMinDesiredHeight();
 
 	/** When specified, will report the MaxDesiredWidth if smaller than the content's desired width. */
 	UFUNCTION(BlueprintCallable, Category="Layout|Size Box")
-		void SetMaxDesiredWidth(float InMaxDesiredWidth);
+	void SetMaxDesiredWidth(float InMaxDesiredWidth);
 
 	UFUNCTION(BlueprintCallable, Category="Layout|Size Box")
-		void ClearMaxDesiredWidth();
+	void ClearMaxDesiredWidth();
 
 	/** When specified, will report the MaxDesiredHeight if smaller than the content's desired height. */
 	UFUNCTION(BlueprintCallable, Category="Layout|Size Box")
-		void SetMaxDesiredHeight(float InMaxDesiredHeight);
+	void SetMaxDesiredHeight(float InMaxDesiredHeight);
 
 	UFUNCTION(BlueprintCallable, Category="Layout|Size Box")
-		void ClearMaxDesiredHeight();
+	void ClearMaxDesiredHeight();
 
 	UFUNCTION(BlueprintCallable, Category="Layout|Size Box")
 	void SetMinAspectRatio(float InMinAspectRatio);
@@ -145,7 +146,6 @@ public:
 	UFUNCTION(BlueprintCallable, Category="Layout|Size Box")
 	void ClearMaxAspectRatio();
 
-public:
 
 	// UWidget interface
 	virtual void SynchronizeProperties() override;
@@ -167,10 +167,6 @@ protected:
 	virtual void OnSlotRemoved(UPanelSlot* Slot) override;
 	// End UPanelWidget
 
-protected:
-	TSharedPtr<SBox> MySizeBox;
-
-protected:
 	// UWidget interface
 	virtual TSharedRef<SWidget> RebuildWidget() override;
 	// End of UWidget interface
diff --git a/Engine/Source/Runtime/UMG/Public/Components/SizeBoxSlot.h b/Engine/Source/Runtime/UMG/Public/Components/SizeBoxSlot.h
--- a/Engine/Source/Runtime/UMG/Public/Components/SizeBoxSlot.h
+++ b/Engine/Source/Runtime/UMG/Public/Components/SizeBoxSlot.h
@@ -23,6 +23,11 @@ public:
 	UPROPERTY(EditAnywhere, BlueprintReadOnly, Category="Layout|SizeBox Slot")
 	FMargin Padding;
 
+private:
+	/** A pointer to the button to allow us to adjust the size, padding...etc at runtime. */
+	TWeakPtr<SBox> SizeBox;
+
+public:
 	/** The alignment of the object horizontally. */
 	UPROPERTY(EditAnywhere, BlueprintReadOnly, Category="Layout|SizeBox Slot")
 	TEnumAsByte<EHorizontalAlignment> HorizontalAlignment;
@@ -31,7 +36,6 @@ public:
 	UPROPERTY(EditAnywhere, BlueprintReadOnly, Category="Layout|SizeBox Slot")
 	TEnumAsByte<EVerticalAlignment> VerticalAlignment;
 
-public:
 
 	UFUNCTION(BlueprintCallable, Category="Layout|SizeBox Slot")
 	void SetPadding(FMargin InPadding);
@@ -42,8 +46,6 @@ public:
 	UFUNCTION(BlueprintCallable, Category="Layout|SizeBox Slot")
 	void SetVerticalAlignment(EVerticalAlignment InVerticalAlignment);
 
-public:
-
 	// UPanelSlot interface
 	virtual void SynchronizeProperties() override;
 	// End of UPanelSlot interface
@@ -52,9 +54,4 @@ public:
 	void BuildSlot(TSharedRef<SBox> InSizeBox);
 
 	virtual void ReleaseSlateResources(bool bReleaseChildren) override;
-
-private:
-
-	/** A pointer to the button to allow us to adjust the size, padding...etc at runtime. */
-	TWeakPtr<SBox> SizeBox;
 };
