Fix nonunity compile

diff --git a/Engine/Source/Runtime/UMG/Private/Components/WidgetInteractionComponent.cpp b/Engine/Source/Runtime/UMG/Private/Components/WidgetInteractionComponent.cpp
--- a/Engine/Source/Runtime/UMG/Private/Components/WidgetInteractionComponent.cpp
+++ b/Engine/Source/Runtime/UMG/Private/Components/WidgetInteractionComponent.cpp
@@ -6,6 +6,7 @@
 #include "Kismet/KismetSystemLibrary.h"
 #include "Components/ArrowComponent.h"
 #include "Framework/Application/SlateApplication.h"
+#include "Framework/Application/SlateUser.h"
 #include "Kismet/GameplayStatics.h"
 #include "DrawDebugHelpers.h"
 
