Editor - The material thumbnail renderer no longer attempts to use an in world render to draw UI Materials, instead it uses the WidgetRenderer to draw the SlateMaterial like it would be for real.  Unselecting the material in the material editor when changing materials.

Matt.Kuhlenschmidt, Vincent.Gauthier

diff --git a/Engine/Source/Editor/MaterialEditor/Private/MaterialEditor.cpp b/Engine/Source/Editor/MaterialEditor/Private/MaterialEditor.cpp
--- a/Engine/Source/Editor/MaterialEditor/Private/MaterialEditor.cpp
+++ b/Engine/Source/Editor/MaterialEditor/Private/MaterialEditor.cpp
@@ -1598,9 +1598,19 @@ void FMaterialEditor::SetPreviewMaterial(UMaterialInterface* InMaterialInterface
 		{
 			PreviewUIViewport->SetPreviewMaterial(InMaterialInterface);
 		}
+
+		if (PreviewViewport.IsValid())
+		{
+			PreviewViewport->SetPreviewMaterial(nullptr);
+		}
 	}
 	else
 	{
+		if (PreviewUIViewport.IsValid())
+		{
+			PreviewUIViewport->SetPreviewMaterial(nullptr);
+		}
+
 		if (PreviewViewport.IsValid())
 		{
 			PreviewViewport->SetPreviewMaterial(InMaterialInterface);
diff --git a/Engine/Source/Editor/MaterialEditor/Private/SMaterialEditorViewport.cpp b/Engine/Source/Editor/MaterialEditor/Private/SMaterialEditorViewport.cpp
--- a/Engine/Source/Editor/MaterialEditor/Private/SMaterialEditorViewport.cpp
+++ b/Engine/Source/Editor/MaterialEditor/Private/SMaterialEditorViewport.cpp
@@ -427,7 +427,7 @@ void SMaterialEditor3DPreviewViewport::SetPreviewMaterial(UMaterialInterface* In
 	PreviewMaterial = InMaterialInterface;
 
 	// Spawn post processing volume actor if the material has post processing as domain.
-	if (PreviewMaterial->GetMaterial()->IsPostProcessMaterial())
+	if (PreviewMaterial && PreviewMaterial->GetMaterial()->IsPostProcessMaterial())
 	{
 		if (PostProcessVolumeActor == nullptr)
 		{
@@ -456,7 +456,12 @@ void SMaterialEditor3DPreviewViewport::SetPreviewMaterial(UMaterialInterface* In
 		if (PreviewMeshComponent != nullptr)
 		{
 			PreviewMeshComponent->OverrideMaterials.Empty();
-			PreviewMeshComponent->OverrideMaterials.Add(PreviewMaterial);
+
+			if (PreviewMaterial)
+			{
+				PreviewMeshComponent->OverrideMaterials.Add(PreviewMaterial);
+			}
+
 			PreviewMeshComponent->MarkRenderStateDirty();
 		}
 		
diff --git a/Engine/Source/Runtime/Slate/Public/SlateMaterialBrush.h b/Engine/Source/Runtime/Slate/Public/SlateMaterialBrush.h
--- a/Engine/Source/Runtime/Slate/Public/SlateMaterialBrush.h
+++ b/Engine/Source/Runtime/Slate/Public/SlateMaterialBrush.h
@@ -22,7 +22,7 @@ struct FSlateMaterialBrush : public FSlateBrush
 	 * @param InImageSize The material's dimensions.
 	 */
 	FSlateMaterialBrush( class UMaterialInterface& InMaterial, const FVector2D& InImageSize )
-		: FSlateBrush( ESlateBrushDrawType::Image, FName(TEXT("None")), FMargin(0), ESlateBrushTileType::NoTile, ESlateBrushImageType::FullColor, InImageSize, FLinearColor::White, &InMaterial )
+		: FSlateBrush( ESlateBrushDrawType::Image, NAME_None, FMargin(0), ESlateBrushTileType::NoTile, ESlateBrushImageType::FullColor, InImageSize, FLinearColor::White, &InMaterial )
 	{
 		ResourceName = FName( *InMaterial.GetFullName() );
 	}
diff --git a/Engine/Source/Runtime/UMG/Private/Slate/WidgetRenderer.cpp b/Engine/Source/Runtime/UMG/Private/Slate/WidgetRenderer.cpp
--- a/Engine/Source/Runtime/UMG/Private/Slate/WidgetRenderer.cpp
+++ b/Engine/Source/Runtime/UMG/Private/Slate/WidgetRenderer.cpp
@@ -96,23 +96,17 @@ UTextureRenderTarget2D* FWidgetRenderer::CreateTargetFor(FVector2D DrawSize, Tex
 	return nullptr;
 }
 
-void FWidgetRenderer::DrawWidget(FRenderTarget* RenderTarget, const TSharedRef<SWidget>& Widget, FVector2D DrawSize, float DeltaTime, bool bDeferRenderTargetUpdate)
-{
-	DrawWidget(RenderTarget, Widget, 1.f, DrawSize, DeltaTime, bDeferRenderTargetUpdate);
-}
-
 void FWidgetRenderer::DrawWidget(UTextureRenderTarget2D* RenderTarget, const TSharedRef<SWidget>& Widget, FVector2D DrawSize, float DeltaTime, bool bDeferRenderTargetUpdate)
 {
-	DrawWidget(RenderTarget->GameThread_GetRenderTargetResource(), Widget, DrawSize, DeltaTime, bDeferRenderTargetUpdate);
+	TSharedRef<SVirtualWindow> Window = SNew(SVirtualWindow).Size(DrawSize);
+	TUniquePtr<FHittestGrid> HitTestGrid = MakeUnique<FHittestGrid>();
+
+	Window->SetContent(Widget);
+	Window->Resize(DrawSize);
+
 }
 
-void FWidgetRenderer::DrawWidget(
-	FRenderTarget* RenderTarget,
-	const TSharedRef<SWidget>& Widget,
-	float Scale,
-	FVector2D DrawSize,
-	float DeltaTime,
-	bool bDeferRenderTargetUpdate)
+void FWidgetRenderer::DrawWidget(FRenderTarget* RenderTarget, const TSharedRef<SWidget>& Widget, FVector2D DrawSize, float DeltaTime, bool bDeferRenderTargetUpdate)
 {
 	TSharedRef<SVirtualWindow> Window = SNew(SVirtualWindow).Size(DrawSize);
 	TUniquePtr<FHittestGrid> HitTestGrid = MakeUnique<FHittestGrid>();
@@ -120,22 +114,24 @@ void FWidgetRenderer::DrawWidget(
 	Window->SetContent(Widget);
 	Window->Resize(DrawSize);
 
-	DrawWindow(RenderTarget, *HitTestGrid, Window, Scale, DrawSize, DeltaTime, bDeferRenderTargetUpdate);
-}
+	const float Scale = 1;
+	FGeometry WindowGeometry = FGeometry::MakeRoot(DrawSize * (1 / Scale), FSlateLayoutTransform(Scale));
 
-void FWidgetRenderer::DrawWidget(
-	UTextureRenderTarget2D* RenderTarget,
-	const TSharedRef<SWidget>& Widget,
-	float Scale,
-	FVector2D DrawSize,
-	float DeltaTime,
-	bool bDeferRenderTargetUpdate)
-{
-	DrawWidget(RenderTarget->GameThread_GetRenderTargetResource(), Widget, Scale, DrawSize, DeltaTime, bDeferRenderTargetUpdate);
+	FPaintArgs PaintArgs(nullptr, *HitTestGrid, FVector2D::ZeroVector, FApp::GetCurrentTime(), DeltaTime);
+
+	DrawWindow
+	(
+		PaintArgs,
+		RenderTarget,
+		Window,
+		WindowGeometry,
+		WindowGeometry.GetLayoutBoundingRect(),
+		DeltaTime,
+		bDeferRenderTargetUpdate
+	);
 }
 
 void FWidgetRenderer::DrawWindow(
-	FRenderTarget* RenderTarget,
 	FHittestGrid& HitTestGrid,
 	TSharedRef<SWindow> Window,
 	float Scale,
@@ -161,18 +157,6 @@ void FWidgetRenderer::DrawWindow(
 	UTextureRenderTarget2D* RenderTarget,
 	FHittestGrid& HitTestGrid,
 	TSharedRef<SWindow> Window,
-	float Scale,
-	FVector2D DrawSize,
-	float DeltaTime,
-	bool bDeferRenderTargetUpdate)
-{
-	DrawWindow(RenderTarget->GameThread_GetRenderTargetResource(), HitTestGrid, Window, Scale, DrawSize, DeltaTime, bDeferRenderTargetUpdate);
-}
-
-void FWidgetRenderer::DrawWindow(
-	FRenderTarget* RenderTarget,
-	FHittestGrid& HitTestGrid,
-	TSharedRef<SWindow> Window,
 	FGeometry WindowGeometry,
 	FSlateRect WindowClipRect,
 	float DeltaTime,
@@ -183,15 +167,14 @@ void FWidgetRenderer::DrawWindow(
 }
 
 void FWidgetRenderer::DrawWindow(
-	UTextureRenderTarget2D* RenderTarget,
-	FHittestGrid& HitTestGrid,
+	const FPaintArgs& PaintArgs,
 	TSharedRef<SWindow> Window,
 	FGeometry WindowGeometry,
 	FSlateRect WindowClipRect,
 	float DeltaTime,
 	bool bDeferRenderTargetUpdate)
 {
-	DrawWindow(RenderTarget->GameThread_GetRenderTargetResource(), HitTestGrid, Window, WindowGeometry, WindowClipRect, DeltaTime, bDeferRenderTargetUpdate);
+	DrawWindow(PaintArgs, RenderTarget->GameThread_GetRenderTargetResource(), Window, WindowGeometry, WindowClipRect, DeltaTime, bDeferRenderTargetUpdate);
 }
 
 void FWidgetRenderer::DrawWindow(
