Undo revision 20.

diff --git a/Engine/Source/Runtime/UMG/Private/Slate/WidgetRenderer.cpp b/Engine/Source/Runtime/UMG/Private/Slate/WidgetRenderer.cpp
--- a/Engine/Source/Runtime/UMG/Private/Slate/WidgetRenderer.cpp
+++ b/Engine/Source/Runtime/UMG/Private/Slate/WidgetRenderer.cpp
@@ -96,17 +96,23 @@ UTextureRenderTarget2D* FWidgetRenderer::CreateTargetFor(FVector2D DrawSize, Tex
 	return nullptr;
 }
 
-void FWidgetRenderer::DrawWidget(UTextureRenderTarget2D* RenderTarget, const TSharedRef<SWidget>& Widget, FVector2D DrawSize, float DeltaTime, bool bDeferRenderTargetUpdate)
+void FWidgetRenderer::DrawWidget(FRenderTarget* RenderTarget, const TSharedRef<SWidget>& Widget, FVector2D DrawSize, float DeltaTime, bool bDeferRenderTargetUpdate)
 {
-	TSharedRef<SVirtualWindow> Window = SNew(SVirtualWindow).Size(DrawSize);
-	TUniquePtr<FHittestGrid> HitTestGrid = MakeUnique<FHittestGrid>();
-
-	Window->SetContent(Widget);
-	Window->Resize(DrawSize);
+	DrawWidget(RenderTarget, Widget, 1.f, DrawSize, DeltaTime, bDeferRenderTargetUpdate);
+}
 
+void FWidgetRenderer::DrawWidget(UTextureRenderTarget2D* RenderTarget, const TSharedRef<SWidget>& Widget, FVector2D DrawSize, float DeltaTime, bool bDeferRenderTargetUpdate)
+{
+	DrawWidget(RenderTarget->GameThread_GetRenderTargetResource(), Widget, DrawSize, DeltaTime, bDeferRenderTargetUpdate);
 }
 
-void FWidgetRenderer::DrawWidget(FRenderTarget* RenderTarget, const TSharedRef<SWidget>& Widget, FVector2D DrawSize, float DeltaTime, bool bDeferRenderTargetUpdate)
+void FWidgetRenderer::DrawWidget(
+	FRenderTarget* RenderTarget,
+	const TSharedRef<SWidget>& Widget,
+	float Scale,
+	FVector2D DrawSize,
+	float DeltaTime,
+	bool bDeferRenderTargetUpdate)
 {
 	TSharedRef<SVirtualWindow> Window = SNew(SVirtualWindow).Size(DrawSize);
 	TUniquePtr<FHittestGrid> HitTestGrid = MakeUnique<FHittestGrid>();
@@ -114,24 +120,22 @@ void FWidgetRenderer::DrawWidget(FRenderTarget* RenderTarget, const TSharedRef<S
 	Window->SetContent(Widget);
 	Window->Resize(DrawSize);
 
-	const float Scale = 1;
-	FGeometry WindowGeometry = FGeometry::MakeRoot(DrawSize * (1 / Scale), FSlateLayoutTransform(Scale));
-
-	FPaintArgs PaintArgs(nullptr, *HitTestGrid, FVector2D::ZeroVector, FApp::GetCurrentTime(), DeltaTime);
+	DrawWindow(RenderTarget, *HitTestGrid, Window, Scale, DrawSize, DeltaTime, bDeferRenderTargetUpdate);
+}
 
-	DrawWindow
-	(
-		PaintArgs,
-		RenderTarget,
-		Window,
-		WindowGeometry,
-		WindowGeometry.GetLayoutBoundingRect(),
-		DeltaTime,
-		bDeferRenderTargetUpdate
-	);
+void FWidgetRenderer::DrawWidget(
+	UTextureRenderTarget2D* RenderTarget,
+	const TSharedRef<SWidget>& Widget,
+	float Scale,
+	FVector2D DrawSize,
+	float DeltaTime,
+	bool bDeferRenderTargetUpdate)
+{
+	DrawWidget(RenderTarget->GameThread_GetRenderTargetResource(), Widget, Scale, DrawSize, DeltaTime, bDeferRenderTargetUpdate);
 }
 
 void FWidgetRenderer::DrawWindow(
+	FRenderTarget* RenderTarget,
 	FHittestGrid& HitTestGrid,
 	TSharedRef<SWindow> Window,
 	float Scale,
@@ -157,6 +161,18 @@ void FWidgetRenderer::DrawWindow(
 	UTextureRenderTarget2D* RenderTarget,
 	FHittestGrid& HitTestGrid,
 	TSharedRef<SWindow> Window,
+	float Scale,
+	FVector2D DrawSize,
+	float DeltaTime,
+	bool bDeferRenderTargetUpdate)
+{
+	DrawWindow(RenderTarget->GameThread_GetRenderTargetResource(), HitTestGrid, Window, Scale, DrawSize, DeltaTime, bDeferRenderTargetUpdate);
+}
+
+void FWidgetRenderer::DrawWindow(
+	FRenderTarget* RenderTarget,
+	FHittestGrid& HitTestGrid,
+	TSharedRef<SWindow> Window,
 	FGeometry WindowGeometry,
 	FSlateRect WindowClipRect,
 	float DeltaTime,
@@ -167,14 +183,15 @@ void FWidgetRenderer::DrawWindow(
 }
 
 void FWidgetRenderer::DrawWindow(
-	const FPaintArgs& PaintArgs,
+	UTextureRenderTarget2D* RenderTarget,
+	FHittestGrid& HitTestGrid,
 	TSharedRef<SWindow> Window,
 	FGeometry WindowGeometry,
 	FSlateRect WindowClipRect,
 	float DeltaTime,
 	bool bDeferRenderTargetUpdate)
 {
-	DrawWindow(PaintArgs, RenderTarget->GameThread_GetRenderTargetResource(), Window, WindowGeometry, WindowClipRect, DeltaTime, bDeferRenderTargetUpdate);
+	DrawWindow(RenderTarget->GameThread_GetRenderTargetResource(), HitTestGrid, Window, WindowGeometry, WindowClipRect, DeltaTime, bDeferRenderTargetUpdate);
 }
 
 void FWidgetRenderer::DrawWindow(
