Updated WidgetComponent to no longer display screen-space widgets when owning actor is hidden

#rn Updated WidgetComponent to no longer add screen-space widgets to the viewport if the owning actor is set to be hidden in game

diff --git a/Engine/Source/Runtime/UMG/Private/Components/WidgetComponent.cpp b/Engine/Source/Runtime/UMG/Private/Components/WidgetComponent.cpp
--- a/Engine/Source/Runtime/UMG/Private/Components/WidgetComponent.cpp
+++ b/Engine/Source/Runtime/UMG/Private/Components/WidgetComponent.cpp
@@ -996,7 +996,7 @@ void UWidgetComponent::TickComponent(float DeltaTime, enum ELevelTick TickType,
 				ULocalPlayer* TargetPlayer = GetOwnerPlayer();
 				APlayerController* PlayerController = TargetPlayer ? TargetPlayer->PlayerController : nullptr;
 
-				if ( TargetPlayer && PlayerController && IsVisible() )
+				if ( TargetPlayer && PlayerController && IsVisible() && !(GetOwner()->bHidden))
 				{
 					if ( !bAddedToScreen )
 					{
