Force a clear of the rendering of a WidgetComponent when the WidgetClass is set to null.

diff --git a/Engine/Source/Runtime/UMG/Private/Components/WidgetComponent.cpp b/Engine/Source/Runtime/UMG/Private/Components/WidgetComponent.cpp
--- a/Engine/Source/Runtime/UMG/Private/Components/WidgetComponent.cpp
+++ b/Engine/Source/Runtime/UMG/Private/Components/WidgetComponent.cpp
@@ -581,6 +581,7 @@ UWidgetComponent::UWidgetComponent( const FObjectInitializer& PCIP )
 	, LayerZOrder(-100)
 	, GeometryMode(EWidgetGeometryMode::Plane)
 	, CylinderArcAngle( 180.0f )
+    , bRenderCleared(false)
 {
 	PrimaryComponentTick.bCanEverTick = true;
 	bTickInEditor = true;
@@ -971,9 +972,10 @@ void UWidgetComponent::TickComponent(float DeltaTime, enum ELevelTick TickType,
 	{
 		UpdateWidget();
 
-		if ( Widget == nullptr && !SlateWidget.IsValid() )
+		if (Widget == nullptr && !SlateWidget.IsValid() && bRenderCleared)
 		{
-			return;
+			// We will enter here if the WidgetClass is empty and we already renderered an empty widget. No need to continue.
+			return;	
 		}
 
 	    if ( Space != EWidgetSpace::Screen )
@@ -985,6 +987,12 @@ void UWidgetComponent::TickComponent(float DeltaTime, enum ELevelTick TickType,
 				// a different rate than the rest of the world.
 				const float DeltaTimeFromLastDraw = LastWidgetRenderTime == 0 ? 0 : (GetCurrentTime() - LastWidgetRenderTime);
 				DrawWidgetToRenderTarget(DeltaTimeFromLastDraw);
+
+				// We draw an empty widget.
+				if (Widget == nullptr && !SlateWidget.IsValid())
+				{
+					bRenderCleared = true;
+				}
 		    }
 	    }
 	    else
@@ -1443,6 +1451,7 @@ void UWidgetComponent::UpdateWidget()
 				if (CurrentSlateWidget != SNullWidget::NullWidget)
 				{
 					CurrentSlateWidget = SNullWidget::NullWidget;
+					bRenderCleared = false;
 					bWidgetChanged = true;
 				}
 				SlateWindow->SetContent( SNullWidget::NullWidget );
diff --git a/Engine/Source/Runtime/UMG/Public/Components/WidgetComponent.h b/Engine/Source/Runtime/UMG/Public/Components/WidgetComponent.h
--- a/Engine/Source/Runtime/UMG/Public/Components/WidgetComponent.h
+++ b/Engine/Source/Runtime/UMG/Public/Components/WidgetComponent.h
@@ -570,6 +570,10 @@ protected:
 
 	/** Helper class for drawing widgets to a render target. */
 	class FWidgetRenderer* WidgetRenderer;
+
+private: 
+	/** Set to true after a draw of an empty component.*/
+	bool bRenderCleared;
 };
 
 USTRUCT()
