PR #5602: Added GetText() method to URichTextBlock (Contributed by [REDACTED_PERSON])

diff --git a/Engine/Source/Runtime/UMG/Private/Components/RichTextBlock.cpp b/Engine/Source/Runtime/UMG/Private/Components/RichTextBlock.cpp
--- a/Engine/Source/Runtime/UMG/Private/Components/RichTextBlock.cpp
+++ b/Engine/Source/Runtime/UMG/Private/Components/RichTextBlock.cpp
@@ -97,6 +97,16 @@ void URichTextBlock::UpdateStyleData()
 	}
 }
 
+FText URichTextBlock::GetText() const
+{
+	if (MyRichTextBlock.IsValid())
+	{
+		return MyRichTextBlock->GetText();
+	}
+
+	return Text;
+}
+
 void URichTextBlock::SetText(const FText& InText)
 {
 	Text = InText;
diff --git a/Engine/Source/Runtime/UMG/Public/Components/RichTextBlock.h b/Engine/Source/Runtime/UMG/Public/Components/RichTextBlock.h
--- a/Engine/Source/Runtime/UMG/Public/Components/RichTextBlock.h
+++ b/Engine/Source/Runtime/UMG/Public/Components/RichTextBlock.h
@@ -124,6 +124,12 @@ public:
 	// End UWidget interface
 #endif
 
+	/**
+	 * Returns widgets text.
+	 */
+	UFUNCTION(BlueprintCallable, Category = "Widget")
+	FText GetText() const;
+	
 	/**
 	 * Directly sets the widget text.
 	 * Warning: This will wipe any binding created for the Text property!
