Fixed leaks in RetainerWidget and InvalidationPanel

diff --git a/Engine/Source/Runtime/Slate/Private/Widgets/SInvalidationPanel.cpp b/Engine/Source/Runtime/Slate/Private/Widgets/SInvalidationPanel.cpp
--- a/Engine/Source/Runtime/Slate/Private/Widgets/SInvalidationPanel.cpp
+++ b/Engine/Source/Runtime/Slate/Private/Widgets/SInvalidationPanel.cpp
@@ -125,6 +125,7 @@ SInvalidationPanel::~SInvalidationPanel()
 
 	if ( FSlateApplication::IsInitialized() )
 	{
+		FSlateApplicationBase::Get().OnGlobalInvalidate().RemoveAll(this);
 		FSlateApplication::Get().ReleaseResourcesForLayoutCache(this);
 	}
 }
diff --git a/Engine/Source/Runtime/UMG/Private/Slate/SRetainerWidget.cpp b/Engine/Source/Runtime/UMG/Private/Slate/SRetainerWidget.cpp
--- a/Engine/Source/Runtime/UMG/Private/Slate/SRetainerWidget.cpp
+++ b/Engine/Source/Runtime/UMG/Private/Slate/SRetainerWidget.cpp
@@ -88,8 +88,14 @@ SRetainerWidget::SRetainerWidget()
 
 SRetainerWidget::~SRetainerWidget()
 {
+	for (int32 i = 0; i < NodePool.Num(); i++)
+	{
+		delete NodePool[i];
+	}
+
 	if( FSlateApplication::IsInitialized() )
 	{
+		FSlateApplicationBase::Get().OnGlobalInvalidate().RemoveAll( this );
 #if !UE_BUILD_SHIPPING
 		OnRetainerModeChangedDelegate.RemoveAll( this );
 #endif
