#rn UMG WidgetComponent default behavior was changed to handle hit test differently. The inner SlateWindow used by the component now use visibility SelfHitTestInvisible by Default which means that a hit test in a empty portion of the WidgetCompoent window will let the mous hit go through. All data created with previous versions will keep the previous default value of Visible.

diff --git a/Engine/Source/Runtime/Core/Public/UObject/EditorObjectVersion.h b/Engine/Source/Runtime/Core/Public/UObject/EditorObjectVersion.h
--- a/Engine/Source/Runtime/Core/Public/UObject/EditorObjectVersion.h
+++ b/Engine/Source/Runtime/Core/Public/UObject/EditorObjectVersion.h
@@ -71,6 +71,8 @@ struct CORE_API FEditorObjectVersion
 		MeshDescriptionBulkDataGuid,
 		// Change to MeshDescription serialization (removed FMeshPolygon::HoleContours)
 		MeshDescriptionRemovedHoles,
+		// Change to the WidgetCompoent WindowVisibilty default value
+		ChangedWidgetComponentWindowVisibilityDefault,
 		// -----<new versions can be added above this line>-------------------------------------------------
 		VersionPlusOne,
 		LatestVersion = VersionPlusOne - 1
diff --git a/Engine/Source/Runtime/UMG/Private/Components/WidgetComponent.cpp b/Engine/Source/Runtime/UMG/Private/Components/WidgetComponent.cpp
--- a/Engine/Source/Runtime/UMG/Private/Components/WidgetComponent.cpp
+++ b/Engine/Source/Runtime/UMG/Private/Components/WidgetComponent.cpp
@@ -570,6 +570,7 @@ UWidgetComponent::UWidgetComponent( const FObjectInitializer& PCIP )
 	, LastWidgetRenderTime(0)
 	, bReceiveHardwareInput(false)
 	, bWindowFocusable(true)
+	, WindowVisibility(EWindowVisibility::SelfHitTestInvisible)
 	, bApplyGammaCorrection(false)
 	, BackgroundColor( FLinearColor::Transparent )
 	, TintColorAndOpacity( FLinearColor::White )
@@ -619,6 +620,19 @@ UWidgetComponent::UWidgetComponent( const FObjectInitializer& PCIP )
 	bAddedToScreen = false;
 }
 
+void UWidgetComponent::Serialize(FArchive& Ar)
+{
+	Super::Serialize(Ar);
+
+	Ar.UsingCustomVersion(FEditorObjectVersion::GUID);
+
+	if (Ar.CustomVer(FEditorObjectVersion::GUID) < FEditorObjectVersion::ChangedWidgetComponentWindowVisibilityDefault)
+	{
+		// Reset the default value for visibility
+		WindowVisibility = EWindowVisibility::Visible;
+	}
+}
+
 void UWidgetComponent::BeginPlay()
 {
 	Super::BeginPlay();
@@ -823,6 +837,38 @@ void UWidgetComponent::OnRegister()
 #endif // !UE_SERVER
 }
 
+void UWidgetComponent::SetWindowFocusable(bool bInWindowFocusable)
+{
+	bWindowFocusable = bInWindowFocusable;
+ 	if (SlateWindow.IsValid())
+ 	{
+ 		SlateWindow->SetIsFocusable(bWindowFocusable);
+ 	}
+};
+
+EVisibility UWidgetComponent::ConvertWindowVisibilityToVisibility(EWindowVisibility visibility)
+{
+	switch (visibility)
+	{
+	case EWindowVisibility::Visible:
+		return EVisibility::Visible;
+	case EWindowVisibility::SelfHitTestInvisible:
+		return EVisibility::SelfHitTestInvisible;
+	default:
+		checkNoEntry();
+		return EVisibility::SelfHitTestInvisible;
+	}	
+}
+
+void UWidgetComponent::SetWindowVisibility(EWindowVisibility InVisibility)
+{
+	WindowVisibility = InVisibility;
+ 	if (SlateWindow.IsValid())
+ 	{		
+ 		SlateWindow->SetVisibility(ConvertWindowVisibilityToVisibility(WindowVisibility));
+ 	}
+}
+
 bool UWidgetComponent::CanReceiveHardwareInput() const
 {
 	return bReceiveHardwareInput && GeometryMode == EWidgetGeometryMode::Plane;
@@ -1231,6 +1277,7 @@ bool UWidgetComponent::CanEditChange(const UProperty* InProperty) const
 		if ( PropertyName == GET_MEMBER_NAME_STRING_CHECKED(UWidgetComponent, GeometryMode) ||
 			 PropertyName == GET_MEMBER_NAME_STRING_CHECKED(UWidgetComponent, TimingPolicy) ||
 			 PropertyName == GET_MEMBER_NAME_STRING_CHECKED(UWidgetComponent, bWindowFocusable) ||
+			 PropertyName == GET_MEMBER_NAME_STRING_CHECKED(UWidgetComponent, WindowVisibility) ||
 			 PropertyName == GET_MEMBER_NAME_STRING_CHECKED(UWidgetComponent, bManuallyRedraw) ||
 			 PropertyName == GET_MEMBER_NAME_STRING_CHECKED(UWidgetComponent, RedrawTime) ||
 			 PropertyName == GET_MEMBER_NAME_STRING_CHECKED(UWidgetComponent, BackgroundColor) ||
@@ -1275,6 +1322,8 @@ void UWidgetComponent::PostEditChangeProperty(FPropertyChangedEvent& PropertyCha
 		static FName BlendModeName( TEXT( "BlendMode" ) );
 		static FName GeometryModeName( TEXT("GeometryMode") );
 		static FName CylinderArcAngleName( TEXT("CylinderArcAngle") );
+		static FName bWindowFocusableName(TEXT("bWindowFocusable"));
+		static FName WindowVisibilityName(TEXT("WindowVisibility"));
 
 		auto PropertyName = Property->GetFName();
 
@@ -1303,6 +1352,15 @@ void UWidgetComponent::PostEditChangeProperty(FPropertyChangedEvent& PropertyCha
 		{
 			MarkRenderStateDirty();
 		}
+		else if (PropertyName == bWindowFocusableName)
+		{
+			SetWindowFocusable(bWindowFocusable);
+		}
+		else if (PropertyName == WindowVisibilityName)
+		{
+			SetWindowVisibility(WindowVisibility);
+		}
+
 	}
 
 	Super::PostEditChangeProperty(PropertyChangedEvent);
@@ -1420,6 +1478,7 @@ void UWidgetComponent::UpdateWidget()
 
 				SlateWindow = SNew(SVirtualWindow).Size(CurrentDrawSize);
 				SlateWindow->SetIsFocusable(bWindowFocusable);
+				SlateWindow->SetVisibility(ConvertWindowVisibilityToVisibility(WindowVisibility));
 				RegisterWindow();
 
 				bNeededNewWindow = true;
diff --git a/Engine/Source/Runtime/UMG/Public/Components/WidgetComponent.h b/Engine/Source/Runtime/UMG/Public/Components/WidgetComponent.h
--- a/Engine/Source/Runtime/UMG/Public/Components/WidgetComponent.h
+++ b/Engine/Source/Runtime/UMG/Public/Components/WidgetComponent.h
@@ -56,6 +56,16 @@ enum class EWidgetGeometryMode : uint8
 	Cylinder
 };
 
+UENUM(BlueprintType)
+enum class EWindowVisibility : uint8
+{
+	/** The window visibility is Visible */
+	Visible,
+
+	/** The window visibility is SelfHitTestInvisible */
+	SelfHitTestInvisible
+};
+
 
 /**
  * The widget component provides a surface in the 3D environment on which to render widgets normally rendered to the screen.
@@ -73,6 +83,10 @@ class UMG_API UWidgetComponent : public UMeshComponent
 	GENERATED_UCLASS_BODY()
 
 public:
+	//UObject interface
+	virtual void Serialize(FArchive& Ar) override;
+	//~ End UObject Interface
+
 	/** UActorComponent Interface */
 	virtual void BeginPlay() override;
 	virtual void EndPlay(const EEndPlayReason::Type EndPlayReason) override;
@@ -344,10 +358,18 @@ public:
 
 	/** @see bWindowFocusable */
 	UFUNCTION(BlueprintCallable, Category = UserInterface)
-	void SetWindowFocusable(bool bInWindowFocusable)
+	void SetWindowFocusable(bool bInWindowFocusable);
+
+	/** Gets the visibility of the virtual window created to host the widget focusable. */
+	UFUNCTION(BlueprintCallable, Category = UserInterface)
+	EWindowVisibility GetWindowVisiblility() const
 	{
-		bWindowFocusable = bInWindowFocusable;
-	};
+		return WindowVisibility;
+	}
+
+	/** Sets the visibility of the virtual window created to host the widget focusable. */
+	UFUNCTION(BlueprintCallable, Category = UserInterface)
+	void SetWindowVisibility(EWindowVisibility InVisibility);
 
 protected:
 	/** Just because the user attempts to receive hardware input does not mean it's possible. */
@@ -448,6 +470,10 @@ protected:
 	UPROPERTY(EditAnywhere, Category=Interaction)
 	bool bWindowFocusable;
 
+	/** The visibility of the virtual window created to host the widget */
+	UPROPERTY(EditAnywhere, Category = Interaction)
+	EWindowVisibility WindowVisibility;
+
 	/**
 	 * Widget components that appear in the world will be gamma corrected by the 3D renderer.
 	 * In some cases, widget components are blitted directly into the backbuffer, in which case gamma correction should be enabled.
@@ -572,6 +598,7 @@ protected:
 	class FWidgetRenderer* WidgetRenderer;
 
 private: 
+	static EVisibility ConvertWindowVisibilityToVisibility(EWindowVisibility visibility);
 	/** Set to true after a draw of an empty component.*/
 	bool bRenderCleared;
 };
