PR #5504: UMG Play Animation nodes with Finished output pin (Contributed by [REDACTED_PERSON])

diff --git a/Engine/Source/Editor/UMGEditor/Private/Nodes/K2Node_PlayAnimation.cpp b/Engine/Source/Editor/UMGEditor/Private/Nodes/K2Node_PlayAnimation.cpp
new file mode 100644
--- /dev/null
+++ b/Engine/Source/Editor/UMGEditor/Private/Nodes/K2Node_PlayAnimation.cpp
@@ -0,0 +1,24 @@
+// Copyright 1998-2019 Epic Games, Inc. All Rights Reserved.
+
+#include "K2Node_PlayAnimation.h"
+#include "Animation/WidgetAnimationPlayCallbackProxy.h"
+
+#define LOCTEXT_NAMESPACE "UMG"
+
+UK2Node_PlayAnimation::UK2Node_PlayAnimation(const FObjectInitializer& ObjectInitializer)
+	: Super(ObjectInitializer)
+{
+	ProxyFactoryFunctionName = GET_FUNCTION_NAME_CHECKED(UWidgetAnimationPlayCallbackProxy, CreatePlayAnimationProxyObject);
+	ProxyFactoryClass = UWidgetAnimationPlayCallbackProxy::StaticClass();
+	ProxyClass = UWidgetAnimationPlayCallbackProxy::StaticClass();
+}
+
+UK2Node_PlayAnimationTimeRange::UK2Node_PlayAnimationTimeRange(const FObjectInitializer& ObjectInitializer)
+	: Super(ObjectInitializer)
+{
+	ProxyFactoryFunctionName = GET_FUNCTION_NAME_CHECKED(UWidgetAnimationPlayCallbackProxy, CreatePlayAnimationTimeRangeProxyObject);
+	ProxyFactoryClass = UWidgetAnimationPlayCallbackProxy::StaticClass();
+	ProxyClass = UWidgetAnimationPlayCallbackProxy::StaticClass();
+}
+
+#undef LOCTEXT_NAMESPACE
diff --git a/Engine/Source/Editor/UMGEditor/Private/Nodes/K2Node_PlayAnimation.h b/Engine/Source/Editor/UMGEditor/Private/Nodes/K2Node_PlayAnimation.h
new file mode 100644
--- /dev/null
+++ b/Engine/Source/Editor/UMGEditor/Private/Nodes/K2Node_PlayAnimation.h
@@ -0,0 +1,18 @@
+// Copyright 1998-2019 Epic Games, Inc. All Rights Reserved.
+
+#pragma once
+
+#include "K2Node_BaseAsyncTask.h"
+#include "K2Node_PlayAnimation.generated.h"
+
+UCLASS()
+class UMGEDITOR_API UK2Node_PlayAnimation : public UK2Node_BaseAsyncTask
+{
+	GENERATED_UCLASS_BODY()
+};
+
+UCLASS()
+class UMGEDITOR_API UK2Node_PlayAnimationTimeRange : public UK2Node_BaseAsyncTask
+{
+	GENERATED_UCLASS_BODY()
+};
diff --git a/Engine/Source/Runtime/UMG/Private/Animation/WidgetAnimationPlayCallbackProxy.cpp b/Engine/Source/Runtime/UMG/Private/Animation/WidgetAnimationPlayCallbackProxy.cpp
new file mode 100644
--- /dev/null
+++ b/Engine/Source/Runtime/UMG/Private/Animation/WidgetAnimationPlayCallbackProxy.cpp
@@ -0,0 +1,85 @@
+// Copyright 1998-2019 Epic Games, Inc. All Rights Reserved.
+
+#include "Animation/WidgetAnimationPlayCallbackProxy.h"
+#include "Animation/UMGSequencePlayer.h"
+#include "Engine/World.h"
+#include "TimerManager.h"
+
+#define LOCTEXT_NAMESPACE "UMG"
+
+UWidgetAnimationPlayCallbackProxy* UWidgetAnimationPlayCallbackProxy::CreatePlayAnimationProxyObject(class UUMGSequencePlayer*& Result, class UUserWidget* Widget, UWidgetAnimation* InAnimation, float StartAtTime, int32 NumLoopsToPlay, EUMGSequencePlayMode::Type PlayMode, float PlaybackSpeed)
+{
+	UWidgetAnimationPlayCallbackProxy* Proxy = NewObject<UWidgetAnimationPlayCallbackProxy>();
+	Proxy->SetFlags(RF_StrongRefOnFrame);
+	Result = Proxy->ExecutePlayAnimation(Widget, InAnimation, StartAtTime, NumLoopsToPlay, PlayMode, PlaybackSpeed);
+	return Proxy;
+}
+
+UWidgetAnimationPlayCallbackProxy* UWidgetAnimationPlayCallbackProxy::CreatePlayAnimationTimeRangeProxyObject(class UUMGSequencePlayer*& Result, class UUserWidget* Widget, UWidgetAnimation* InAnimation, float StartAtTime, float EndAtTime, int32 NumLoopsToPlay, EUMGSequencePlayMode::Type PlayMode, float PlaybackSpeed)
+{
+	UWidgetAnimationPlayCallbackProxy* Proxy = NewObject<UWidgetAnimationPlayCallbackProxy>();
+	Proxy->SetFlags(RF_StrongRefOnFrame);
+	Result = Proxy->ExecutePlayAnimationTimeRange(Widget, InAnimation, StartAtTime, EndAtTime, NumLoopsToPlay, PlayMode, PlaybackSpeed);
+	return Proxy;
+}
+
+UWidgetAnimationPlayCallbackProxy::UWidgetAnimationPlayCallbackProxy(const FObjectInitializer& ObjectInitializer)
+	: Super(ObjectInitializer)
+{
+}
+
+class UUMGSequencePlayer* UWidgetAnimationPlayCallbackProxy::ExecutePlayAnimation(class UUserWidget* Widget, UWidgetAnimation* InAnimation, float StartAtTime, int32 NumLoopsToPlay, EUMGSequencePlayMode::Type PlayMode, float PlaybackSpeed)
+{
+	if (!Widget)
+	{
+		return nullptr;
+	}
+
+	WorldPtr = Widget->GetWorld();
+
+	UUMGSequencePlayer* Player = Widget->PlayAnimation(InAnimation, StartAtTime, NumLoopsToPlay, PlayMode, PlaybackSpeed);
+	if (Player)
+	{
+		Player->OnSequenceFinishedPlaying().AddUObject(this, &UWidgetAnimationPlayCallbackProxy::OnFinished);
+	}
+
+	return Player;
+}
+
+class UUMGSequencePlayer* UWidgetAnimationPlayCallbackProxy::ExecutePlayAnimationTimeRange(class UUserWidget* Widget, UWidgetAnimation* InAnimation, float StartAtTime, float EndAtTime, int32 NumLoopsToPlay, EUMGSequencePlayMode::Type PlayMode, float PlaybackSpeed)
+{
+	if (!Widget)
+	{
+		return nullptr;
+	}
+
+	WorldPtr = Widget->GetWorld();
+
+	UUMGSequencePlayer* Player = Widget->PlayAnimationTimeRange(InAnimation, StartAtTime, EndAtTime, NumLoopsToPlay, PlayMode, PlaybackSpeed);
+	if (Player)
+	{
+		OnFinishedHandle = Player->OnSequenceFinishedPlaying().AddUObject(this, &UWidgetAnimationPlayCallbackProxy::OnFinished);
+	}
+
+	return Player;
+}
+
+void UWidgetAnimationPlayCallbackProxy::OnFinished(class UUMGSequencePlayer& Player)
+{
+	Player.OnSequenceFinishedPlaying().Remove(OnFinishedHandle);
+
+	// We delay the Finish trigger to next frame.
+	if (UWorld* World = WorldPtr.Get())
+	{
+		// Use a dummy timer handle as we don't need to store it for later but we don't need to look for something to clear
+		FTimerHandle TimerHandle;
+		World->GetTimerManager().SetTimer(TimerHandle, this, &UWidgetAnimationPlayCallbackProxy::OnFinished_Delayed, 0.001f, false);
+	}
+}
+
+void UWidgetAnimationPlayCallbackProxy::OnFinished_Delayed()
+{
+	Finished.Broadcast();
+}
+
+#undef LOCTEXT_NAMESPACE
diff --git a/Engine/Source/Runtime/UMG/Public/Animation/WidgetAnimationPlayCallbackProxy.h b/Engine/Source/Runtime/UMG/Public/Animation/WidgetAnimationPlayCallbackProxy.h
new file mode 100644
--- /dev/null
+++ b/Engine/Source/Runtime/UMG/Public/Animation/WidgetAnimationPlayCallbackProxy.h
@@ -0,0 +1,35 @@
+// Copyright 1998-2019 Epic Games, Inc. All Rights Reserved.
+
+#pragma once
+
+#include "Blueprint/UserWidget.h"
+#include "WidgetAnimationPlayCallbackProxy.generated.h"
+
+DECLARE_DYNAMIC_MULTICAST_DELEGATE(FWidgetAnimationResult);
+
+UCLASS(MinimalAPI)
+class UWidgetAnimationPlayCallbackProxy : public UObject
+{
+	GENERATED_UCLASS_BODY()
+
+	// Called when animation has been completed
+	UPROPERTY(BlueprintAssignable)
+	FWidgetAnimationResult Finished;
+
+	UFUNCTION(BlueprintCallable, meta = (BlueprintInternalUseOnly = "true", DisplayName = "Play Animation with Finished event", ShortToolTip = "Play Animation and trigger event on Finished",  ToolTip="Play Animation on widget and trigger Finish event when the animation is done."), Category = "User Interface|Animation")
+	static UWidgetAnimationPlayCallbackProxy* CreatePlayAnimationProxyObject(class UUMGSequencePlayer*& Result, class UUserWidget* Widget, class UWidgetAnimation* InAnimation, float StartAtTime = 0.0f, int32 NumLoopsToPlay = 1, EUMGSequencePlayMode::Type PlayMode = EUMGSequencePlayMode::Forward, float PlaybackSpeed = 1.0f);
+
+	UFUNCTION(BlueprintCallable, meta = (BlueprintInternalUseOnly = "true", DisplayName = "Play Animation Time Range with Finished event", ShortToolTip = "Play Animation Time Range and trigger event on Finished", ToolTip = "Play Animation Time Range on widget and trigger Finish event when the animation is done."), Category = "User Interface|Animation")
+	static UWidgetAnimationPlayCallbackProxy* CreatePlayAnimationTimeRangeProxyObject(class UUMGSequencePlayer*& Result, class UUserWidget* Widget, class UWidgetAnimation* InAnimation, float StartAtTime = 0.0f, float EndAtTime = 0.0f, int32 NumLoopsToPlay = 1, EUMGSequencePlayMode::Type PlayMode = EUMGSequencePlayMode::Forward, float PlaybackSpeed = 1.0f);
+
+private:
+	class UUMGSequencePlayer* ExecutePlayAnimation(class UUserWidget* Widget, class UWidgetAnimation* InAnimation, float StartAtTime, int32 NumLoopsToPlay, EUMGSequencePlayMode::Type PlayMode, float PlaybackSpeed);
+	class UUMGSequencePlayer* ExecutePlayAnimationTimeRange(class UUserWidget* Widget, class UWidgetAnimation* InAnimation, float StartAtTime, float EndAtTime, int32 NumLoopsToPlay, EUMGSequencePlayMode::Type PlayMode, float PlaybackSpeed);
+	void OnFinished(class UUMGSequencePlayer& Player);
+	void OnFinished_Delayed();
+
+	// Pointer to the world, needed to delay the results slightly
+	TWeakObjectPtr<UWorld> WorldPtr;
+
+	FDelegateHandle OnFinishedHandle;
+};
