Exposed edit condition property for list view offset.

diff --git a/Engine/Source/Runtime/UMG/Public/Components/ListViewBase.h b/Engine/Source/Runtime/UMG/Public/Components/ListViewBase.h
--- a/Engine/Source/Runtime/UMG/Public/Components/ListViewBase.h
+++ b/Engine/Source/Runtime/UMG/Public/Components/ListViewBase.h
@@ -671,7 +671,7 @@ protected:
 	UPROPERTY(EditAnywhere, BlueprintReadOnly, Category = Scrolling)
 	bool bEnableScrollAnimation = false;
 
-	UPROPERTY()
+	UPROPERTY(EditAnywhere, Category = Scrolling)
 	bool bEnableFixedLineOffset = false;
 
 	/** 
