Exposed SetScrollOffset function to the ListViewBase class.

#rnx

diff --git a/Engine/Source/Runtime/UMG/Private/Components/ListViewBase.cpp b/Engine/Source/Runtime/UMG/Private/Components/ListViewBase.cpp
--- a/Engine/Source/Runtime/UMG/Private/Components/ListViewBase.cpp
+++ b/Engine/Source/Runtime/UMG/Private/Components/ListViewBase.cpp
@@ -62,6 +62,14 @@ void UListViewBase::ScrollToBottom()
 	}
 }
 
+void UListViewBase::SetScrollOffset(const float InScrollOffset)
+{
+	if (MyTableViewBase.IsValid())
+	{
+		MyTableViewBase->SetScrollOffset(InScrollOffset);
+	}
+}
+
 void UListViewBase::SetWheelScrollMultiplier(float NewWheelScrollMultiplier)
 {
 	WheelScrollMultiplier = NewWheelScrollMultiplier;
diff --git a/Engine/Source/Runtime/UMG/Public/Components/ListViewBase.h b/Engine/Source/Runtime/UMG/Public/Components/ListViewBase.h
--- a/Engine/Source/Runtime/UMG/Public/Components/ListViewBase.h
+++ b/Engine/Source/Runtime/UMG/Public/Components/ListViewBase.h
@@ -547,6 +547,10 @@ public:
 	UFUNCTION(BlueprintCallable, Category = ListViewBase)
 	void ScrollToBottom();
 
+	/** Set the scroll offset of this view (in items) */
+	UFUNCTION(BlueprintCallable, Category = ListView)
+	void SetScrollOffset(const float InScrollOffset);
+
 	UFUNCTION(BlueprintCallable, Category = ListViewBase)
 	void SetWheelScrollMultiplier(float NewWheelScrollMultiplier);
 
