Adding a helper function to UWidget: SetFocus. Simply calls SetUserFocus(GetOwningPlayer).

diff --git a/Engine/Source/Runtime/UMG/Private/Components/Widget.cpp b/Engine/Source/Runtime/UMG/Private/Components/Widget.cpp
--- a/Engine/Source/Runtime/UMG/Private/Components/Widget.cpp
+++ b/Engine/Source/Runtime/UMG/Private/Components/Widget.cpp
@@ -543,6 +543,11 @@ bool UWidget::HasUserFocusedDescendants(APlayerController* PlayerController) con
 	return false;
 }
 
+void UWidget::SetFocus()
+{
+	SetUserFocus(GetOwningPlayer());
+}
+
 void UWidget::SetUserFocus(APlayerController* PlayerController)
 {
 	if ( PlayerController == nullptr || !PlayerController->IsLocalPlayerController() || PlayerController->Player == nullptr )
diff --git a/Engine/Source/Runtime/UMG/Public/Components/Widget.h b/Engine/Source/Runtime/UMG/Public/Components/Widget.h
--- a/Engine/Source/Runtime/UMG/Public/Components/Widget.h
+++ b/Engine/Source/Runtime/UMG/Public/Components/Widget.h
@@ -595,6 +595,10 @@ public:
 	UFUNCTION(BlueprintCallable, Category="Widget")
 	bool HasUserFocusedDescendants(APlayerController* PlayerController) const;
 	
+	/** Sets the focus to this widget for the owning user */
+	UFUNCTION(BlueprintCallable, Category = "Widget")
+	void SetFocus();
+
 	/** Sets the focus to this widget for a specific user */
 	UFUNCTION(BlueprintCallable, Category="Widget")
 	void SetUserFocus(APlayerController* PlayerController);
