Fixed velocity for Widget3D proxy.

diff --git a/Engine/Source/Runtime/UMG/Private/Components/WidgetComponent.cpp b/Engine/Source/Runtime/UMG/Private/Components/WidgetComponent.cpp
--- a/Engine/Source/Runtime/UMG/Private/Components/WidgetComponent.cpp
+++ b/Engine/Source/Runtime/UMG/Private/Components/WidgetComponent.cpp
@@ -329,7 +329,14 @@ public:
 		//FSpriteTextureOverrideRenderProxy* TextureOverrideMaterialProxy = new FSpriteTextureOverrideRenderProxy(ParentMaterialProxy,
 
 		const FMatrix& ViewportLocalToWorld = GetLocalToWorld();
-	
+
+		FMatrix PreviousLocalToWorld;
+
+		if (!GetScene().GetPreviousLocalToWorld(GetPrimitiveSceneInfo(), PreviousLocalToWorld))
+		{
+			PreviousLocalToWorld = GetLocalToWorld();
+		}
+
 		if( RenderTarget )
 		{
 			FTextureResource* TextureResource = RenderTarget->Resource;
@@ -358,7 +365,11 @@ public:
 							MeshBuilder.AddTriangle(VertexIndices[0], VertexIndices[1], VertexIndices[2]);
 							MeshBuilder.AddTriangle(VertexIndices[0], VertexIndices[2], VertexIndices[3]);
 
-							MeshBuilder.GetMesh(ViewportLocalToWorld, ParentMaterialProxy, SDPG_World, false, true, ViewIndex, Collector);
+							FDynamicMeshBuilderSettings Settings;
+							Settings.bDisableBackfaceCulling = false;
+							Settings.bReceivesDecals = true;
+							Settings.bUseSelectionOutline = true;
+							MeshBuilder.GetMesh(ViewportLocalToWorld, PreviousLocalToWorld, ParentMaterialProxy, SDPG_World, Settings, nullptr, ViewIndex, Collector, FHitProxyId());
 						}
 					}
 				}
@@ -435,7 +446,12 @@ public:
 								LastTangentY = TangentY;
 								LastTangentZ = TangentZ;
 							}
-							MeshBuilder.GetMesh(ViewportLocalToWorld, ParentMaterialProxy, SDPG_World, false, true, ViewIndex, Collector);
+
+							FDynamicMeshBuilderSettings Settings;
+							Settings.bDisableBackfaceCulling = false;
+							Settings.bReceivesDecals = true;
+							Settings.bUseSelectionOutline = true;
+							MeshBuilder.GetMesh(ViewportLocalToWorld, PreviousLocalToWorld, ParentMaterialProxy, SDPG_World, Settings, nullptr, ViewIndex, Collector, FHitProxyId());
 						}
 					}
 				}
