UMG: Fix crash when placing a native UserWidget in a panel within another UserWidget

diff --git a/Engine/Source/Runtime/UMG/Private/UserWidget.cpp b/Engine/Source/Runtime/UMG/Private/UserWidget.cpp
--- a/Engine/Source/Runtime/UMG/Private/UserWidget.cpp
+++ b/Engine/Source/Runtime/UMG/Private/UserWidget.cpp
@@ -941,18 +941,13 @@ void UUserWidget::OnWidgetRebuilt()
 
 TSharedPtr<SWidget> UUserWidget::GetSlateWidgetFromName(const FName& Name) const
 {
-	UWidget* WidgetObject = WidgetTree->FindWidget(Name);
-	if ( WidgetObject )
-	{
-		return WidgetObject->GetCachedWidget();
-	}
-
-	return TSharedPtr<SWidget>();
+	UWidget* WidgetObject = GetWidgetFromName(Name);
+	return WidgetObject ? WidgetObject->GetCachedWidget() : TSharedPtr<SWidget>();
 }
 
 UWidget* UUserWidget::GetWidgetFromName(const FName& Name) const
 {
-	return WidgetTree->FindWidget(Name);
+	return WidgetTree ? WidgetTree->FindWidget(Name) : nullptr;
 }
 
 void UUserWidget::GetSlotNames(TArray<FName>& SlotNames) const
@@ -962,17 +957,14 @@ void UUserWidget::GetSlotNames(TArray<FName>& SlotNames) const
 	{
 		SlotNames.Append(BGClass->NamedSlots);
 	}
-	else // For non-blueprint widget blueprints we have to go through the widget tree to locate the named slots dynamically.
+	else if (WidgetTree) // For non-blueprint widget blueprints we have to go through the widget tree to locate the named slots dynamically.
 	{
-		TArray<FName> NamedSlots;
-		WidgetTree->ForEachWidget([&NamedSlots] (UWidget* Widget) {
+		WidgetTree->ForEachWidget([&SlotNames] (UWidget* Widget) {
 			if ( Widget && Widget->IsA<UNamedSlot>() )
 			{
-				NamedSlots.Add(Widget->GetFName());
+				SlotNames.Add(Widget->GetFName());
 			}
 		});
-
-		SlotNames.Append(NamedSlots);
 	}
 }
 
