Fix typo in comment on UWidget::ForceLayoutPrepass

diff --git a/Engine/Source/Runtime/UMG/Public/Components/Widget.h b/Engine/Source/Runtime/UMG/Public/Components/Widget.h
--- a/Engine/Source/Runtime/UMG/Public/Components/Widget.h
+++ b/Engine/Source/Runtime/UMG/Public/Components/Widget.h
@@ -605,7 +605,7 @@ public:
 
 	/**
 	 * Forces a pre-pass.  A pre-pass caches the desired size of the widget hierarchy owned by this widget.  
-	 * One pre-pass is already happens for every widget before Tick occurs.  You only need to perform another 
+	 * One pre-pass already happens for every widget before Tick occurs.  You only need to perform another 
 	 * pre-pass if you are adding child widgets this frame and want them to immediately be visible this frame.
 	 */
 	UFUNCTION(BlueprintCallable, Category="Widget")
