Correction to Drag And drop to Viewport to correctly release Slate ressources when it was allocated only for Preview Widget in Canvas

diff --git a/Engine/Source/Runtime/UMG/Private/Components/CanvasPanelSlot.cpp b/Engine/Source/Runtime/UMG/Private/Components/CanvasPanelSlot.cpp
--- a/Engine/Source/Runtime/UMG/Private/Components/CanvasPanelSlot.cpp
+++ b/Engine/Source/Runtime/UMG/Private/Components/CanvasPanelSlot.cpp
@@ -82,6 +82,9 @@ bool UCanvasPanelSlot::NudgeByDesigner(const FVector2D& NudgeDirection, const TO
 
 bool UCanvasPanelSlot::DragDropPreviewByDesigner(const FVector2D& LocalCursorPosition, const TOptional<int32>& XGridSnapSize, const TOptional<int32>& YGridSnapSize)
 {
+	// If the widget is not constructed yet, we need to call ReleaseSlateResources
+	bool bReleaseSlateResources = !Content->IsConstructed();
+
 	// HACK UMG - This seems like a bad idea to call TakeWidget
 	TSharedPtr<SWidget> SlateWidget = Content->TakeWidget();
 	SlateWidget->SlatePrepass();
@@ -100,16 +103,25 @@ bool UCanvasPanelSlot::DragDropPreviewByDesigner(const FVector2D& LocalCursorPos
 		NewPosition.Y = ((int32)NewPosition.Y) - (((int32)NewPosition.Y) % YGridSnapSize.GetValue());
 	}
 
+	bool LayoutChanged = true;
 	// Return false and early out if there are no effective changes.
 	if (GetSize() == LocalSize && GetPosition() == NewPosition)
 	{
-		return false;
+		LayoutChanged = false;
+	}
+	else
+	{
+		SetPosition(NewPosition);
+		SetSize(LocalSize);
 	}
 
-	SetPosition(NewPosition);
-	SetSize(LocalSize);
+	if (bReleaseSlateResources)
+	{
+		// When we are done, we free the Widget that was created by TakeWidget.
+		Content->ReleaseSlateResources(true);
+	}
 
-	return true;
+	return LayoutChanged;
 }
 
 void UCanvasPanelSlot::SynchronizeFromTemplate(const UPanelSlot* const TemplateSlot)
