Unshelved from pending changelist '9745484':

- Fixes MobileFontSizeMultiplier not being applied when ApplyUpdatedDefaultTextStyle() is called on CommonRichTextBlock - affects left blade menu friends list

#rnx

diff --git a/Engine/Source/Runtime/UMG/Public/Components/RichTextBlock.h b/Engine/Source/Runtime/UMG/Public/Components/RichTextBlock.h
--- a/Engine/Source/Runtime/UMG/Public/Components/RichTextBlock.h
+++ b/Engine/Source/Runtime/UMG/Public/Components/RichTextBlock.h
@@ -156,7 +156,7 @@ protected:
 	virtual TSharedPtr< class IRichTextMarkupWriter > CreateMarkupWriter();
 
 	void BeginDefaultStyleOverride();
-	void ApplyUpdatedDefaultTextStyle();
+	virtual void ApplyUpdatedDefaultTextStyle();
 
 protected:
 	/** The text to display */
