#UE4 Only ensure for a missing StyleInstance in GetCurrentDefaultTextStyle if you are not using an overridden style

diff --git a/Engine/Source/Runtime/UMG/Private/Components/RichTextBlock.cpp b/Engine/Source/Runtime/UMG/Private/Components/RichTextBlock.cpp
--- a/Engine/Source/Runtime/UMG/Private/Components/RichTextBlock.cpp
+++ b/Engine/Source/Runtime/UMG/Private/Components/RichTextBlock.cpp
@@ -169,8 +169,15 @@ const FTextBlockStyle& URichTextBlock::GetDefaultTextStyle() const
 
 const FTextBlockStyle& URichTextBlock::GetCurrentDefaultTextStyle() const
 {
-	ensure(StyleInstance.IsValid());
-	return bOverrideDefaultStyle ? DefaultTextStyleOverride : DefaultTextStyle;
+	if (bOverrideDefaultStyle)
+	{
+		return DefaultTextStyleOverride;
+	}
+	else
+	{
+		ensure(StyleInstance.IsValid());
+		return DefaultTextStyle;
+	}
 }
 
 URichTextBlockDecorator* URichTextBlock::GetDecoratorByClass(TSubclassOf<URichTextBlockDecorator> DecoratorClass)
