Expose EVirtualKeyboardTrigger attribute to UMG.

diff --git a/Engine/Source/Runtime/Slate/Public/Widgets/Text/ISlateEditableTextWidget.h b/Engine/Source/Runtime/Slate/Public/Widgets/Text/ISlateEditableTextWidget.h
--- a/Engine/Source/Runtime/Slate/Public/Widgets/Text/ISlateEditableTextWidget.h
+++ b/Engine/Source/Runtime/Slate/Public/Widgets/Text/ISlateEditableTextWidget.h
@@ -50,7 +50,8 @@ enum class ETextLocation
 	NextPage,
 };
 
-enum class EVirtualKeyboardTrigger
+UENUM(BlueprintType)
+enum class EVirtualKeyboardTrigger : uint8
 {
 	/** Display the virtual keyboard when the widget gains keyboard focus by a pointer action. */
 	OnFocusByPointer,
diff --git a/Engine/Source/Runtime/UMG/Private/Components/EditableText.cpp b/Engine/Source/Runtime/UMG/Private/Components/EditableText.cpp
--- a/Engine/Source/Runtime/UMG/Private/Components/EditableText.cpp
+++ b/Engine/Source/Runtime/UMG/Private/Components/EditableText.cpp
@@ -45,6 +45,7 @@ UEditableText::UEditableText(const FObjectInitializer& ObjectInitializer)
 	ClearKeyboardFocusOnCommit = true;
 	SelectAllTextOnCommit = false;
 	AllowContextMenu = true;
+	VirtualKeyboardTrigger = EVirtualKeyboardTrigger::OnFocusByPointer;
 	VirtualKeyboardDismissAction = EVirtualKeyboardDismissAction::TextChangeOnDismiss;
 	Clipping = EWidgetClipping::ClipToBounds;
 
@@ -75,6 +76,7 @@ TSharedRef<SWidget> UEditableText::RebuildWidget()
 		.OnTextCommitted( BIND_UOBJECT_DELEGATE( FOnTextCommitted, HandleOnTextCommitted ) )
 		.VirtualKeyboardType( EVirtualKeyboardType::AsKeyboardType( KeyboardType.GetValue() ) )
 		.VirtualKeyboardOptions(VirtualKeyboardOptions)
+		.VirtualKeyboardTrigger(VirtualKeyboardTrigger)
 		.VirtualKeyboardDismissAction(VirtualKeyboardDismissAction)
 		.Justification( Justification );
 	
diff --git a/Engine/Source/Runtime/UMG/Public/Components/EditableText.h b/Engine/Source/Runtime/UMG/Public/Components/EditableText.h
--- a/Engine/Source/Runtime/UMG/Public/Components/EditableText.h
+++ b/Engine/Source/Runtime/UMG/Public/Components/EditableText.h
@@ -121,6 +121,9 @@ public:
 	UPROPERTY(EditAnywhere, Category=Behavior, AdvancedDisplay)
 	FVirtualKeyboardOptions VirtualKeyboardOptions;
 
+	UPROPERTY(EditAnywhere, Category = Behavior, AdvancedDisplay)
+	EVirtualKeyboardTrigger VirtualKeyboardTrigger;
+
 	/** What action should be taken when the virtual keyboard is dismissed? */
 	UPROPERTY(EditAnywhere, Category=Behavior, AdvancedDisplay)
 	EVirtualKeyboardDismissAction VirtualKeyboardDismissAction;
