#UE4 Mark AccessibleWidgetData instanced in UWidget so doesnt get cleared in CDOs when being compiled

diff --git a/Engine/Source/Runtime/UMG/Public/Components/Widget.h b/Engine/Source/Runtime/UMG/Public/Components/Widget.h
--- a/Engine/Source/Runtime/UMG/Public/Components/Widget.h
+++ b/Engine/Source/Runtime/UMG/Public/Components/Widget.h
@@ -366,7 +366,7 @@ public:
 
 private:
 	/** A custom set of accessibility rules for this widget. If null, default rules for the widget are used. */
-	UPROPERTY()
+	UPROPERTY(Instanced)
 	USlateAccessibleWidgetData* AccessibleWidgetData;
 
 protected:
