Added hook to allow the widget component to modify the projected local position.

#rnx

diff --git a/Engine/Source/Runtime/UMG/Private/Slate/SWorldWidgetScreenLayer.cpp b/Engine/Source/Runtime/UMG/Private/Slate/SWorldWidgetScreenLayer.cpp
--- a/Engine/Source/Runtime/UMG/Private/Slate/SWorldWidgetScreenLayer.cpp
+++ b/Engine/Source/Runtime/UMG/Private/Slate/SWorldWidgetScreenLayer.cpp
@@ -132,6 +132,8 @@ void SWorldWidgetScreenLayer::Tick(const FGeometry& AllottedGeometry, const doub
 
 							if ( Entry.WidgetComponent )
 							{
+								LocalPosition = Entry.WidgetComponent->ModifyProjectedLocalPosition(ViewportGeometry, LocalPosition);
+
 								FVector2D ComponentDrawSize = Entry.WidgetComponent->GetDrawSize();
 								FVector2D ComponentPivot = Entry.WidgetComponent->GetPivot();
 								
diff --git a/Engine/Source/Runtime/UMG/Public/Components/WidgetComponent.h b/Engine/Source/Runtime/UMG/Public/Components/WidgetComponent.h
--- a/Engine/Source/Runtime/UMG/Public/Components/WidgetComponent.h
+++ b/Engine/Source/Runtime/UMG/Public/Components/WidgetComponent.h
@@ -372,6 +372,9 @@ public:
 	UFUNCTION(BlueprintCallable, Category = UserInterface)
 	void SetWindowVisibility(EWindowVisibility InVisibility);
 
+	/** Hook to allow this component modify the local position of the widget after it has been projected from world space to screen space. */
+	virtual FVector2D ModifyProjectedLocalPosition(const FGeometry& ViewportGeometry, const FVector2D& LocalPosition) { return LocalPosition; }
+
 protected:
 	void OnLevelRemovedFromWorld(ULevel* InLevel, UWorld* InWorld);
 
