Remove Unnecesary include from Widget

diff --git a/Engine/Source/Runtime/UMG/Public/Components/Widget.h b/Engine/Source/Runtime/UMG/Public/Components/Widget.h
--- a/Engine/Source/Runtime/UMG/Public/Components/Widget.h
+++ b/Engine/Source/Runtime/UMG/Public/Components/Widget.h
@@ -21,11 +21,6 @@
 #include "GameFramework/PlayerController.h"
 #include "Blueprint/WidgetNavigation.h"
 
-#if WITH_EDITOR
-// This violates IWYU, but the alternative is .cpp includes that are invariably not within #if WITH_EDITOR and cause non-editor build failures
-#include "Kismet2/CompilerResultsLog.h"
-#endif
-
 #include "Widget.generated.h"
 
 class ULocalPlayer;
@@ -596,10 +591,10 @@ public:
 	bool HasUserFocusedDescendants(APlayerController* PlayerController) const;
 	
 	/** Sets the focus to this widget for the owning user */
-	UFUNCTION(BlueprintCallable, Category = "Widget")
+	UFUNCTION(BlueprintCallable, Category="Widget")
 	void SetFocus();
 
-	/** Sets the focus to this widget for a specific user */
+	/** Sets the focus to this widget for a specific user (if setting focus for the owning user, prefer SetFocus()) */
 	UFUNCTION(BlueprintCallable, Category="Widget")
 	void SetUserFocus(APlayerController* PlayerController);
 
