Fix missing include
#rnx

diff --git a/Engine/Source/Runtime/UMG/Public/Editor/WidgetCompilerLog.h b/Engine/Source/Runtime/UMG/Public/Editor/WidgetCompilerLog.h
--- a/Engine/Source/Runtime/UMG/Public/Editor/WidgetCompilerLog.h
+++ b/Engine/Source/Runtime/UMG/Public/Editor/WidgetCompilerLog.h
@@ -5,9 +5,11 @@
 #include "CoreMinimal.h"
 #include "UObject/ObjectMacros.h"
 #include "UObject/Object.h"
+#include "Logging/TokenizedMessage.h
 
 #if WITH_EDITOR
 
+
 class IWidgetCompilerLog
 {
 public:
