Replace lost character
#rnx

diff --git a/Engine/Source/Runtime/UMG/Public/Editor/WidgetCompilerLog.h b/Engine/Source/Runtime/UMG/Public/Editor/WidgetCompilerLog.h
--- a/Engine/Source/Runtime/UMG/Public/Editor/WidgetCompilerLog.h
+++ b/Engine/Source/Runtime/UMG/Public/Editor/WidgetCompilerLog.h
@@ -5,11 +5,10 @@
 #include "CoreMinimal.h"
 #include "UObject/ObjectMacros.h"
 #include "UObject/Object.h"
-#include "Logging/TokenizedMessage.h
+#include "Logging/TokenizedMessage.h"
 
 #if WITH_EDITOR
 
-
 class IWidgetCompilerLog
 {
 public:
