Added paramter to pass padding to scrolling a widget into view. Slate already supports taking this.
#rnx

diff --git a/Engine/Source/Runtime/UMG/Private/Components/ScrollBox.cpp b/Engine/Source/Runtime/UMG/Private/Components/ScrollBox.cpp
--- a/Engine/Source/Runtime/UMG/Private/Components/ScrollBox.cpp
+++ b/Engine/Source/Runtime/UMG/Private/Components/ScrollBox.cpp
@@ -190,7 +190,7 @@ void UScrollBox::ScrollToEnd()
 	}
 }
 
-void UScrollBox::ScrollWidgetIntoView(UWidget* WidgetToFind, bool AnimateScroll, EDescendantScrollDestination InScrollDestination)
+void UScrollBox::ScrollWidgetIntoView(UWidget* WidgetToFind, bool AnimateScroll, EDescendantScrollDestination InScrollDestination, float Padding)
 {
 	TSharedPtr<SWidget> SlateWidgetToFind;
 	if (WidgetToFind)
@@ -202,7 +202,7 @@ void UScrollBox::ScrollWidgetIntoView(UWidget* WidgetToFind, bool AnimateScroll,
 	{
 		// NOTE: Pass even if null! This, in effect, cancels a request to scroll which is necessary to avoid warnings/ensures 
 		//       when we request to scroll to a widget and later remove that widget!
-		MyScrollBox->ScrollDescendantIntoView(SlateWidgetToFind, AnimateScroll, InScrollDestination);
+		MyScrollBox->ScrollDescendantIntoView(SlateWidgetToFind, AnimateScroll, InScrollDestination, Padding);
 	}
 }
 
diff --git a/Engine/Source/Runtime/UMG/Public/Components/ScrollBox.h b/Engine/Source/Runtime/UMG/Public/Components/ScrollBox.h
--- a/Engine/Source/Runtime/UMG/Public/Components/ScrollBox.h
+++ b/Engine/Source/Runtime/UMG/Public/Components/ScrollBox.h
@@ -163,7 +163,7 @@ public:
 
 	/** Scrolls the ScrollBox to the widget during the next layout pass. */
 	UFUNCTION(BlueprintCallable, Category="Widget")
-	void ScrollWidgetIntoView(UWidget* WidgetToFind, bool AnimateScroll = true, EDescendantScrollDestination ScrollDestination = EDescendantScrollDestination::IntoView );
+	void ScrollWidgetIntoView(UWidget* WidgetToFind, bool AnimateScroll = true, EDescendantScrollDestination ScrollDestination = EDescendantScrollDestination::IntoView, float Padding = 0);
 
 	//~ Begin UWidget Interface
 	virtual void SynchronizeProperties() override;
