When a widget isn't focusable, deduce whether it's a UserWidget before suggesting doing something UserWidget-y to it. Also improve the message in the Message Log by adding a clickable class name to the message.

Removed an unnecessary conversion to PlayerContext when Setting User Focus on a widget

diff --git a/Engine/Source/Runtime/UMG/Private/Components/Widget.cpp b/Engine/Source/Runtime/UMG/Private/Components/Widget.cpp
--- a/Engine/Source/Runtime/UMG/Private/Components/Widget.cpp
+++ b/Engine/Source/Runtime/UMG/Private/Components/Widget.cpp
@@ -2,6 +2,7 @@
 
 #include "Components/Widget.h"
 #include "Misc/ConfigCacheIni.h"
+#include "Misc/UObjectToken.h"
 #include "CoreGlobals.h"
 #include "Widgets/SNullWidget.h"
 #include "Types/NavigationMetaData.h"
@@ -564,13 +565,26 @@ void UWidget::SetUserFocus(APlayerController* PlayerController)
 #if !(UE_BUILD_SHIPPING || UE_BUILD_TEST)
 		if ( !SafeWidget->SupportsKeyboardFocus() )
 		{
-			FMessageLog("PIE").Warning(LOCTEXT("ThisWidgetDoesntSupportFocus", "This widget does not support focus.  If this is a UserWidget, you should set bIsFocusable to true."));
+			TSharedRef<FTokenizedMessage> Message = FMessageLog("PIE").Warning()->AddToken(FUObjectToken::Create(this));
+#if WITH_EDITORONLY_DATA
+			if(UObject* GeneratedBy = WidgetGeneratedBy.Get())
+			{
+				Message->AddToken(FTextToken::Create(FText::FromString(TEXT(" in "))))->AddToken(FUObjectToken::Create(GeneratedBy));
+			}
+#endif
+			if (IsA(UUserWidget::StaticClass()))
+			{
+				Message->AddToken(FTextToken::Create(LOCTEXT("UserWidgetDoesntSupportFocus", " does not support focus, you should set bIsFocusable to true.")));
+			}
+			else
+			{
+				Message->AddToken(FTextToken::Create(LOCTEXT("NonUserWidgetDoesntSupportFocus", " does not support focus.")));
+			}
+			
 		}
 #endif
 
-		FLocalPlayerContext Context(PlayerController);
-
-		if ( ULocalPlayer* LocalPlayer = Context.GetLocalPlayer() )
+		if ( ULocalPlayer* LocalPlayer = PlayerController->GetLocalPlayer() )
 		{
 			// HACK: We use the controller Id as the local player index for focusing widgets in Slate.
 			int32 UserIndex = LocalPlayer->GetControllerId();
