Fix mouse down on table rows setting focus back to the parent table. This causes issues with selecting the right list entries if they rely on themselves being focused. This behavior is now controllable with a setting

diff --git a/Engine/Source/Runtime/Slate/Public/Framework/Views/ITypedTableView.h b/Engine/Source/Runtime/Slate/Public/Framework/Views/ITypedTableView.h
--- a/Engine/Source/Runtime/Slate/Public/Framework/Views/ITypedTableView.h
+++ b/Engine/Source/Runtime/Slate/Public/Framework/Views/ITypedTableView.h
@@ -73,12 +73,21 @@ public:
 	virtual bool Private_IsItemSelected( const ItemType& TheItem ) const = 0;
 
 	/**
-	* Test if the current item should be highlighted. This is separate from hover highlights.
-	*
-	* @param TheItem The item to test.
-	*
-	* @return true if the item is marked as highlighted in this list; false otherwise.
-	*/
+	 * Test if the current item can be selected or navigated to.
+	 *
+	 * @param InItem The item to test.
+	 *
+	 * @return true if the item is selectable or navigable; false otherwise.
+	 */
+	virtual bool Private_IsItemSelectableOrNavigable(const ItemType& TheItem) const = 0;
+
+	/**
+	 * Test if the current item should be highlighted. This is separate from hover highlights.
+	 *
+	 * @param TheItem The item to test.
+	 *
+	 * @return true if the item is marked as highlighted in this list; false otherwise.
+	 */
 	virtual bool Private_IsItemHighlighted(const ItemType& TheItem) const = 0;
 	
 	/**
diff --git a/Engine/Source/Runtime/Slate/Public/Widgets/Views/SListView.h b/Engine/Source/Runtime/Slate/Public/Widgets/Views/SListView.h
--- a/Engine/Source/Runtime/Slate/Public/Widgets/Views/SListView.h
+++ b/Engine/Source/Runtime/Slate/Public/Widgets/Views/SListView.h
@@ -840,6 +840,11 @@ public:
 		return false;	
 	}
 
+	virtual bool Private_IsItemSelectableOrNavigable(const ItemType& TheItem) const override
+	{
+		return OnIsSelectableOrNavigable.IsBound() ? OnIsSelectableOrNavigable.Execute(TheItem) : true;
+	}
+
 	virtual void Private_SetItemExpansion( ItemType TheItem, bool bShouldBeExpanded ) override
 	{
 		// Do nothing; you cannot expand an item in a list!
diff --git a/Engine/Source/Runtime/UMG/Public/Slate/SObjectTableRow.h b/Engine/Source/Runtime/UMG/Public/Slate/SObjectTableRow.h
--- a/Engine/Source/Runtime/UMG/Public/Slate/SObjectTableRow.h
+++ b/Engine/Source/Runtime/UMG/Public/Slate/SObjectTableRow.h
@@ -328,8 +328,9 @@ public:
 			{
 				if (IsItemSelectable())
 				{
+					const ItemType* MyItemPtr = GetItemForThis(OwnerTable);
 					// New selections are handled on mouse down, deselection is handled on mouse up
-					if (const ItemType* MyItemPtr = GetItemForThis(OwnerTable))
+					if (MyItemPtr)
 					{
 						const ItemType& MyItem = *MyItemPtr;
 						if (!OwnerTable->Private_IsItemSelected(MyItem))
@@ -342,12 +343,19 @@ public:
 							bChangedSelectionOnMouseDown = true;
 						}
 					}
+
+					Reply = FReply::Handled()
+						.DetectDrag(SharedThis(this), EKeys::LeftMouseButton)
+						.CaptureMouse(SharedThis(this));
+
+					// Set focus back to the owning widget if the item is invalid somehow or its not selectable or can be navigated to
+					if (!MyItemPtr || !OwnerTable->Private_IsItemSelectableOrNavigable(*MyItemPtr))
+					{
+						Reply.SetUserFocus(OwnerTable->AsWidget(), EFocusCause::Mouse);
+					}
 				}
 
-				Reply = FReply::Handled()
-					.DetectDrag(SharedThis(this), EKeys::LeftMouseButton)
-					.SetUserFocus(OwnerTable->AsWidget(), EFocusCause::Mouse)
-					.CaptureMouse(SharedThis(this));
+		
 			}
 		}
 		
