Accounting for a (major) missed spot when initially updating the varios places UWidget passes a UserIndex to FSlateApplication - the processing of the FReply that lives on the LocalPlayer itself
- Now, if ever the fallback case of UWidget::SetFocus() is hit, for example, the desired focus target will be updated for the correct SlateUser that corresponds to the LocalPlayer
- Also updating UWidget's uses of GetUserIndexForController() to the new signature that accounts for the possibility of a player not yet having a valid SlateUser index

diff --git a/Engine/Source/Runtime/Launch/Private/LaunchEngineLoop.cpp b/Engine/Source/Runtime/Launch/Private/LaunchEngineLoop.cpp
--- a/Engine/Source/Runtime/Launch/Private/LaunchEngineLoop.cpp
+++ b/Engine/Source/Runtime/Launch/Private/LaunchEngineLoop.cpp
@@ -4140,13 +4140,16 @@ void FEngineLoop::ProcessLocalPlayerSlateOperations() const
 						APlayerController* PlayerController = Iterator->Get();
 						if (PlayerController)
 						{
-							ULocalPlayer* LocalPlayer = Cast< ULocalPlayer >(PlayerController->Player);
-							if (LocalPlayer)
+							if (ULocalPlayer* LocalPlayer = Cast<ULocalPlayer>(PlayerController->Player))
 							{
-								FReply& TheReply = LocalPlayer->GetSlateOperations();
-								SlateApp.ProcessExternalReply(PathToWidget, TheReply, LocalPlayer->GetControllerId());
+								TOptional<int32> UserIndex = SlateApp.GetUserIndexForController(LocalPlayer->GetControllerId());
+								if (UserIndex.IsSet())
+								{
+									FReply& TheReply = LocalPlayer->GetSlateOperations();
+									SlateApp.ProcessExternalReply(PathToWidget, TheReply, UserIndex.GetValue());
 
-								TheReply = FReply::Unhandled();
+									TheReply = FReply::Unhandled();
+								}
 							}
 						}
 					}
diff --git a/Engine/Source/Runtime/UMG/Private/Components/Widget.cpp b/Engine/Source/Runtime/UMG/Private/Components/Widget.cpp
--- a/Engine/Source/Runtime/UMG/Private/Components/Widget.cpp
+++ b/Engine/Source/Runtime/UMG/Private/Components/Widget.cpp
@@ -486,10 +486,12 @@ bool UWidget::HasUserFocus(APlayerController* PlayerController) const
 
 		if ( ULocalPlayer* LocalPlayer = Context.GetLocalPlayer() )
 		{
-			int32 UserIndex = FSlateApplication::Get().GetUserIndexForController(LocalPlayer->GetControllerId());
-
-			TOptional<EFocusCause> FocusCause = SafeWidget->HasUserFocus(UserIndex);
-			return FocusCause.IsSet();
+			TOptional<int32> UserIndex = FSlateApplication::Get().GetUserIndexForController(LocalPlayer->GetControllerId());
+			if (UserIndex.IsSet())
+			{
+				TOptional<EFocusCause> FocusCause = SafeWidget->HasUserFocus(UserIndex.GetValue());
+				return FocusCause.IsSet();
+			}
 		}
 	}
 
@@ -533,9 +535,11 @@ bool UWidget::HasUserFocusedDescendants(APlayerController* PlayerController) con
 
 		if ( ULocalPlayer* LocalPlayer = Context.GetLocalPlayer() )
 		{
-			int32 UserIndex = FSlateApplication::Get().GetUserIndexForController(LocalPlayer->GetControllerId());
-
-			return SafeWidget->HasUserFocusedDescendants(UserIndex);
+			TOptional<int32> UserIndex = FSlateApplication::Get().GetUserIndexForController(LocalPlayer->GetControllerId());
+			if (UserIndex.IsSet())
+			{
+				return SafeWidget->HasUserFocusedDescendants(UserIndex.GetValue());
+			}
 		}
 	}
 
@@ -584,9 +588,8 @@ void UWidget::SetUserFocus(APlayerController* PlayerController)
 
 		if ( ULocalPlayer* LocalPlayer = PlayerController->GetLocalPlayer() )
 		{
-			int32 UserIndex = FSlateApplication::Get().GetUserIndexForController(LocalPlayer->GetControllerId());
-
-			if (UserIndex >= 0 && !FSlateApplication::Get().SetUserFocus(UserIndex, SafeWidget) )
+			TOptional<int32> UserIndex = FSlateApplication::Get().GetUserIndexForController(LocalPlayer->GetControllerId());
+			if (UserIndex.IsSet() && !FSlateApplication::Get().SetUserFocus(UserIndex.GetValue(), SafeWidget) )
 			{
 				LocalPlayer->GetSlateOperations().SetUserFocus(SafeWidget.ToSharedRef());
 			}
