Adding null check to SObjectWidget::OnNavigation to avoid potential crash

diff --git a/Engine/Source/Runtime/UMG/Private/Slate/SObjectWidget.cpp b/Engine/Source/Runtime/UMG/Private/Slate/SObjectWidget.cpp
--- a/Engine/Source/Runtime/UMG/Private/Slate/SObjectWidget.cpp
+++ b/Engine/Source/Runtime/UMG/Private/Slate/SObjectWidget.cpp
@@ -513,13 +513,13 @@ FReply SObjectWidget::OnTouchForceChanged(const FGeometry& MyGeometry, const FPo
 
 FNavigationReply SObjectWidget::OnNavigation(const FGeometry& MyGeometry, const FNavigationEvent& InNavigationEvent)
 {
-	if (WidgetObject->NativeSupportsCustomNavigation())
+	if (WidgetObject && WidgetObject->NativeSupportsCustomNavigation())
 	{
 		return WidgetObject->NativeOnNavigation(MyGeometry, InNavigationEvent);
 	}
 	FNavigationReply Reply = SCompoundWidget::OnNavigation(MyGeometry, InNavigationEvent);
 
-	if ( CanRouteEvent() )
+	if (WidgetObject && CanRouteEvent() )
 	{
 		return WidgetObject->NativeOnNavigation(MyGeometry, InNavigationEvent, Reply);
 	}
