PR #5663: Fixed category of EditableTextBox (Contributed by [REDACTED_PERSON])

diff --git a/Engine/Source/Runtime/UMG/Private/Components/EditableTextBox.cpp b/Engine/Source/Runtime/UMG/Private/Components/EditableTextBox.cpp
--- a/Engine/Source/Runtime/UMG/Private/Components/EditableTextBox.cpp
+++ b/Engine/Source/Runtime/UMG/Private/Components/EditableTextBox.cpp
@@ -231,7 +231,7 @@ void UEditableTextBox::PostLoad()
 
 const FText UEditableTextBox::GetPaletteCategory()
 {
-	return LOCTEXT("Common", "Common");
+	return LOCTEXT("Input", "Input");
 }
 
 #endif
