PR #5661: Fixed UComboBoxString not broadcasting all changes (Contributed by [REDACTED_PERSON])

With these changes, we broadcast the change of selection even is the slate widget is not ready.

diff --git a/Engine/Source/Runtime/UMG/Private/Components/ComboBoxString.cpp b/Engine/Source/Runtime/UMG/Private/Components/ComboBoxString.cpp
--- a/Engine/Source/Runtime/UMG/Private/Components/ComboBoxString.cpp
+++ b/Engine/Source/Runtime/UMG/Private/Components/ComboBoxString.cpp
@@ -235,6 +235,10 @@ void UComboBoxString::SetSelectedIndex(const int32 Index)
 			{
 				MyComboBox->SetSelectedItem(CurrentOptionPtr);
 				UpdateOrGenerateWidget(CurrentOptionPtr);
+			}		
+			else
+			{
+				HandleSelectionChanged(CurrentOptionPtr, ESelectInfo::Direct);
 			}
 		}
 	}
