PR #5666: UMG bindings for integer and float (Contributed by [REDACTED_PERSON])

diff --git a/Engine/Source/Runtime/UMG/Private/Binding/TextBinding.cpp b/Engine/Source/Runtime/UMG/Private/Binding/TextBinding.cpp
--- a/Engine/Source/Runtime/UMG/Private/Binding/TextBinding.cpp
+++ b/Engine/Source/Runtime/UMG/Private/Binding/TextBinding.cpp
@@ -4,10 +4,6 @@
 
 #define LOCTEXT_NAMESPACE "UMG"
 
-UTextBinding::UTextBinding()
-{
-}
-
 bool UTextBinding::IsSupportedDestination(UProperty* Property) const
 {
 	return IsConcreteTypeCompatibleWithReflectedType<FText>(Property);
@@ -17,7 +13,9 @@ bool UTextBinding::IsSupportedSource(UProperty* Property) const
 {
 	return
 		IsConcreteTypeCompatibleWithReflectedType<FText>(Property) ||
-		IsConcreteTypeCompatibleWithReflectedType<FString>(Property);
+		IsConcreteTypeCompatibleWithReflectedType<FString>(Property) ||
+		IsConcreteTypeCompatibleWithReflectedType<int32>(Property) ||
+		IsConcreteTypeCompatibleWithReflectedType<float>(Property);
 }
 
 void UTextBinding::Bind(UProperty* Property, FScriptDelegate* Delegate)
@@ -40,25 +38,43 @@ FText UTextBinding::GetTextValue() const
 
 	if ( UObject* Source = SourceObject.Get() )
 	{
-		if ( !bNeedsConversion.Get(false) )
+		if (NeedsConversion.Get(EConversion::None) == EConversion::None)
 		{
 			FText TextValue = FText::GetEmpty();
 			if ( SourcePath.GetValue<FText>(Source, TextValue) )
 			{
-				bNeedsConversion = false;
+				NeedsConversion = EConversion::None;
 				return TextValue;
 			}
 		}
 
-		if ( bNeedsConversion.Get(true) )
+		if (NeedsConversion.Get(EConversion::String) == EConversion::String)
 		{
 			FString StringValue;
-			if ( SourcePath.GetValue<FString>(Source, StringValue) )
+			if (SourcePath.GetValue<FString>(Source, StringValue))
 			{
-				bNeedsConversion = true;
+				NeedsConversion = EConversion::String;
 				return FText::FromString(StringValue);
 			}
 		}
+		if (NeedsConversion.Get(EConversion::Integer) == EConversion::Integer)
+		{
+			int32 IntegerValue;
+			if (SourcePath.GetValue<int32>(Source, IntegerValue))
+			{
+				NeedsConversion = EConversion::Integer;
+				return FText::AsNumber(IntegerValue);
+			}
+		}
+		if (NeedsConversion.Get(EConversion::Float) == EConversion::Float)
+		{
+			float FloatValue;
+			if (SourcePath.GetValue<float>(Source, FloatValue))
+			{
+				NeedsConversion = EConversion::Float;
+				return FText::AsNumber(FloatValue);
+			}
+		}
 	}
 
 	return FText::GetEmpty();
@@ -68,27 +84,45 @@ FString UTextBinding::GetStringValue() const
 {
 	//SCOPE_CYCLE_COUNTER(STAT_UMGBinding);
 
-	if ( UObject* Source = SourceObject.Get() )
+	if(UObject* Source = SourceObject.Get())
 	{
-		if ( !bNeedsConversion.Get(false) )
+		if (NeedsConversion.Get(EConversion::None) == EConversion::None)
 		{
 			FString StringValue;
 			if ( SourcePath.GetValue<FString>(Source, StringValue) )
 			{
-				bNeedsConversion = false;
+				NeedsConversion = EConversion::None;
 				return StringValue;
 			}
 		}
 
-		if ( bNeedsConversion.Get(true) )
+		if (NeedsConversion.Get(EConversion::Words) == EConversion::Words)
 		{
 			FText TextValue = FText::GetEmpty();
-			if ( SourcePath.GetValue<FText>(Source, TextValue) )
+			if (SourcePath.GetValue<FText>(Source, TextValue))
 			{
-				bNeedsConversion = true;
+				NeedsConversion = EConversion::Words;
 				return TextValue.ToString();
 			}
 		}
+		if (NeedsConversion.Get(EConversion::Integer) == EConversion::Integer)
+		{
+			int32 IntegerValue;
+			if (SourcePath.GetValue<int32>(Source, IntegerValue))
+			{
+				NeedsConversion = EConversion::Integer;
+				return FString::FromInt(IntegerValue);
+			}
+		}
+		if (NeedsConversion.Get(EConversion::Float) == EConversion::Float)
+		{
+			float FloatValue;
+			if (SourcePath.GetValue<float>(Source, FloatValue))
+			{
+				NeedsConversion = EConversion::Float;
+				return FString::SanitizeFloat(FloatValue);
+			}
+		}
 	}
 
 	return FString();
diff --git a/Engine/Source/Runtime/UMG/Public/Binding/TextBinding.h b/Engine/Source/Runtime/UMG/Public/Binding/TextBinding.h
--- a/Engine/Source/Runtime/UMG/Public/Binding/TextBinding.h
+++ b/Engine/Source/Runtime/UMG/Public/Binding/TextBinding.h
@@ -15,8 +15,6 @@ class UMG_API UTextBinding : public UPropertyBinding
 
 public:
 
-	UTextBinding();
-
 	virtual bool IsSupportedSource(UProperty* Property) const override;
 	virtual bool IsSupportedDestination(UProperty* Property) const override;
 
@@ -29,5 +27,15 @@ public:
 	FString GetStringValue() const;
 
 private:
-	mutable TOptional<bool> bNeedsConversion;
+
+	enum class EConversion : uint8
+	{
+		None,
+		String,
+		Words,
+		Integer,
+		Float
+	};
+
+	mutable TOptional<EConversion> NeedsConversion;
 };
