PR #5699: Remove padding from scrollbar borders.
Because of the (2,2,2,2) default padding on two of the borders that make up the scrollbar, a scrollbar could only have a minimum width of 9, regardless of what the scrollbar thickness property was. Now a scrollbar of thickness 1 and of padding zero (and of brush images size of 1) will be of size 1.
The scrollbar outer padding is now exposed as a property with the same default values as before as to not affect existing widget layout (in SScrollBar, UScrollBar, SScrollBox and UScrollBox).
Since the padding around the child SSpacer for the scrollbar thumb went from (2,2,2,2) to zero, the scrollbar thickness values (default value and values explicitly set in code) have been incremented by 4 so that existing slate widget layout is not affected.
UMG scrollbar thickness data has been deprecated to add 4 to the serialized thickness value so that existing UMG widget layout is not affected.
User plugins could be affected by this change: if a thickness is specified by the plugin, the size of the scrollbar could end up being 4 pixels smaller than intended (but not necessarily; as it also depends on the size of the top/bottom slot images).

diff --git a/Engine/Plugins/Developer/Concert/ConcertFrontend/Source/ConcertFrontend/Private/Widgets/SConcertScrollBox.cpp b/Engine/Plugins/Developer/Concert/ConcertFrontend/Source/ConcertFrontend/Private/Widgets/SConcertScrollBox.cpp
--- a/Engine/Plugins/Developer/Concert/ConcertFrontend/Source/ConcertFrontend/Private/Widgets/SConcertScrollBox.cpp
+++ b/Engine/Plugins/Developer/Concert/ConcertFrontend/Source/ConcertFrontend/Private/Widgets/SConcertScrollBox.cpp
@@ -16,7 +16,7 @@ void SConcertScrollBox::Construct(const FArguments& InArgs)
 	bPreventLock = false;
 
 	ScrollBar = SNew(SScrollBar)
-		.Thickness(FVector2D(8.0f, 8.0f));
+		.Thickness(FVector2D(12.0f, 12.0f));
 
 	ChildSlot
 		[
diff --git a/Engine/Plugins/Experimental/CodeEditor/Source/CodeEditor/Private/SCodeEditor.cpp b/Engine/Plugins/Experimental/CodeEditor/Source/CodeEditor/Private/SCodeEditor.cpp
--- a/Engine/Plugins/Experimental/CodeEditor/Source/CodeEditor/Private/SCodeEditor.cpp
+++ b/Engine/Plugins/Experimental/CodeEditor/Source/CodeEditor/Private/SCodeEditor.cpp
@@ -32,12 +32,12 @@ void SCodeEditor::Construct(const FArguments& InArgs, UCodeProjectItem* InCodePr
 	HorizontalScrollbar = 
 		SNew(SScrollBar)
 		.Orientation(Orient_Horizontal)
-		.Thickness(FVector2D(10.0f, 10.0f));
+		.Thickness(FVector2D(14.0f, 14.0f));
 
 	VerticalScrollbar = 
 		SNew(SScrollBar)
 		.Orientation(Orient_Vertical)
-		.Thickness(FVector2D(10.0f, 10.0f));
+		.Thickness(FVector2D(14.0f, 14.0f));
 
 	ChildSlot
 	[
diff --git a/Engine/Plugins/FX/Niagara/Source/NiagaraEditor/Private/Widgets/SNiagaraGeneratedCodeView.cpp b/Engine/Plugins/FX/Niagara/Source/NiagaraEditor/Private/Widgets/SNiagaraGeneratedCodeView.cpp
--- a/Engine/Plugins/FX/Niagara/Source/NiagaraEditor/Private/Widgets/SNiagaraGeneratedCodeView.cpp
+++ b/Engine/Plugins/FX/Niagara/Source/NiagaraEditor/Private/Widgets/SNiagaraGeneratedCodeView.cpp
@@ -476,14 +476,14 @@ void SNiagaraGeneratedCodeView::UpdateUI()
 		{
 			GeneratedCode[i].HorizontalScrollBar = SNew(SScrollBar)
 				.Orientation(Orient_Horizontal)
-				.Thickness(FVector2D(8.0f, 8.0f));
+				.Thickness(FVector2D(12.0f, 12.0f));
 		}
 
 		if (!GeneratedCode[i].VerticalScrollBar.IsValid())
 		{
 			GeneratedCode[i].VerticalScrollBar = SNew(SScrollBar)
 				.Orientation(Orient_Vertical)
-				.Thickness(FVector2D(8.0f, 8.0f));
+				.Thickness(FVector2D(12.0f, 12.0f));
 		}
 		
 		if (!GeneratedCode[i].Container.IsValid())
diff --git a/Engine/Plugins/FX/Niagara/Source/NiagaraEditor/Private/Widgets/SNiagaraSpreadsheetView.cpp b/Engine/Plugins/FX/Niagara/Source/NiagaraEditor/Private/Widgets/SNiagaraSpreadsheetView.cpp
--- a/Engine/Plugins/FX/Niagara/Source/NiagaraEditor/Private/Widgets/SNiagaraSpreadsheetView.cpp
+++ b/Engine/Plugins/FX/Niagara/Source/NiagaraEditor/Private/Widgets/SNiagaraSpreadsheetView.cpp
@@ -281,19 +281,19 @@ void SNiagaraSpreadsheetView::Construct(const FArguments& InArgs, TSharedRef<FNi
 
 		CaptureData[i].OutputHorizontalScrollBar = SNew(SScrollBar)
 			.Orientation(Orient_Horizontal)
-			.Thickness(FVector2D(8.0f, 8.0f));
+			.Thickness(FVector2D(12.0f, 12.0f));
 
 		CaptureData[i].OutputVerticalScrollBar = SNew(SScrollBar)
 			.Orientation(Orient_Vertical)
-			.Thickness(FVector2D(8.0f, 8.0f));
+			.Thickness(FVector2D(12.0f, 12.0f));
 
 		CaptureData[i].InputHorizontalScrollBar = SNew(SScrollBar)
 			.Orientation(Orient_Horizontal)
-			.Thickness(FVector2D(8.0f, 8.0f));
+			.Thickness(FVector2D(12.0f, 12.0f));
 
 		CaptureData[i].InputVerticalScrollBar = SNew(SScrollBar)
 			.Orientation(Orient_Vertical)
-			.Thickness(FVector2D(8.0f, 8.0f));
+			.Thickness(FVector2D(12.0f, 12.0f));
 
 		SAssignNew(CaptureData[i].OutputsListView, STreeView< TSharedPtr<int32> >)
 			.IsEnabled(this, &SNiagaraSpreadsheetView::IsPausedAtRightTimeOnRightHandle)
diff --git a/Engine/Source/Developer/LogVisualizer/Private/SVisualLoggerView.cpp b/Engine/Source/Developer/LogVisualizer/Private/SVisualLoggerView.cpp
--- a/Engine/Source/Developer/LogVisualizer/Private/SVisualLoggerView.cpp
+++ b/Engine/Source/Developer/LogVisualizer/Private/SVisualLoggerView.cpp
@@ -79,7 +79,7 @@ void SVisualLoggerView::Construct(const FArguments& InArgs, const TSharedRef<FUI
 	TSharedRef<SScrollBar> ZoomScrollBar =
 		SNew(SScrollBar)
 		.Orientation(EOrientation::Orient_Horizontal)
-		.Thickness(FVector2D(2.0f, 2.0f));
+		.Thickness(FVector2D(6.0f, 6.0f));
 	ZoomScrollBar->SetState(0.0f, 1.0f);
 	FLogVisualizer::Get().GetTimeSliderController()->SetExternalScrollbar(ZoomScrollBar);
 
@@ -90,7 +90,7 @@ void SVisualLoggerView::Construct(const FArguments& InArgs, const TSharedRef<FUI
 
 	TSharedRef<SScrollBar> ScrollBar =
 		SNew(SScrollBar)
-		.Thickness(FVector2D(2.0f, 2.0f));
+		.Thickness(FVector2D(6.0f, 6.0f));
 
 
 	ULogVisualizerSettings* Settings = ULogVisualizerSettings::StaticClass()->GetDefaultObject<ULogVisualizerSettings>();
diff --git a/Engine/Source/Developer/Profiler/Private/Widgets/SProfilerGraphPanel.cpp b/Engine/Source/Developer/Profiler/Private/Widgets/SProfilerGraphPanel.cpp
--- a/Engine/Source/Developer/Profiler/Private/Widgets/SProfilerGraphPanel.cpp
+++ b/Engine/Source/Developer/Profiler/Private/Widgets/SProfilerGraphPanel.cpp
@@ -82,7 +82,7 @@ void SProfilerGraphPanel::Construct(const FArguments& InArgs)
 					.Orientation(Orient_Horizontal)
 					.AlwaysShowScrollbar(true)
 					.Visibility(EVisibility::Visible)
-					.Thickness(FVector2D(8.0f, 8.0f))
+					.Thickness(FVector2D(12.0f, 12.0f))
 					.OnUserScrolled(this, &SProfilerGraphPanel::HorizontalScrollBar_OnUserScrolled)
 				]
 			]
@@ -94,7 +94,7 @@ void SProfilerGraphPanel::Construct(const FArguments& InArgs)
 				.Orientation(Orient_Vertical)
 				.AlwaysShowScrollbar(true)
 				.Visibility(EVisibility::Visible)
-				.Thickness(FVector2D(8.0f, 8.0f))
+				.Thickness(FVector2D(12.0f, 12.0f))
 				.OnUserScrolled(this, &SProfilerGraphPanel::VerticalScrollBar_OnUserScrolled)
 			]
 		]
diff --git a/Engine/Source/Editor/CurveTableEditor/Private/CurveTableEditor.cpp b/Engine/Source/Editor/CurveTableEditor/Private/CurveTableEditor.cpp
--- a/Engine/Source/Editor/CurveTableEditor/Private/CurveTableEditor.cpp
+++ b/Engine/Source/Editor/CurveTableEditor/Private/CurveTableEditor.cpp
@@ -274,11 +274,11 @@ TSharedRef<SDockTab> FCurveTableEditor::SpawnTab_CurveTable( const FSpawnTabArgs
 
 	TSharedRef<SScrollBar> HorizontalScrollBar = SNew(SScrollBar)
 		.Orientation(Orient_Horizontal)
-		.Thickness(FVector2D(8.0f, 8.0f));
+		.Thickness(FVector2D(12.0f, 12.0f));
 
 	TSharedRef<SScrollBar> VerticalScrollBar = SNew(SScrollBar)
 		.Orientation(Orient_Vertical)
-		.Thickness(FVector2D(8.0f, 8.0f));
+		.Thickness(FVector2D(12.0f, 12.0f));
 
 	TSharedRef<SHeaderRow> RowNamesHeaderRow = SNew(SHeaderRow)
 		.Visibility(this, &FCurveTableEditor::GetGridViewControlsVisibility);
diff --git a/Engine/Source/Editor/DataTableEditor/Private/DataTableEditor.cpp b/Engine/Source/Editor/DataTableEditor/Private/DataTableEditor.cpp
--- a/Engine/Source/Editor/DataTableEditor/Private/DataTableEditor.cpp
+++ b/Engine/Source/Editor/DataTableEditor/Private/DataTableEditor.cpp
@@ -784,11 +784,11 @@ TSharedRef<SVerticalBox> FDataTableEditor::CreateContentBox()
 {
 	TSharedRef<SScrollBar> HorizontalScrollBar = SNew(SScrollBar)
 		.Orientation(Orient_Horizontal)
-		.Thickness(FVector2D(8.0f, 8.0f));
+		.Thickness(FVector2D(12.0f, 12.0f));
 
 	TSharedRef<SScrollBar> VerticalScrollBar = SNew(SScrollBar)
 		.Orientation(Orient_Vertical)
-		.Thickness(FVector2D(8.0f, 8.0f));
+		.Thickness(FVector2D(12.0f, 12.0f));
 
 	TSharedRef<SHeaderRow> RowNamesHeaderRow = SNew(SHeaderRow);
 	RowNamesHeaderRow->AddColumn(
diff --git a/Engine/Source/Editor/MaterialEditor/Private/SMaterialEditorStatsWidget.cpp b/Engine/Source/Editor/MaterialEditor/Private/SMaterialEditorStatsWidget.cpp
--- a/Engine/Source/Editor/MaterialEditor/Private/SMaterialEditorStatsWidget.cpp
+++ b/Engine/Source/Editor/MaterialEditor/Private/SMaterialEditorStatsWidget.cpp
@@ -570,11 +570,11 @@ void SMaterialEditorStatsWidget::Construct(const FArguments& InArgs)
 
 	const auto VerticalScrollbar = SNew(SScrollBar)
 		.Orientation(Orient_Vertical)
-		.Thickness(FVector2D(7.0f, 7.0f));
+		.Thickness(FVector2D(11.0f, 11.0f));
 
 	const auto HorizontalScrollbar = SNew(SScrollBar)
 		.Orientation(Orient_Horizontal)
-		.Thickness(FVector2D(7.0f, 7.0f));
+		.Thickness(FVector2D(11.0f, 11.0f));
 
 	this->ChildSlot
 	[
diff --git a/Engine/Source/Editor/PlacementMode/Private/SPlacementModeTools.cpp b/Engine/Source/Editor/PlacementMode/Private/SPlacementModeTools.cpp
--- a/Engine/Source/Editor/PlacementMode/Private/SPlacementModeTools.cpp
+++ b/Engine/Source/Editor/PlacementMode/Private/SPlacementModeTools.cpp
@@ -321,7 +321,7 @@ void SPlacementModeTools::Construct( const FArguments& InArgs )
 	}
 
 	TSharedRef<SScrollBar> ScrollBar = SNew(SScrollBar)
-		.Thickness(FVector2D(5, 5));
+		.Thickness(FVector2D(9.0f, 9.0f));
 
 	ChildSlot
 	[
diff --git a/Engine/Source/Editor/Sequencer/Private/SSequencer.cpp b/Engine/Source/Editor/Sequencer/Private/SSequencer.cpp
--- a/Engine/Source/Editor/Sequencer/Private/SSequencer.cpp
+++ b/Engine/Source/Editor/Sequencer/Private/SSequencer.cpp
@@ -210,7 +210,7 @@ void SSequencer::Construct(const FArguments& InArgs, TSharedRef<FSequencer> InSe
 	}
 
 	TSharedRef<SScrollBar> ScrollBar = SNew(SScrollBar)
-		.Thickness(FVector2D(5.0f, 5.0f));
+		.Thickness(FVector2D(9.0f, 9.0f));
 	SAssignNew(TrackOutliner, SSequencerTrackOutliner);
 	SAssignNew(TrackArea, SSequencerTrackArea, TimeSliderControllerRef, InSequencer);
 	SAssignNew(TreeView, SSequencerTreeView, InSequencer->GetNodeTree(), TrackArea.ToSharedRef())
diff --git a/Engine/Source/Runtime/Core/Public/UObject/EditorObjectVersion.h b/Engine/Source/Runtime/Core/Public/UObject/EditorObjectVersion.h
--- a/Engine/Source/Runtime/Core/Public/UObject/EditorObjectVersion.h
+++ b/Engine/Source/Runtime/Core/Public/UObject/EditorObjectVersion.h
@@ -75,6 +75,8 @@ struct CORE_API FEditorObjectVersion
 		ChangedWidgetComponentWindowVisibilityDefault,
 		// Avoid keying culture invariant display strings during serialization to avoid non-deterministic cooking issues
 		CultureInvariantTextSerializationKeyStability,
+		// Change to UScrollBar and UScrollBox thickness property (removed implicit padding of 2, so thickness value must be incremented by 4).
+		ScrollBarThicknessChange,
 		// -----<new versions can be added above this line>-------------------------------------------------
 		VersionPlusOne,
 		LatestVersion = VersionPlusOne - 1
diff --git a/Engine/Source/Runtime/Slate/Private/Widgets/Input/SMultiLineEditableTextBox.cpp b/Engine/Source/Runtime/Slate/Private/Widgets/Input/SMultiLineEditableTextBox.cpp
--- a/Engine/Source/Runtime/Slate/Private/Widgets/Input/SMultiLineEditableTextBox.cpp
+++ b/Engine/Source/Runtime/Slate/Private/Widgets/Input/SMultiLineEditableTextBox.cpp
@@ -57,7 +57,7 @@ void SMultiLineEditableTextBox::Construct( const FArguments& InArgs )
 			.Style(&InArgs._Style->ScrollBarStyle)
 			.Orientation(Orient_Horizontal)
 			.AlwaysShowScrollbar(InArgs._AlwaysShowScrollbars)
-			.Thickness(FVector2D(5.0f, 5.0f));
+			.Thickness(FVector2D(9.0f, 9.0f));
 	}
 	
 	bHasExternalVScrollBar = InArgs._VScrollBar.IsValid();
@@ -69,7 +69,7 @@ void SMultiLineEditableTextBox::Construct( const FArguments& InArgs )
 			.Style(&InArgs._Style->ScrollBarStyle)
 			.Orientation(Orient_Vertical)
 			.AlwaysShowScrollbar(InArgs._AlwaysShowScrollbars)
-			.Thickness(FVector2D(5.0f, 5.0f));
+			.Thickness(FVector2D(9.0f, 9.0f));
 	}
 
 	SBorder::Construct( SBorder::FArguments()
diff --git a/Engine/Source/Runtime/Slate/Private/Widgets/Layout/SScrollBar.cpp b/Engine/Source/Runtime/Slate/Private/Widgets/Layout/SScrollBar.cpp
--- a/Engine/Source/Runtime/Slate/Private/Widgets/Layout/SScrollBar.cpp
+++ b/Engine/Source/Runtime/Slate/Private/Widgets/Layout/SScrollBar.cpp
@@ -27,11 +27,12 @@ void SScrollBar::Construct(const FArguments& InArgs)
 
 	SBorder::Construct( SBorder::FArguments()
 		.BorderImage(FCoreStyle::Get().GetBrush("NoBorder"))
+		.Padding(InArgs._Padding)
 		[
 			SNew(SVerticalBox)
 
 			+ SVerticalBox::Slot()
-			.FillHeight( 1 )
+			.FillHeight(1)
 			[
 				SNew(SBorder)
 				.BorderImage(BackgroundBrush)
@@ -55,10 +56,11 @@ void SScrollBar::Construct(const FArguments& InArgs)
 					.ThumbSlot()
 					[
 						SAssignNew(DragThumb, SBorder)
-						.BorderImage( this, &SScrollBar::GetDragThumbImage )
+						.BorderImage(this, &SScrollBar::GetDragThumbImage)
 						.ColorAndOpacity(this, &SScrollBar::GetThumbOpacity)
 						.HAlign(HAlign_Center)
 						.VAlign(VAlign_Center)
+						.Padding(0)
 						[
 							SAssignNew(ThicknessSpacer, SSpacer)
 							.Size(InArgs._Thickness)
diff --git a/Engine/Source/Runtime/Slate/Private/Widgets/Layout/SScrollBox.cpp b/Engine/Source/Runtime/Slate/Private/Widgets/Layout/SScrollBox.cpp
--- a/Engine/Source/Runtime/Slate/Private/Widgets/Layout/SScrollBox.cpp
+++ b/Engine/Source/Runtime/Slate/Private/Widgets/Layout/SScrollBox.cpp
@@ -144,6 +144,7 @@ void SScrollBox::Construct( const FArguments& InArgs )
 		ScrollBar = ConstructScrollBar();
 		ScrollBar->SetDragFocusCause(InArgs._ScrollBarDragFocusCause);
 		ScrollBar->SetThickness(InArgs._ScrollBarThickness);
+		ScrollBar->SetPadding(InArgs._ScrollBarPadding);
 		ScrollBar->SetUserVisibility(InArgs._ScrollBarVisibility);
 		ScrollBar->SetScrollBarAlwaysVisible(InArgs._ScrollBarAlwaysVisible);
 
@@ -496,6 +497,11 @@ void SScrollBox::SetScrollBarThickness(FVector2D InThickness)
 	ScrollBar->SetThickness(InThickness);
 }
 
+void SScrollBox::SetScrollBarPadding(const FMargin& InPadding)
+{
+	ScrollBar->SetPadding(InPadding);
+}
+
 void SScrollBox::SetScrollBarRightClickDragAllowed(bool bIsAllowed)
 {
 	bAllowsRightClickDragScrolling = bIsAllowed;
diff --git a/Engine/Source/Runtime/Slate/Public/Widgets/Layout/SScrollBar.h b/Engine/Source/Runtime/Slate/Public/Widgets/Layout/SScrollBar.h
--- a/Engine/Source/Runtime/Slate/Public/Widgets/Layout/SScrollBar.h
+++ b/Engine/Source/Runtime/Slate/Public/Widgets/Layout/SScrollBar.h
@@ -39,7 +39,8 @@ public:
 #endif
 		, _Orientation( Orient_Vertical )
 		, _DragFocusCause( EFocusCause::Mouse )
-		, _Thickness( FVector2D(12.0f, 12.0f) )
+		, _Thickness( FVector2D(16.0f, 16.0f) )
+		, _Padding( 2.0f )
 		{}
 
 		/** The style to use for this scrollbar */
@@ -52,6 +53,8 @@ public:
 		SLATE_ARGUMENT( EFocusCause, DragFocusCause )
 		/** The thickness of the scrollbar thumb */
 		SLATE_ATTRIBUTE( FVector2D, Thickness )
+		/** The margin around the scrollbar */
+		SLATE_ATTRIBUTE( FMargin, Padding )
 	SLATE_END_ARGS()
 
 	/**
diff --git a/Engine/Source/Runtime/Slate/Public/Widgets/Layout/SScrollBox.h b/Engine/Source/Runtime/Slate/Public/Widgets/Layout/SScrollBox.h
--- a/Engine/Source/Runtime/Slate/Public/Widgets/Layout/SScrollBox.h
+++ b/Engine/Source/Runtime/Slate/Public/Widgets/Layout/SScrollBox.h
@@ -89,7 +89,8 @@ public:
 		, _ScrollBarVisibility(EVisibility::Visible)
 		, _ScrollBarAlwaysVisible(false)
 		, _ScrollBarDragFocusCause(EFocusCause::Mouse)
-		, _ScrollBarThickness(FVector2D(5, 5))
+		, _ScrollBarThickness(FVector2D(9.0f, 9.0f))
+		, _ScrollBarPadding(2.0f)
 		, _AllowOverscroll(EAllowOverscroll::Yes)
 		, _NavigationDestination(EDescendantScrollDestination::IntoView)
 		, _NavigationScrollPadding(0.0f)
@@ -121,6 +122,8 @@ public:
 
 		SLATE_ARGUMENT( FVector2D, ScrollBarThickness )
 
+		SLATE_ARGUMENT( FMargin, ScrollBarPadding )
+
 		SLATE_ARGUMENT(EAllowOverscroll, AllowOverscroll);
 
 		SLATE_ARGUMENT(EDescendantScrollDestination, NavigationDestination);
@@ -204,6 +207,8 @@ public:
 
 	void SetScrollBarThickness(FVector2D InThickness);
 
+	void SetScrollBarPadding(const FMargin& InPadding);
+
 	void SetScrollBarRightClickDragAllowed(bool bIsAllowed);
 public:
 
diff --git a/Engine/Source/Runtime/UMG/Private/Components/ScrollBar.cpp b/Engine/Source/Runtime/UMG/Private/Components/ScrollBar.cpp
--- a/Engine/Source/Runtime/UMG/Private/Components/ScrollBar.cpp
+++ b/Engine/Source/Runtime/UMG/Private/Components/ScrollBar.cpp
@@ -1,6 +1,7 @@
 // Copyright 1998-2019 Epic Games, Inc. All Rights Reserved.
 
 #include "Components/ScrollBar.h"
+#include "UObject/EditorObjectVersion.h"
 
 #define LOCTEXT_NAMESPACE "UMG"
 
@@ -15,7 +16,8 @@ UScrollBar::UScrollBar(const FObjectInitializer& ObjectInitializer)
 	bAlwaysShowScrollbar = true;
 	bAlwaysShowScrollbarTrack = true;
 	Orientation = Orient_Vertical;
-	Thickness = FVector2D(12.0f, 12.0f);
+	Thickness = FVector2D(16.0f, 16.0f);
+	Padding = FMargin(2.0f);
 
 	// HACK: THIS SHOULD NOT COME FROM CORESTYLE AND SHOULD INSTEAD BY DEFINED BY ENGINE TEXTURES/PROJECT SETTINGS
 	static const FScrollBarStyle StaticScrollbar = FCoreStyle::Get().GetWidgetStyle<FScrollBarStyle>("Scrollbar");
@@ -36,7 +38,8 @@ TSharedRef<SWidget> UScrollBar::RebuildWidget()
 		.AlwaysShowScrollbar(bAlwaysShowScrollbar)
 		.AlwaysShowScrollbarTrack(bAlwaysShowScrollbarTrack)
 		.Orientation(Orientation)
-		.Thickness(Thickness);
+		.Thickness(Thickness)
+		.Padding(Padding);
 
 	//SLATE_EVENT(FOnUserScrolled, OnUserScrolled)
 
@@ -58,6 +61,28 @@ void UScrollBar::SetState(float InOffsetFraction, float InThumbSizeFraction)
 	}
 }
 
+#if WITH_EDITORONLY_DATA
+
+void UScrollBar::Serialize(FArchive& Ar)
+{
+	Ar.UsingCustomVersion(FEditorObjectVersion::GUID);
+
+	const bool bDeprecateThickness = Ar.IsLoading() && Ar.CustomVer(FEditorObjectVersion::GUID) < FEditorObjectVersion::ScrollBarThicknessChange;
+	if (bDeprecateThickness)
+	{
+		// Set Thickness property to previous default value.
+		Thickness.Set(12.0f, 12.0f);
+	}
+
+	Super::Serialize(Ar);
+
+	if (bDeprecateThickness)
+	{
+		// Implicit padding of 2 was removed, so Thickness value must be incremented by 4.
+		Thickness += FVector2D(4.0f, 4.0f);
+	}
+}
+
 void UScrollBar::PostLoad()
 {
 	Super::PostLoad();
@@ -77,6 +102,8 @@ void UScrollBar::PostLoad()
 	}
 }
 
+#endif // if WITH_EDITORONLY_DATA
+
 #if WITH_EDITOR
 
 const FText UScrollBar::GetPaletteCategory()
diff --git a/Engine/Source/Runtime/UMG/Private/Components/ScrollBox.cpp b/Engine/Source/Runtime/UMG/Private/Components/ScrollBox.cpp
--- a/Engine/Source/Runtime/UMG/Private/Components/ScrollBox.cpp
+++ b/Engine/Source/Runtime/UMG/Private/Components/ScrollBox.cpp
@@ -3,6 +3,7 @@
 #include "Components/ScrollBox.h"
 #include "Containers/Ticker.h"
 #include "Components/ScrollBoxSlot.h"
+#include "UObject/EditorObjectVersion.h"
 
 #define LOCTEXT_NAMESPACE "UMG"
 
@@ -14,7 +15,8 @@ UScrollBox::UScrollBox(const FObjectInitializer& ObjectInitializer)
 	, Orientation(Orient_Vertical)
 	, ScrollBarVisibility(ESlateVisibility::Visible)
 	, ConsumeMouseWheel(EConsumeMouseWheel::WhenScrollingPossible)
-	, ScrollbarThickness(5.0f, 5.0f)
+	, ScrollbarThickness(9.0f, 9.0f)
+	, ScrollbarPadding(2.0f)
 	, AlwaysShowScrollbar(false)
 	, AlwaysShowScrollbarTrack(false)
 	, AllowOverscroll(true)
@@ -103,6 +105,7 @@ void UScrollBox::SynchronizeProperties()
 	MyScrollBox->SetOrientation(Orientation);
 	MyScrollBox->SetScrollBarVisibility(UWidget::ConvertSerializedVisibilityToRuntime(ScrollBarVisibility));
 	MyScrollBox->SetScrollBarThickness(ScrollbarThickness);
+	MyScrollBox->SetScrollBarPadding(ScrollbarPadding);
 	MyScrollBox->SetScrollBarAlwaysVisible(AlwaysShowScrollbar);
 	MyScrollBox->SetScrollBarTrackAlwaysVisible(AlwaysShowScrollbarTrack);
 	MyScrollBox->SetAllowOverscroll(AllowOverscroll ? EAllowOverscroll::Yes : EAllowOverscroll::No);
@@ -182,6 +185,28 @@ void UScrollBox::ScrollWidgetIntoView(UWidget* WidgetToFind, bool AnimateScroll,
 	}
 }
 
+#if WITH_EDITORONLY_DATA
+
+void UScrollBox::Serialize(FArchive& Ar)
+{
+	Ar.UsingCustomVersion(FEditorObjectVersion::GUID);
+	
+	const bool bDeprecateThickness = Ar.IsLoading() && Ar.CustomVer(FEditorObjectVersion::GUID) < FEditorObjectVersion::ScrollBarThicknessChange;
+	if (bDeprecateThickness)
+	{
+		// Set ScrollbarThickness property to previous default value.
+		ScrollbarThickness.Set(5.0f, 5.0f);
+	}
+
+	Super::Serialize(Ar);
+
+	if (bDeprecateThickness)
+	{
+		// Implicit padding of 2 was removed, so ScrollbarThickness value must be incremented by 4.
+		ScrollbarThickness += FVector2D(4.0f, 4.0f);
+	}
+}
+
 void UScrollBox::PostLoad()
 {
 	Super::PostLoad();
@@ -212,6 +237,8 @@ void UScrollBox::PostLoad()
 	}
 }
 
+#endif // if WITH_EDITORONLY_DATA
+
 void UScrollBox::SetConsumeMouseWheel(EConsumeMouseWheel NewConsumeMouseWheel)
 {
 	ConsumeMouseWheel = NewConsumeMouseWheel;
@@ -259,6 +286,16 @@ void UScrollBox::SetScrollbarThickness(const FVector2D& NewScrollbarThickness)
 	}
 }
 
+void UScrollBox::SetScrollbarPadding(const FMargin& NewScrollbarPadding)
+{
+	ScrollbarPadding = NewScrollbarPadding;
+
+	if (MyScrollBox.IsValid())
+	{
+		MyScrollBox->SetScrollBarPadding(ScrollbarPadding);
+	}
+}
+
 void UScrollBox::SetAlwaysShowScrollbar(bool NewAlwaysShowScrollbar)
 {
 	AlwaysShowScrollbar = NewAlwaysShowScrollbar;
diff --git a/Engine/Source/Runtime/UMG/Public/Components/ScrollBar.h b/Engine/Source/Runtime/UMG/Public/Components/ScrollBar.h
--- a/Engine/Source/Runtime/UMG/Public/Components/ScrollBar.h
+++ b/Engine/Source/Runtime/UMG/Public/Components/ScrollBar.h
@@ -41,6 +41,10 @@ public:
 	UPROPERTY(EditAnywhere, Category="Behavior")
 	FVector2D Thickness;
 
+	/** The margin around the scrollbar */
+	UPROPERTY(EditAnywhere, Category = "Behavior")
+	FMargin Padding;
+
 public:
 
 	/**
@@ -78,7 +82,10 @@ public:
 	//~ End UVisual Interface
 
 	//~ Begin UObject Interface
+#if WITH_EDITORONLY_DATA
+	virtual void Serialize(FArchive& Ar) override;
 	virtual void PostLoad() override;
+#endif // if WITH_EDITORONLY_DATA
 	//~ End UObject Interface
 
 #if WITH_EDITOR
diff --git a/Engine/Source/Runtime/UMG/Public/Components/ScrollBox.h b/Engine/Source/Runtime/UMG/Public/Components/ScrollBox.h
--- a/Engine/Source/Runtime/UMG/Public/Components/ScrollBox.h
+++ b/Engine/Source/Runtime/UMG/Public/Components/ScrollBox.h
@@ -48,10 +48,14 @@ public:
 	UPROPERTY(EditAnywhere, BlueprintReadOnly, Category = "Scroll")
 	EConsumeMouseWheel ConsumeMouseWheel;
 
-	/**  */
+	/** The thickness of the scrollbar thumb */
 	UPROPERTY(EditAnywhere, BlueprintReadOnly, Category="Scroll")
 	FVector2D ScrollbarThickness;
 
+	/** The margin around the scrollbar */
+	UPROPERTY(EditAnywhere, BlueprintReadOnly, Category = "Scroll")
+	FMargin ScrollbarPadding;
+
 	/**  */
 	UPROPERTY(EditAnywhere, BlueprintReadOnly, Category="Scroll")
 	bool AlwaysShowScrollbar;
@@ -92,6 +96,9 @@ public:
 	UFUNCTION(BlueprintCallable, Category = "Scroll")
 	void SetScrollbarThickness(const FVector2D& NewScrollbarThickness);
 
+	UFUNCTION(BlueprintCallable, Category = "Scroll")
+	void SetScrollbarPadding(const FMargin& NewScrollbarPadding);
+
 	UFUNCTION(BlueprintCallable, Category = "Scroll")
 	void SetAlwaysShowScrollbar(bool NewAlwaysShowScrollbar);
 	
@@ -149,7 +156,10 @@ public:
 	//~ End UVisual Interface
 
 	//~ Begin UObject Interface
+#if WITH_EDITORONLY_DATA
+	virtual void Serialize(FArchive& Ar) override;
 	virtual void PostLoad() override;
+#endif // if WITH_EDITORONLY_DATA
 	//~ End UObject Interface
 
 #if WITH_EDITOR
