Correction to ListViewBase to validate that EntryWidgetClass does implement UUserListEntry:
- At compilation of the widget that contains the list
- Before preview Widget are created to prevent then from being created and hit an ensuremsg call

diff --git a/Engine/Source/Runtime/UMG/Private/Components/ListViewBase.cpp b/Engine/Source/Runtime/UMG/Private/Components/ListViewBase.cpp
--- a/Engine/Source/Runtime/UMG/Private/Components/ListViewBase.cpp
+++ b/Engine/Source/Runtime/UMG/Private/Components/ListViewBase.cpp
@@ -28,6 +28,10 @@ void UListViewBase::ValidateCompiledDefaults(IWidgetCompilerLog& CompileLog) con
 	{
 		CompileLog.Error(FText::Format(LOCTEXT("Error_ListViewBase_MissingEntryClass", "{0} has no EntryWidgetClass specified - required for any UListViewBase to function."), FText::FromString(GetName())));
 	}
+	else if (!EntryWidgetClass->ImplementsInterface(UUserListEntry::StaticClass()))
+	{
+		CompileLog.Error(FText::Format(LOCTEXT("Error_ListViewBase_EntryClassNotImplementingInterface", "'{0}' has EntryWidgetClass property set to'{1}' and that Class doesn't implement User List Entry Interface - required for any UListViewBase to function."), FText::FromString(GetName()), FText::FromString(EntryWidgetClass->GetName())));
+	}
 }
 #endif
 
diff --git a/Engine/Source/Runtime/UMG/Public/Components/ListViewBase.h b/Engine/Source/Runtime/UMG/Public/Components/ListViewBase.h
--- a/Engine/Source/Runtime/UMG/Public/Components/ListViewBase.h
+++ b/Engine/Source/Runtime/UMG/Public/Components/ListViewBase.h
@@ -527,7 +527,7 @@ protected:
 	{
 		bNeedsToCallRefreshDesignerItems = false;
 		bool bRefresh = false;
-		if (EntryWidgetClass && NumDesignerPreviewEntries > 0)
+		if (EntryWidgetClass && NumDesignerPreviewEntries > 0 && EntryWidgetClass->ImplementsInterface(UUserListEntry::StaticClass()))
 		{
 			if (ListItems.Num() < NumDesignerPreviewEntries)
 			{
