Support changing WidgetStyle (FTextBlockStyle) on MultiLineEditableText in BP.
Support changing WidgetStyle (FEditableTextBoxStyle) and TextStyle (FTextBlockStyle) on MultiLineEditableTextBox in BP.

diff --git a/Engine/Source/Runtime/Slate/Private/Widgets/Input/SMultiLineEditableTextBox.cpp b/Engine/Source/Runtime/Slate/Private/Widgets/Input/SMultiLineEditableTextBox.cpp
--- a/Engine/Source/Runtime/Slate/Private/Widgets/Input/SMultiLineEditableTextBox.cpp
+++ b/Engine/Source/Runtime/Slate/Private/Widgets/Input/SMultiLineEditableTextBox.cpp
@@ -198,6 +198,11 @@ void SMultiLineEditableTextBox::SetStyle(const FEditableTextBoxStyle* InStyle)
 	BorderImageReadOnly = &Style->BackgroundImageReadOnly;
 }
 
+void SMultiLineEditableTextBox::SetTextStyle(const FTextBlockStyle* InTextStyle)
+{
+	EditableText->SetTextStyle(InTextStyle);
+}
+
 FSlateColor SMultiLineEditableTextBox::DetermineForegroundColor() const
 {
 	check(Style);
diff --git a/Engine/Source/Runtime/Slate/Public/Widgets/Input/SMultiLineEditableTextBox.h b/Engine/Source/Runtime/Slate/Public/Widgets/Input/SMultiLineEditableTextBox.h
--- a/Engine/Source/Runtime/Slate/Public/Widgets/Input/SMultiLineEditableTextBox.h
+++ b/Engine/Source/Runtime/Slate/Public/Widgets/Input/SMultiLineEditableTextBox.h
@@ -260,6 +260,9 @@ public:
 	/** See attribute Style */
 	void SetStyle(const FEditableTextBoxStyle* InStyle);
 
+	/** See attribute TextStyle */
+	void SetTextStyle(const FTextBlockStyle* InTextStyle);	
+
 	/**
 	 * Sets the text string currently being edited 
 	 *
diff --git a/Engine/Source/Runtime/UMG/Private/Components/MultiLineEditableText.cpp b/Engine/Source/Runtime/UMG/Private/Components/MultiLineEditableText.cpp
--- a/Engine/Source/Runtime/UMG/Private/Components/MultiLineEditableText.cpp
+++ b/Engine/Source/Runtime/UMG/Private/Components/MultiLineEditableText.cpp
@@ -123,6 +123,16 @@ void UMultiLineEditableText::SetIsReadOnly(bool bReadOnly)
 	}
 }
 
+void UMultiLineEditableText::SetWidgetStyle(const FTextBlockStyle& InWidgetStyle)
+{
+	WidgetStyle = InWidgetStyle;
+
+	if (MyMultiLineEditableText.IsValid())
+	{
+		MyMultiLineEditableText->SetTextStyle(&WidgetStyle);
+	}
+}
+
 void UMultiLineEditableText::HandleOnTextChanged(const FText& InText)
 {
 	OnTextChanged.Broadcast(InText);
diff --git a/Engine/Source/Runtime/UMG/Private/Components/MultiLineEditableTextBox.cpp b/Engine/Source/Runtime/UMG/Private/Components/MultiLineEditableTextBox.cpp
--- a/Engine/Source/Runtime/UMG/Private/Components/MultiLineEditableTextBox.cpp
+++ b/Engine/Source/Runtime/UMG/Private/Components/MultiLineEditableTextBox.cpp
@@ -129,6 +129,26 @@ void UMultiLineEditableTextBox::SetIsReadOnly(bool bReadOnly)
 	}
 }
 
+void UMultiLineEditableTextBox::SetWidgetStyle(const FEditableTextBoxStyle& InWidgetStyle)
+{
+	WidgetStyle = InWidgetStyle;
+
+	if (MyEditableTextBlock.IsValid())
+	{
+		MyEditableTextBlock->SetStyle(&WidgetStyle);
+	}
+}
+
+void UMultiLineEditableTextBox::SetTextStyle(const FTextBlockStyle& InTextStyle)
+{
+	TextStyle = InTextStyle;
+
+	if (MyEditableTextBlock.IsValid())
+	{
+		MyEditableTextBlock->SetTextStyle(&TextStyle);
+	}
+}
+
 void UMultiLineEditableTextBox::HandleOnTextChanged(const FText& InText)
 {
 	OnTextChanged.Broadcast(InText);
diff --git a/Engine/Source/Runtime/UMG/Public/Components/MultiLineEditableText.h b/Engine/Source/Runtime/UMG/Public/Components/MultiLineEditableText.h
--- a/Engine/Source/Runtime/UMG/Public/Components/MultiLineEditableText.h
+++ b/Engine/Source/Runtime/UMG/Public/Components/MultiLineEditableText.h
@@ -41,7 +41,7 @@ public:
 public:
 
 	/** The style */
-	UPROPERTY(EditAnywhere, BlueprintReadWrite, Category="Style", meta=(ShowOnlyInnerProperties))
+	UPROPERTY(EditAnywhere, BlueprintReadWrite, BlueprintSetter=SetWidgetStyle, Category="Style", meta=(ShowOnlyInnerProperties))
 	FTextBlockStyle WidgetStyle;
 
 	/** Sets whether this text block can be modified interactively by the user */
@@ -100,6 +100,9 @@ public:
 
 	UFUNCTION(BlueprintCallable, Category="Widget", meta=(DisplayName="SetIsReadOnly (Multi-Line Editable Text"))
 	void SetIsReadOnly(bool bReadOnly);
+
+	UFUNCTION(BlueprintSetter)
+	void SetWidgetStyle(const FTextBlockStyle& InWidgetStyle);
 	
 	//~ Begin UWidget Interface
 	virtual void SynchronizeProperties() override;
diff --git a/Engine/Source/Runtime/UMG/Public/Components/MultiLineEditableTextBox.h b/Engine/Source/Runtime/UMG/Public/Components/MultiLineEditableTextBox.h
--- a/Engine/Source/Runtime/UMG/Public/Components/MultiLineEditableTextBox.h
+++ b/Engine/Source/Runtime/UMG/Public/Components/MultiLineEditableTextBox.h
@@ -43,11 +43,11 @@ public:
 public:
 
 	/** The style */
-	UPROPERTY(EditAnywhere, BlueprintReadWrite, Category="Style", meta=(DisplayName="Style"))
+	UPROPERTY(EditAnywhere, BlueprintReadWrite, BlueprintSetter = SetWidgetStyle, Category="Style", meta=(DisplayName="Style"))
 	FEditableTextBoxStyle WidgetStyle;
 
 	/** The text style */
-	UPROPERTY(EditAnywhere, BlueprintReadWrite, Category="Style", meta=(DisplayName="Text Style"))
+	UPROPERTY(EditAnywhere, BlueprintReadWrite, BlueprintSetter = SetTextStyle, Category="Style", meta=(DisplayName="Text Style"))
 	FTextBlockStyle TextStyle;
 
 	/** Sets whether this text block can be modified interactively by the user */
@@ -112,6 +112,12 @@ public:
 	UFUNCTION(BlueprintCallable, Category="Widget", meta=(DisplayName="SetIsReadOnly (Multi-Line Text Box)"))
 	void SetIsReadOnly(bool bReadOnly);
 
+	UFUNCTION(BlueprintSetter)
+	void SetWidgetStyle(const FEditableTextBoxStyle& InWidgetStyle);
+
+	UFUNCTION(BlueprintSetter)
+	void SetTextStyle(const FTextBlockStyle& InTextStyle);
+
 	//TODO UMG Add Set ReadOnlyForegroundColor
 	//TODO UMG Add Set BackgroundColor
 	//TODO UMG Add Set ForegroundColor
