Remove useless UMultiLineEditableTextBox.WidgetStyle (FEditableTextBoxStyle) BP setter, since the underlying SMultiLineEditableTextBox as a pointer to the member.
There's no need for deprecation since this setter has never made it to an official release.

diff --git a/Engine/Source/Runtime/UMG/Private/Components/MultiLineEditableTextBox.cpp b/Engine/Source/Runtime/UMG/Private/Components/MultiLineEditableTextBox.cpp
--- a/Engine/Source/Runtime/UMG/Private/Components/MultiLineEditableTextBox.cpp
+++ b/Engine/Source/Runtime/UMG/Private/Components/MultiLineEditableTextBox.cpp
@@ -147,16 +147,6 @@ void UMultiLineEditableTextBox::SetIsReadOnly(bool bReadOnly)
 	}
 }
 
-void UMultiLineEditableTextBox::SetWidgetStyle(const FEditableTextBoxStyle& InWidgetStyle)
-{
-	WidgetStyle = InWidgetStyle;
-
-	if (MyEditableTextBlock.IsValid())
-	{
-		MyEditableTextBlock->SetStyle(&WidgetStyle);
-	}
-}
-
 void UMultiLineEditableTextBox::SetTextStyle(const FTextBlockStyle& InTextStyle)
 {
 	TextStyle = InTextStyle;
diff --git a/Engine/Source/Runtime/UMG/Public/Components/MultiLineEditableTextBox.h b/Engine/Source/Runtime/UMG/Public/Components/MultiLineEditableTextBox.h
--- a/Engine/Source/Runtime/UMG/Public/Components/MultiLineEditableTextBox.h
+++ b/Engine/Source/Runtime/UMG/Public/Components/MultiLineEditableTextBox.h
@@ -43,7 +43,7 @@ public:
 public:
 
 	/** The style */
-	UPROPERTY(EditAnywhere, BlueprintReadWrite, BlueprintSetter = SetWidgetStyle, Category="Style", meta=(DisplayName="Style"))
+	UPROPERTY(EditAnywhere, BlueprintReadWrite, Category="Style", meta=(DisplayName="Style"))
 	FEditableTextBoxStyle WidgetStyle;
 
 	/** The text style */
@@ -112,9 +112,6 @@ public:
 	UFUNCTION(BlueprintCallable, Category="Widget", meta=(DisplayName="SetIsReadOnly (Multi-Line Text Box)"))
 	void SetIsReadOnly(bool bReadOnly);
 
-	UFUNCTION(BlueprintSetter)
-	void SetWidgetStyle(const FEditableTextBoxStyle& InWidgetStyle);
-
 	UFUNCTION(BlueprintSetter)
 	void SetTextStyle(const FTextBlockStyle& InTextStyle);
 
