Changed display names of UMG Visibility HitTestInvisible and SelfHitTestInvisible enums (since many of our UMG users reported that they were counter intuitive especially for newcomers) and streamlined tooltips documentation.

diff --git a/Engine/Source/Runtime/SlateCore/Public/Layout/Visibility.h b/Engine/Source/Runtime/SlateCore/Public/Layout/Visibility.h
--- a/Engine/Source/Runtime/SlateCore/Public/Layout/Visibility.h
+++ b/Engine/Source/Runtime/SlateCore/Public/Layout/Visibility.h
@@ -7,19 +7,19 @@
 /** Is an entity visible? */
 struct SLATECORE_API EVisibility
 {
-	/** Default widget visibility - visible and can interact with the cursor */
+	/** Visible and hit-testable (can interact with cursor). Default value. */
 	static const EVisibility Visible;
 
-	/** Not visible and takes up no space in the layout; can never be clicked on because it takes up no space. */
+	/** Not visible and takes up no space in the layout (obviously not hit-testable). */
 	static const EVisibility Collapsed;
 
-	/** Not visible, but occupies layout space. Not interactive for obvious reasons. */
+	/** Not visible but occupies layout space (obviously not hit-testable). */
 	static const EVisibility Hidden;
 
-	/** Visible to the user, but only as art. The cursors hit tests will never see this widget. */
+	/** Visible but not hit-testable (cannot interact with cursor) and children in the hierarchy (if any) are also not hit-testable. */
 	static const EVisibility HitTestInvisible;
 
-	/** Same as HitTestInvisible, but doesn't apply to child widgets. */
+	/** Visible but not hit-testable (cannot interact with cursor) and doesn't affect hit-testing on children (if any). */
 	static const EVisibility SelfHitTestInvisible;
 
 	/** Any visibility will do */
diff --git a/Engine/Source/Runtime/UMG/Public/Components/SlateWrapperTypes.h b/Engine/Source/Runtime/UMG/Public/Components/SlateWrapperTypes.h
--- a/Engine/Source/Runtime/UMG/Public/Components/SlateWrapperTypes.h
+++ b/Engine/Source/Runtime/UMG/Public/Components/SlateWrapperTypes.h
@@ -20,16 +20,16 @@
 UENUM(BlueprintType)
 enum class ESlateVisibility : uint8
 {
-	/** Default widget visibility - visible and can interact with the cursor */
+	/** Visible and hit-testable (can interact with cursor). Default value. */
 	Visible,
-	/** Not visible and takes up no space in the layout; can never be clicked on because it takes up no space. */
+	/** Not visible and takes up no space in the layout (obviously not hit-testable). */
 	Collapsed,
-	/** Not visible, but occupies layout space. Not interactive for obvious reasons. */
+	/** Not visible but occupies layout space (obviously not hit-testable). */
 	Hidden,
-	/** Visible to the user, but only as art. The cursors hit tests will never see this widget. */
-	HitTestInvisible,
-	/** Same as HitTestInvisible, but doesn't apply to child widgets. */
-	SelfHitTestInvisible
+	/** Visible but not hit-testable (cannot interact with cursor) and children in the hierarchy (if any) are also not hit-testable. */
+	HitTestInvisible UMETA(DisplayName = "Not Hit-Testable (Self & All Children)"),
+	/** Visible but not hit-testable (cannot interact with cursor) and doesn't affect hit-testing on children (if any). */
+	SelfHitTestInvisible UMETA(DisplayName = "Not Hit-Testable (Self Only)")
 };
 
 /** Whether a widget should be included in accessibility, and if so, how its text should be retrieved. */
