PR #5848: Fix tooltip typo "an an" (Contributed by [REDACTED_PERSON])

diff --git a/Engine/Source/Runtime/UMG/Public/Components/ListView.h b/Engine/Source/Runtime/UMG/Public/Components/ListView.h
--- a/Engine/Source/Runtime/UMG/Public/Components/ListView.h
+++ b/Engine/Source/Runtime/UMG/Public/Components/ListView.h
@@ -20,7 +20,7 @@ DECLARE_DYNAMIC_MULTICAST_DELEGATE_TwoParams(FOnListItemScrolledIntoViewDynamic,
  * The list itself is based on a list of n items, but only creates as many entry widgets as can fit on screen.
  * For example, a scrolling ListView of 200 items with 5 currently visible will only have created 5 entry widgets.
  *
- * To make a widget usable an an entry in a ListView, it must inherit from the IUserObjectListEntry interface.
+ * To make a widget usable as an entry in a ListView, it must inherit from the IUserObjectListEntry interface.
  */
 UCLASS(meta = (EntryInterface = UserObjectListEntry))
 class UMG_API UListView : public UListViewBase, public ITypedUMGListView<UObject*>
