UMG: Add Child to Grid now takes in a row and column

diff --git a/Engine/Source/Runtime/UMG/Private/Components/GridPanel.cpp b/Engine/Source/Runtime/UMG/Private/Components/GridPanel.cpp
--- a/Engine/Source/Runtime/UMG/Private/Components/GridPanel.cpp
+++ b/Engine/Source/Runtime/UMG/Private/Components/GridPanel.cpp
@@ -67,9 +67,17 @@ TSharedRef<SWidget> UGridPanel::RebuildWidget()
 	return MyGridPanel.ToSharedRef();
 }
 
-UGridSlot* UGridPanel::AddChildToGrid(UWidget* Content)
+UGridSlot* UGridPanel::AddChildToGrid(UWidget* Content, int32 InRow, int32 InColumn)
 {
-	return Cast<UGridSlot>(Super::AddChild(Content));
+	UGridSlot* GridSlot = Cast<UGridSlot>(Super::AddChild(Content));
+
+	if (GridSlot != nullptr)
+	{
+		GridSlot->SetRow(InRow);
+		GridSlot->SetColumn(InColumn);
+	}
+
+	return GridSlot;
 }
 
 void UGridPanel::SynchronizeProperties()
diff --git a/Engine/Source/Runtime/UMG/Private/Components/UniformGridPanel.cpp b/Engine/Source/Runtime/UMG/Private/Components/UniformGridPanel.cpp
--- a/Engine/Source/Runtime/UMG/Private/Components/UniformGridPanel.cpp
+++ b/Engine/Source/Runtime/UMG/Private/Components/UniformGridPanel.cpp
@@ -67,9 +67,17 @@ TSharedRef<SWidget> UUniformGridPanel::RebuildWidget()
 	return MyUniformGridPanel.ToSharedRef();
 }
 
-UUniformGridSlot* UUniformGridPanel::AddChildToUniformGrid(UWidget* Content)
+UUniformGridSlot* UUniformGridPanel::AddChildToUniformGrid(UWidget* Content, int32 InRow, int32 InColumn)
 {
-	return Cast<UUniformGridSlot>(Super::AddChild(Content));
+	UUniformGridSlot* GridSlot = Cast<UUniformGridSlot>(Super::AddChild(Content));
+
+	if (GridSlot != nullptr)
+	{
+		GridSlot->SetRow(InRow);
+		GridSlot->SetColumn(InColumn);
+	}
+
+	return GridSlot;
 }
 
 void UUniformGridPanel::SetSlotPadding(FMargin InSlotPadding)
diff --git a/Engine/Source/Runtime/UMG/Public/Components/GridPanel.h b/Engine/Source/Runtime/UMG/Public/Components/GridPanel.h
--- a/Engine/Source/Runtime/UMG/Public/Components/GridPanel.h
+++ b/Engine/Source/Runtime/UMG/Public/Components/GridPanel.h
@@ -35,7 +35,7 @@ public:
 
 	/**  */
 	UFUNCTION(BlueprintCallable, Category="Widget")
-	UGridSlot* AddChildToGrid(UWidget* Content);
+	UGridSlot* AddChildToGrid(UWidget* Content, int32 InRow = 0, int32 InColumn = 0);
 
 	UFUNCTION(BlueprintCallable, Category = "Widget")
 	void SetColumnFill(int32 ColumnIndex, float Coefficient);
diff --git a/Engine/Source/Runtime/UMG/Public/Components/UniformGridPanel.h b/Engine/Source/Runtime/UMG/Public/Components/UniformGridPanel.h
--- a/Engine/Source/Runtime/UMG/Public/Components/UniformGridPanel.h
+++ b/Engine/Source/Runtime/UMG/Public/Components/UniformGridPanel.h
@@ -52,7 +52,7 @@ public:
 
 	/**  */
 	UFUNCTION(BlueprintCallable, Category="Widget")
-	UUniformGridSlot* AddChildToUniformGrid(UWidget* Content);
+	UUniformGridSlot* AddChildToUniformGrid(UWidget* Content, int32 InRow = 0, int32 InColumn = 0);
 
 public:
 
