PR #5917: Fix the description of Grid Panel (Contributed by [REDACTED_PERSON])

diff --git a/Engine/Source/Runtime/UMG/Public/Components/GridPanel.h b/Engine/Source/Runtime/UMG/Public/Components/GridPanel.h
--- a/Engine/Source/Runtime/UMG/Public/Components/GridPanel.h
+++ b/Engine/Source/Runtime/UMG/Public/Components/GridPanel.h
@@ -12,7 +12,7 @@ class SGridPanel;
 class UGridSlot;
 
 /**
- * A panel that evenly divides up available space between all of its children.
+ * A table-like panel that retains the width of every column throughout the table.
  * 
  * * Many Children
  */
