Updated widget component to set initial value for CurrentDrawSize in BeginPlay

diff --git a/Engine/Source/Runtime/UMG/Private/Components/WidgetComponent.cpp b/Engine/Source/Runtime/UMG/Private/Components/WidgetComponent.cpp
--- a/Engine/Source/Runtime/UMG/Private/Components/WidgetComponent.cpp
+++ b/Engine/Source/Runtime/UMG/Private/Components/WidgetComponent.cpp
@@ -639,6 +639,10 @@ void UWidgetComponent::BeginPlay()
 	Super::BeginPlay();
 
 	InitWidget();
+
+	CurrentDrawSize = DrawSize;
+	UpdateBodySetup(true);
+	RecreatePhysicsState();
 }
 
 void UWidgetComponent::EndPlay(const EEndPlayReason::Type EndPlayReason)
