After WidgetRenderer draws a widget, reassign the widget's original parent

diff --git a/Engine/Source/Runtime/UMG/Private/Slate/WidgetRenderer.cpp b/Engine/Source/Runtime/UMG/Private/Slate/WidgetRenderer.cpp
--- a/Engine/Source/Runtime/UMG/Private/Slate/WidgetRenderer.cpp
+++ b/Engine/Source/Runtime/UMG/Private/Slate/WidgetRenderer.cpp
@@ -117,10 +117,18 @@ void FWidgetRenderer::DrawWidget(
 	TSharedRef<SVirtualWindow> Window = SNew(SVirtualWindow).Size(DrawSize);
 	TSharedRef<FHittestGrid> HitTestGrid = MakeShareable(new FHittestGrid());
 
+	TSharedPtr<SWidget> OldParent = Widget->GetParentWidget();
 	Window->SetContent(Widget);
 	Window->Resize(DrawSize);
 
 	DrawWindow(RenderTarget, HitTestGrid, Window, Scale, DrawSize, DeltaTime, bDeferRenderTargetUpdate);
+
+	// Calling Window->SetContent will set Widget's parent to the Window, so we need to reset to the original parent.
+	// A better solution would be to have a way to render widgets without assigning them to a widget first.
+	if (OldParent.IsValid())
+	{
+		Widget->AssignParentWidget(OldParent);
+	}
 }
 
 void FWidgetRenderer::DrawWidget(
