PR #6025: FWidget3DSceneProxy - Fix assert firing in MaterialInstance::GetMaterial() (Contributed by [REDACTED_PERSON])

diff --git a/Engine/Source/Runtime/UMG/Private/Components/WidgetComponent.cpp b/Engine/Source/Runtime/UMG/Private/Components/WidgetComponent.cpp
--- a/Engine/Source/Runtime/UMG/Private/Components/WidgetComponent.cpp
+++ b/Engine/Source/Runtime/UMG/Private/Components/WidgetComponent.cpp
@@ -297,7 +297,7 @@ public:
 	{
 		bWillEverBeLit = false;
 
-		MaterialRelevance = MaterialInstance->GetRelevance(GetScene().GetFeatureLevel());
+		MaterialRelevance = MaterialInstance->GetRelevance_Concurrent(GetScene().GetFeatureLevel());
 	}
 
 	// FPrimitiveSceneProxy interface.
