PR #6016: Add interface requirement to the comment of ListViews (Contributed by [REDACTED_PERSON])

diff --git a/Engine/Source/Runtime/UMG/Public/Components/TileView.h b/Engine/Source/Runtime/UMG/Public/Components/TileView.h
--- a/Engine/Source/Runtime/UMG/Public/Components/TileView.h
+++ b/Engine/Source/Runtime/UMG/Public/Components/TileView.h
@@ -8,6 +8,8 @@
 
 /**
  * A ListView that presents the contents as a set of tiles all uniformly sized.
+ *
+ * To make a widget usable as an entry in a TileView, it must inherit from the IUserObjectListEntry interface.
  */
 UCLASS()
 class UMG_API UTileView : public UListView
diff --git a/Engine/Source/Runtime/UMG/Public/Components/TreeView.h b/Engine/Source/Runtime/UMG/Public/Components/TreeView.h
--- a/Engine/Source/Runtime/UMG/Public/Components/TreeView.h
+++ b/Engine/Source/Runtime/UMG/Public/Components/TreeView.h
@@ -13,6 +13,8 @@ DECLARE_DYNAMIC_MULTICAST_DELEGATE_TwoParams(FOnItemExpansionChangedDynamic, UOb
  * Similar to ListView, but can display a hierarchical tree of elements.
  * The base items source for the tree identifies the root items, each of which can have n associated child items.
  * There is no hard limit to the nesting - child items can have children and so on
+ *
+ * To make a widget usable as an entry in a TreeView, it must inherit from the IUserObjectListEntry interface.
  */
 UCLASS()
 class UMG_API UTreeView : public UListView
