UMG: Add option for global capture to allow for restore state of umg animation

diff --git a/Engine/Source/Runtime/UMG/Private/Animation/MovieSceneMarginSection.cpp b/Engine/Source/Runtime/UMG/Private/Animation/MovieSceneMarginSection.cpp
--- a/Engine/Source/Runtime/UMG/Private/Animation/MovieSceneMarginSection.cpp
+++ b/Engine/Source/Runtime/UMG/Private/Animation/MovieSceneMarginSection.cpp
@@ -125,6 +125,8 @@ struct FMarginSectionEditorData
 UMovieSceneMarginSection::UMovieSceneMarginSection( const FObjectInitializer& ObjectInitializer )
 	: Super( ObjectInitializer )
 {
+	EvalOptions.EnableAndSetCompletionMode(EMovieSceneCompletionMode::ProjectDefault);
+
 	BlendType = EMovieSceneBlendType::Absolute;
 
 	FMovieSceneChannelProxyData Channels;
diff --git a/Engine/Source/Runtime/UMG/Private/Animation/UMGSequencePlayer.cpp b/Engine/Source/Runtime/UMG/Private/Animation/UMGSequencePlayer.cpp
--- a/Engine/Source/Runtime/UMG/Private/Animation/UMGSequencePlayer.cpp
+++ b/Engine/Source/Runtime/UMG/Private/Animation/UMGSequencePlayer.cpp
@@ -13,6 +13,7 @@ UUMGSequencePlayer::UUMGSequencePlayer(const FObjectInitializer& ObjectInitializ
 	PlayerStatus = EMovieScenePlayerStatus::Stopped;
 	TimeCursorPosition = FFrameTime(0);
 	PlaybackSpeed = 1;
+	bRestoreState = false;
 	Animation = nullptr;
 	bIsEvaluating = false;
 	UserTag = NAME_None;
@@ -125,16 +126,28 @@ void UUMGSequencePlayer::Tick(float DeltaTime)
 		if ( bCompleted )
 		{
 			PlayerStatus = EMovieScenePlayerStatus::Stopped;
+			
+			if (bRestoreState)
+			{
+				RestorePreAnimatedState();
+			}
+
 			UserWidget->OnAnimationFinishedPlaying(*this);
 			OnSequenceFinishedPlayingEvent.Broadcast(*this);
 		}
 	}
 }
 
-void UUMGSequencePlayer::PlayInternal(double StartAtTime, double EndAtTime, int32 InNumLoopsToPlay, EUMGSequencePlayMode::Type InPlayMode, float InPlaybackSpeed)
+void UUMGSequencePlayer::PlayInternal(double StartAtTime, double EndAtTime, int32 InNumLoopsToPlay, EUMGSequencePlayMode::Type InPlayMode, float InPlaybackSpeed, bool bInRestoreState)
 {
+	if (bInRestoreState)
+	{
+		PreAnimatedState.EnableGlobalCapture();
+	}
+
 	RootTemplateInstance.Initialize(*Animation, *this);
 
+	bRestoreState = bInRestoreState;
 	PlaybackSpeed = FMath::Abs(InPlaybackSpeed);
 	PlayMode = InPlayMode;
 
@@ -180,14 +193,14 @@ void UUMGSequencePlayer::PlayInternal(double StartAtTime, double EndAtTime, int3
 	}
 }
 
-void UUMGSequencePlayer::Play(float StartAtTime, int32 InNumLoopsToPlay, EUMGSequencePlayMode::Type InPlayMode, float InPlaybackSpeed)
+void UUMGSequencePlayer::Play(float StartAtTime, int32 InNumLoopsToPlay, EUMGSequencePlayMode::Type InPlayMode, float InPlaybackSpeed, bool bInRestoreState)
 {
-	PlayInternal(StartAtTime, 0.0, InNumLoopsToPlay, InPlayMode, InPlaybackSpeed);
+	PlayInternal(StartAtTime, 0.0, InNumLoopsToPlay, InPlayMode, InPlaybackSpeed, bInRestoreState);
 }
 
-void UUMGSequencePlayer::PlayTo(float StartAtTime, float EndAtTime, int32 InNumLoopsToPlay, EUMGSequencePlayMode::Type InPlayMode, float InPlaybackSpeed)
+void UUMGSequencePlayer::PlayTo(float StartAtTime, float EndAtTime, int32 InNumLoopsToPlay, EUMGSequencePlayMode::Type InPlayMode, float InPlaybackSpeed, bool bInRestoreState)
 {
-	PlayInternal(StartAtTime, EndAtTime, InNumLoopsToPlay, InPlayMode, InPlaybackSpeed);
+	PlayInternal(StartAtTime, EndAtTime, InNumLoopsToPlay, InPlayMode, InPlaybackSpeed, bInRestoreState);
 }
 
 void UUMGSequencePlayer::Pause()
@@ -233,6 +246,11 @@ void UUMGSequencePlayer::Stop()
 		RootTemplateInstance.Finish(*this);
 	}
 
+	if (bRestoreState)
+	{
+		RestorePreAnimatedState();
+	}
+
 	UserWidget->OnAnimationFinishedPlaying(*this);
 	OnSequenceFinishedPlayingEvent.Broadcast(*this);
 
diff --git a/Engine/Source/Runtime/UMG/Private/UserWidget.cpp b/Engine/Source/Runtime/UMG/Private/UserWidget.cpp
--- a/Engine/Source/Runtime/UMG/Private/UserWidget.cpp
+++ b/Engine/Source/Runtime/UMG/Private/UserWidget.cpp
@@ -623,14 +623,14 @@ void UUserWidget::Invalidate(EInvalidateWidget InvalidateReason)
 	}
 }
 
-UUMGSequencePlayer* UUserWidget::PlayAnimation(UWidgetAnimation* InAnimation, float StartAtTime, int32 NumberOfLoops, EUMGSequencePlayMode::Type PlayMode, float PlaybackSpeed)
+UUMGSequencePlayer* UUserWidget::PlayAnimation(UWidgetAnimation* InAnimation, float StartAtTime, int32 NumberOfLoops, EUMGSequencePlayMode::Type PlayMode, float PlaybackSpeed, bool bRestoreState)
 {
 	SCOPED_NAMED_EVENT_TEXT("Widget::PlayAnimation", FColor::Emerald);
 
 	UUMGSequencePlayer* Player = GetOrAddSequencePlayer(InAnimation);
 	if (Player)
 	{
-		Player->Play(StartAtTime, NumberOfLoops, PlayMode, PlaybackSpeed);
+		Player->Play(StartAtTime, NumberOfLoops, PlayMode, PlaybackSpeed, bRestoreState);
 
 		Invalidate(EInvalidateWidget::Volatility);
 
@@ -642,14 +642,14 @@ UUMGSequencePlayer* UUserWidget::PlayAnimation(UWidgetAnimation* InAnimation, fl
 	return Player;
 }
 
-UUMGSequencePlayer* UUserWidget::PlayAnimationTimeRange(UWidgetAnimation* InAnimation, float StartAtTime, float EndAtTime, int32 NumberOfLoops, EUMGSequencePlayMode::Type PlayMode, float PlaybackSpeed)
+UUMGSequencePlayer* UUserWidget::PlayAnimationTimeRange(UWidgetAnimation* InAnimation, float StartAtTime, float EndAtTime, int32 NumberOfLoops, EUMGSequencePlayMode::Type PlayMode, float PlaybackSpeed, bool bRestoreState)
 {
 	SCOPED_NAMED_EVENT_TEXT("Widget::PlayAnimationTimeRange", FColor::Emerald);
 
 	UUMGSequencePlayer* Player = GetOrAddSequencePlayer(InAnimation);
 	if (Player)
 	{
-		Player->PlayTo(StartAtTime, EndAtTime, NumberOfLoops, PlayMode, PlaybackSpeed);
+		Player->PlayTo(StartAtTime, EndAtTime, NumberOfLoops, PlayMode, PlaybackSpeed, bRestoreState);
 
 		Invalidate(EInvalidateWidget::Volatility);
 
@@ -661,7 +661,7 @@ UUMGSequencePlayer* UUserWidget::PlayAnimationTimeRange(UWidgetAnimation* InAnim
 	return Player;
 }
 
-UUMGSequencePlayer* UUserWidget::PlayAnimationForward(UWidgetAnimation* InAnimation, float PlaybackSpeed)
+UUMGSequencePlayer* UUserWidget::PlayAnimationForward(UWidgetAnimation* InAnimation, float PlaybackSpeed, bool bRestoreState)
 {
 	// Don't create the player, only search for it.
 	UUMGSequencePlayer* Player = GetSequencePlayer(InAnimation);
@@ -676,10 +676,10 @@ UUMGSequencePlayer* UUserWidget::PlayAnimationForward(UWidgetAnimation* InAnimat
 		return Player;
 	}
 
-	return PlayAnimation(InAnimation, 0.0f, 1.0f, EUMGSequencePlayMode::Forward, PlaybackSpeed);
+	return PlayAnimation(InAnimation, 0.0f, 1.0f, EUMGSequencePlayMode::Forward, PlaybackSpeed, bRestoreState);
 }
 
-UUMGSequencePlayer* UUserWidget::PlayAnimationReverse(UWidgetAnimation* InAnimation, float PlaybackSpeed)
+UUMGSequencePlayer* UUserWidget::PlayAnimationReverse(UWidgetAnimation* InAnimation, float PlaybackSpeed, bool bRestoreState)
 {
 	// Don't create the player, only search for it.
 	UUMGSequencePlayer* Player = GetSequencePlayer(InAnimation);
@@ -694,7 +694,7 @@ UUMGSequencePlayer* UUserWidget::PlayAnimationReverse(UWidgetAnimation* InAnimat
 		return Player;
 	}
 
-	return PlayAnimation(InAnimation, 0.0f, 1.0f, EUMGSequencePlayMode::Reverse, PlaybackSpeed);
+	return PlayAnimation(InAnimation, 0.0f, 1.0f, EUMGSequencePlayMode::Reverse, PlaybackSpeed, bRestoreState);
 }
 
 void UUserWidget::StopAnimation(const UWidgetAnimation* InAnimation)
diff --git a/Engine/Source/Runtime/UMG/Public/Animation/UMGSequencePlayer.h b/Engine/Source/Runtime/UMG/Public/Animation/UMGSequencePlayer.h
--- a/Engine/Source/Runtime/UMG/Public/Animation/UMGSequencePlayer.h
+++ b/Engine/Source/Runtime/UMG/Public/Animation/UMGSequencePlayer.h
@@ -25,10 +25,10 @@ public:
 	void Tick( float DeltaTime );
 
 	/** Begins playing or restarts an animation */
-	void Play(float StartAtTime, int32 InNumLoopsToPlay, EUMGSequencePlayMode::Type InPlayMode, float InPlaybackSpeed);
+	void Play(float StartAtTime, int32 InNumLoopsToPlay, EUMGSequencePlayMode::Type InPlayMode, float InPlaybackSpeed, bool bRestoreState);
 
 	/** Begins playing or restarts an animation  and plays to the specified end time */
-	void PlayTo(float StartAtTime, float EndAtTime, int32 InNumLoopsToPlay, EUMGSequencePlayMode::Type InPlayMode, float InPlaybackSpeed);
+	void PlayTo(float StartAtTime, float EndAtTime, int32 InNumLoopsToPlay, EUMGSequencePlayMode::Type InPlayMode, float InPlaybackSpeed, bool bRestoreState);
 
 	/** Stops a running animation and resets time */
 	void Stop();
@@ -83,7 +83,7 @@ public:
 
 private:
 	/** Internal play function with a verbose parameter set */
-	void PlayInternal(double StartAtTime, double EndAtTime, int32 InNumLoopsToPlay, EUMGSequencePlayMode::Type InPlayMode, float InPlaybackSpeed);
+	void PlayInternal(double StartAtTime, double EndAtTime, int32 InNumLoopsToPlay, EUMGSequencePlayMode::Type InPlayMode, float InPlaybackSpeed, bool bRestoreState);
 
 	/** Apply any latent actions which may have accumulated while the sequence was being evaluated */
 	void ApplyLatentActions();
@@ -126,6 +126,9 @@ private:
 	/** The speed at which the animation should be played */
 	float PlaybackSpeed;
 
+	/** Whether to restore pre-animated state */
+	bool bRestoreState;
+
 	/** The current playback mode. */
 	EUMGSequencePlayMode::Type PlayMode;
 
diff --git a/Engine/Source/Runtime/UMG/Public/Blueprint/UserWidget.h b/Engine/Source/Runtime/UMG/Public/Blueprint/UserWidget.h
--- a/Engine/Source/Runtime/UMG/Public/Blueprint/UserWidget.h
+++ b/Engine/Source/Runtime/UMG/Public/Blueprint/UserWidget.h
@@ -900,9 +900,10 @@ public:
 	 * @param NumLoopsToPlay The number of times to loop this animation (0 to loop indefinitely)
 	 * @param PlaybackSpeed The speed at which the animation should play
 	 * @param PlayMode Specifies the playback mode
+	 * @param bRestoreState Restores widgets to their pre-animated state when the animation stops
 	 */
 	UFUNCTION(BlueprintCallable, BlueprintCosmetic, Category = "User Interface|Animation")
-	UUMGSequencePlayer* PlayAnimation(UWidgetAnimation* InAnimation, float StartAtTime = 0.0f, int32 NumLoopsToPlay = 1, EUMGSequencePlayMode::Type PlayMode = EUMGSequencePlayMode::Forward, float PlaybackSpeed = 1.0f);
+	UUMGSequencePlayer* PlayAnimation(UWidgetAnimation* InAnimation, float StartAtTime = 0.0f, int32 NumLoopsToPlay = 1, EUMGSequencePlayMode::Type PlayMode = EUMGSequencePlayMode::Forward, float PlaybackSpeed = 1.0f, bool bRestoreState = false);
 
 	/**
 	 * Plays an animation in this widget a specified number of times
@@ -928,9 +929,10 @@ public:
 	 * @param NumLoopsToPlay The number of times to loop this animation (0 to loop indefinitely)
 	 * @param PlayMode Specifies the playback mode
 	 * @param PlaybackSpeed The speed at which the animation should play
+	 * @param bRestoreState Restores widgets to their pre-animated state when the animation stops
 	 */
 	UFUNCTION(BlueprintCallable, BlueprintCosmetic, Category = "User Interface|Animation")
-	UUMGSequencePlayer* PlayAnimationTimeRange(UWidgetAnimation* InAnimation, float StartAtTime = 0.0f, float EndAtTime = 0.0f, int32 NumLoopsToPlay = 1, EUMGSequencePlayMode::Type PlayMode = EUMGSequencePlayMode::Forward, float PlaybackSpeed = 1.0f);
+	UUMGSequencePlayer* PlayAnimationTimeRange(UWidgetAnimation* InAnimation, float StartAtTime = 0.0f, float EndAtTime = 0.0f, int32 NumLoopsToPlay = 1, EUMGSequencePlayMode::Type PlayMode = EUMGSequencePlayMode::Forward, float PlaybackSpeed = 1.0f, bool bRestoreState = false);
 
 	/**
 	 * Plays an animation on this widget relative to it's current state forward.  You should use this version in situations where
@@ -940,9 +942,10 @@ public:
 	 * @param InAnimation The animation to play
 	 * @param PlayMode Specifies the playback mode
 	 * @param PlaybackSpeed The speed at which the animation should play
+	 * @param bRestoreState Restores widgets to their pre-animated state when the animation stops
 	 */
 	UFUNCTION(BlueprintCallable, BlueprintCosmetic, Category = "User Interface|Animation")
-	UUMGSequencePlayer* PlayAnimationForward(UWidgetAnimation* InAnimation, float PlaybackSpeed = 1.0f);
+	UUMGSequencePlayer* PlayAnimationForward(UWidgetAnimation* InAnimation, float PlaybackSpeed = 1.0f, bool bRestoreState = false);
 
 	/**
 	 * Plays an animation on this widget relative to it's current state in reverse.  You should use this version in situations where
@@ -952,9 +955,10 @@ public:
 	 * @param InAnimation The animation to play
 	 * @param PlayMode Specifies the playback mode
 	 * @param PlaybackSpeed The speed at which the animation should play
+	 * @param bRestoreState Restores widgets to their pre-animated state when the animation stops
 	 */
 	UFUNCTION(BlueprintCallable, BlueprintCosmetic, Category = "User Interface|Animation")
-	UUMGSequencePlayer* PlayAnimationReverse(UWidgetAnimation* InAnimation, float PlaybackSpeed = 1.0f);
+	UUMGSequencePlayer* PlayAnimationReverse(UWidgetAnimation* InAnimation, float PlaybackSpeed = 1.0f, bool bRestoreState = false);
 
 	/**
 	 * Stops an already running animation in this widget
