- Remove an ensure in editor that happen when we rename a NamedSlot, it doesn't fix the jira but remove an ensure in Editor, keeping it in engine
- Fix to GetSlotNames to support non Blueprint widgets

diff --git a/Engine/Source/Runtime/UMG/Private/UserWidget.cpp b/Engine/Source/Runtime/UMG/Private/UserWidget.cpp
--- a/Engine/Source/Runtime/UMG/Private/UserWidget.cpp
+++ b/Engine/Source/Runtime/UMG/Private/UserWidget.cpp
@@ -374,12 +374,16 @@ bool UUserWidget::Initialize()
 
 void UUserWidget::InitializeNamedSlots(bool bReparentToWidgetTree)
 {
-	for ( FNamedSlotBinding& Binding : NamedSlotBindings )
+	for (const FNamedSlotBinding& Binding : NamedSlotBindings )
 	{
 		if ( UWidget* BindingContent = Binding.Content )
 		{
 			UObjectPropertyBase* NamedSlotProperty = FindField<UObjectPropertyBase>(GetClass(), Binding.Name);
-			if ( ensure(NamedSlotProperty) )
+#if !WITH_EDITOR
+			// In editor, renaming a NamedSlot widget will cause this ensure in UpdatePreviewWidget of widget that use that namedslot
+			ensure(NamedSlotProperty);
+#endif
+			if ( NamedSlotProperty ) 
 			{
 				UNamedSlot* NamedSlot = Cast<UNamedSlot>(NamedSlotProperty->GetObjectPropertyValue_InContainer(this));
 				if ( ensure(NamedSlot) )
@@ -961,12 +965,14 @@ void UUserWidget::GetSlotNames(TArray<FName>& SlotNames) const
 	else // For non-blueprint widget blueprints we have to go through the widget tree to locate the named slots dynamically.
 	{
 		TArray<FName> NamedSlots;
-		WidgetTree->ForEachWidget([&] (UWidget* Widget) {
+		WidgetTree->ForEachWidget([&NamedSlots] (UWidget* Widget) {
 			if ( Widget && Widget->IsA<UNamedSlot>() )
 			{
 				NamedSlots.Add(Widget->GetFName());
 			}
 		});
+
+		SlotNames.Append(NamedSlots);
 	}
 }
 
diff --git a/Engine/Source/Runtime/UMG/Public/Blueprint/UserWidget.h b/Engine/Source/Runtime/UMG/Public/Blueprint/UserWidget.h
--- a/Engine/Source/Runtime/UMG/Public/Blueprint/UserWidget.h
+++ b/Engine/Source/Runtime/UMG/Public/Blueprint/UserWidget.h
@@ -139,10 +139,6 @@ public:
 	int32 MaxLayer;
 };
 
-/**
-* The state passed into OnPaint that we can expose as a single painting structure to blueprints to
-* allow script code to override OnPaint behavior.
-*/
 USTRUCT()
 struct UMG_API FNamedSlotBinding
 {
