Expose a SetFocus function on the widget interaction component for manually setting the focus of the underlying virtual user

diff --git a/Engine/Source/Runtime/UMG/Private/Components/WidgetInteractionComponent.cpp b/Engine/Source/Runtime/UMG/Private/Components/WidgetInteractionComponent.cpp
--- a/Engine/Source/Runtime/UMG/Private/Components/WidgetInteractionComponent.cpp
+++ b/Engine/Source/Runtime/UMG/Private/Components/WidgetInteractionComponent.cpp
@@ -95,6 +95,14 @@ void UWidgetInteractionComponent::SetCustomHitResult(const FHitResult& HitResult
 	CustomHitResult = HitResult;
 }
 
+void UWidgetInteractionComponent::SetFocus(UWidget* FocusWidget)
+{
+	if (VirtualUser.IsValid())
+	{
+		FSlateApplication::Get().SetUserFocus(VirtualUser->GetUserIndex(), FocusWidget->GetCachedWidget(), EFocusCause::SetDirectly);
+	}
+}
+
 FWidgetPath UWidgetInteractionComponent::FindHoveredWidgetPath(const FWidgetTraceResult& TraceResult) const
 {
 	if (TraceResult.HitWidgetComponent)
diff --git a/Engine/Source/Runtime/UMG/Public/Components/WidgetInteractionComponent.h b/Engine/Source/Runtime/UMG/Public/Components/WidgetInteractionComponent.h
--- a/Engine/Source/Runtime/UMG/Public/Components/WidgetInteractionComponent.h
+++ b/Engine/Source/Runtime/UMG/Public/Components/WidgetInteractionComponent.h
@@ -178,6 +178,12 @@ public:
 	UFUNCTION(BlueprintCallable, Category="Interaction")
 	void SetCustomHitResult(const FHitResult& HitResult);
 
+	/**
+	 * Set the focus target of the virtual user managed by this component
+	 */
+	UFUNCTION(BlueprintCallable, Category = "Interaction")
+	void SetFocus(UWidget* FocusWidget);
+
 private:
 	/**
 	 * Represents the virtual user in slate.  When this component is registered, it gets a handle to the 
