PR #6161: Change comment to 'remove' from 'add'. (Contributed by [REDACTED_PERSON])

diff --git a/Engine/Source/Runtime/UMG/Private/UserWidget.cpp b/Engine/Source/Runtime/UMG/Private/UserWidget.cpp
--- a/Engine/Source/Runtime/UMG/Private/UserWidget.cpp
+++ b/Engine/Source/Runtime/UMG/Private/UserWidget.cpp
@@ -1150,7 +1150,7 @@ void UUserWidget::RemoveFromParent()
 		{
 			TSharedPtr<SWidget> WidgetHost = FullScreenWidget.Pin();
 
-			// If this is a game world add the widget to the current worlds viewport.
+			// If this is a game world remove the widget from the current world's viewport.
 			UWorld* World = GetWorld();
 			if (World && World->IsGameWorld())
 			{
