Only assign the widget component as a child of the game viewport if we are not in some kind of editor world.

diff --git a/Engine/Source/Runtime/UMG/Private/Components/WidgetComponent.cpp b/Engine/Source/Runtime/UMG/Private/Components/WidgetComponent.cpp
--- a/Engine/Source/Runtime/UMG/Private/Components/WidgetComponent.cpp
+++ b/Engine/Source/Runtime/UMG/Private/Components/WidgetComponent.cpp
@@ -1019,13 +1019,16 @@ void UWidgetComponent::RegisterWindow()
 		{
 			if (UWorld* LocalWorld = GetWorld())
 			{
-				UGameInstance* GameInstance = LocalWorld->GetGameInstance();
-				check(GameInstance);
-
-				UGameViewportClient* GameViewportClient = GameInstance->GetGameViewportClient();
-				if (GameViewportClient)
+				if (LocalWorld->IsGameWorld())
 				{
-					SlateWindow->AssignParentWidget(GameViewportClient->GetGameViewportWidget());
+					UGameInstance* GameInstance = LocalWorld->GetGameInstance();
+					check(GameInstance);
+
+					UGameViewportClient* GameViewportClient = GameInstance->GetGameViewportClient();
+					if (GameViewportClient)
+					{
+						SlateWindow->AssignParentWidget(GameViewportClient->GetGameViewportWidget());
+					}
 				}
 			}
 		}
