PR #6223: Add UMG_API to SafeZoneSlot (Contributed by [REDACTED_PERSON])

diff --git a/Engine/Source/Runtime/UMG/Public/Components/SafeZoneSlot.h b/Engine/Source/Runtime/UMG/Public/Components/SafeZoneSlot.h
--- a/Engine/Source/Runtime/UMG/Public/Components/SafeZoneSlot.h
+++ b/Engine/Source/Runtime/UMG/Public/Components/SafeZoneSlot.h
@@ -10,7 +10,7 @@
 #include "SafeZoneSlot.generated.h"
 
 UCLASS()
-class USafeZoneSlot : public UPanelSlot
+class UMG_API USafeZoneSlot : public UPanelSlot
 {
 	GENERATED_BODY()
 public:
