PR #6228: Fixed presumed incorrect usage of member variable (Contributed by [REDACTED_PERSON])

diff --git a/Engine/Source/Runtime/UMG/Private/Components/WidgetComponent.cpp b/Engine/Source/Runtime/UMG/Private/Components/WidgetComponent.cpp
--- a/Engine/Source/Runtime/UMG/Private/Components/WidgetComponent.cpp
+++ b/Engine/Source/Runtime/UMG/Private/Components/WidgetComponent.cpp
@@ -1286,7 +1286,7 @@ void UWidgetComponent::ApplyComponentInstanceData(FWidgetComponentInstanceData*
 	// Note: ApplyComponentInstanceData is called while the component is registered so the rendering thread is already using this component
 	// That means all component state that is modified here must be mirrored on the scene proxy, which will be recreated to receive the changes later due to MarkRenderStateDirty.
 
-	if (GetWidgetClass() != WidgetClass)
+	if (GetWidgetClass() != WidgetInstanceData->WidgetClass)
 	{
 		return;
 	}
