PR #6363: Include name of widget in the log when trying to focus non-focusable widget (Contributed by SaltyPandaStudios)

diff --git a/Engine/Source/Runtime/UMG/Private/Components/Widget.cpp b/Engine/Source/Runtime/UMG/Private/Components/Widget.cpp
--- a/Engine/Source/Runtime/UMG/Private/Components/Widget.cpp
+++ b/Engine/Source/Runtime/UMG/Private/Components/Widget.cpp
@@ -454,7 +454,7 @@ void UWidget::SetKeyboardFocus()
 #if !(UE_BUILD_SHIPPING || UE_BUILD_TEST)
 		if ( !SafeWidget->SupportsKeyboardFocus() )
 		{
-			FMessageLog("PIE").Warning(LOCTEXT("ThisWidgetDoesntSupportFocus", "This widget does not support focus.  If this is a UserWidget, you should set bIsFocusable to true."));
+			FMessageLog("PIE").Warning(FText::Format(LOCTEXT("ThisWidgetDoesntSupportFocus", "The widget {0} does not support focus. If this is a UserWidget, you should set bIsFocusable to true."), FText::FromString(GetNameSafe(this))));
 		}
 #endif
 
@@ -564,7 +564,7 @@ void UWidget::SetUserFocus(APlayerController* PlayerController)
 #if !(UE_BUILD_SHIPPING || UE_BUILD_TEST)
 		if ( !SafeWidget->SupportsKeyboardFocus() )
 		{
-			FMessageLog("PIE").Warning(LOCTEXT("ThisWidgetDoesntSupportFocus", "This widget does not support focus.  If this is a UserWidget, you should set bIsFocusable to true."));
+			FMessageLog("PIE").Warning(FText::Format(LOCTEXT("ThisWidgetDoesntSupportFocus", "The widget {0} does not support focus. If this is a UserWidget, you should set bIsFocusable to true."), FText::FromString(GetNameSafe(this))));
 		}
 #endif
 
