Added option to disable ZOrder manipulations for world screen layer widgets: Slate.WorldWidgetZOrder, enabled by default.

diff --git a/Engine/Source/Runtime/UMG/Private/Slate/SWorldWidgetScreenLayer.cpp b/Engine/Source/Runtime/UMG/Private/Slate/SWorldWidgetScreenLayer.cpp
--- a/Engine/Source/Runtime/UMG/Private/Slate/SWorldWidgetScreenLayer.cpp
+++ b/Engine/Source/Runtime/UMG/Private/Slate/SWorldWidgetScreenLayer.cpp
@@ -9,6 +9,16 @@
 #include "Widgets/SViewport.h"
 #include "Slate/SGameLayerManager.h"
 
+static int32 GSlateWorldWidgetZOrder = 1;
+static FAutoConsoleVariableRef CVarSlateWorldWidgetZOrder(
+	TEXT("Slate.WorldWidgetZOrder"),
+	GSlateWorldWidgetZOrder,
+	TEXT("Whether to re-order world widgets projected to screen by their view point distance\n")
+	TEXT(" 0: Disable re-ordering\n")
+	TEXT(" 1: Re-order by distance (default, less batching, less artifacts when widgets overlap)"),
+	ECVF_Default
+	);
+
 SWorldWidgetScreenLayer::FComponentEntry::FComponentEntry()
 	: Slot(nullptr)
 {
@@ -141,7 +151,11 @@ void SWorldWidgetScreenLayer::Tick(const FGeometry& AllottedGeometry, const doub
 								CanvasSlot->Offset(FMargin(LocalPosition.X, LocalPosition.Y, ComponentDrawSize.X, ComponentDrawSize.Y));
 								CanvasSlot->Anchors(FAnchors(0, 0, 0, 0));
 								CanvasSlot->Alignment(ComponentPivot);
-								CanvasSlot->ZOrder(-ViewportPosition.Z);
+								
+								if (GSlateWorldWidgetZOrder != 0)
+								{
+									CanvasSlot->ZOrder(-ViewportPosition.Z);
+								}
 							}
 							else
 							{
@@ -149,7 +163,11 @@ void SWorldWidgetScreenLayer::Tick(const FGeometry& AllottedGeometry, const doub
 								CanvasSlot->Offset(FMargin(LocalPosition.X, LocalPosition.Y, DrawSize.X, DrawSize.Y));
 								CanvasSlot->Anchors(FAnchors(0, 0, 0, 0));
 								CanvasSlot->Alignment(Pivot);
-								CanvasSlot->ZOrder(-ViewportPosition.Z);
+
+								if (GSlateWorldWidgetZOrder != 0)
+								{
+									CanvasSlot->ZOrder(-ViewportPosition.Z);
+								}
 							}
 						}
 					}
