Add bPlayerViewportRelative to ProjectWorldLocationToWidgetPosition
#rnx

diff --git a/Engine/Source/Runtime/UMG/Private/WidgetLayoutLibrary.cpp b/Engine/Source/Runtime/UMG/Private/WidgetLayoutLibrary.cpp
--- a/Engine/Source/Runtime/UMG/Private/WidgetLayoutLibrary.cpp
+++ b/Engine/Source/Runtime/UMG/Private/WidgetLayoutLibrary.cpp
@@ -41,20 +41,19 @@ UWidgetLayoutLibrary::UWidgetLayoutLibrary(const FObjectInitializer& ObjectIniti
 {
 }
 
-bool UWidgetLayoutLibrary::ProjectWorldLocationToWidgetPosition(APlayerController* PlayerController, FVector WorldLocation, FVector2D& ViewportPosition)
+bool UWidgetLayoutLibrary::ProjectWorldLocationToWidgetPosition(APlayerController* PlayerController, FVector WorldLocation, FVector2D& ViewportPosition, bool bPlayerViewportRelative)
 {
 	FVector ScreenPosition3D;
-	const bool bSuccess = ProjectWorldLocationToWidgetPositionWithDistance(PlayerController, WorldLocation, ScreenPosition3D);
+	const bool bSuccess = ProjectWorldLocationToWidgetPositionWithDistance(PlayerController, WorldLocation, ScreenPosition3D, bPlayerViewportRelative);
 	ViewportPosition = FVector2D(ScreenPosition3D.X, ScreenPosition3D.Y);
 	return bSuccess;
 }
 
-bool UWidgetLayoutLibrary::ProjectWorldLocationToWidgetPositionWithDistance(APlayerController* PlayerController, FVector WorldLocation, FVector& ViewportPosition)
+bool UWidgetLayoutLibrary::ProjectWorldLocationToWidgetPositionWithDistance(APlayerController* PlayerController, FVector WorldLocation, FVector& ViewportPosition, bool bPlayerViewportRelative)
 {	
 	if ( PlayerController )
 	{
 		FVector PixelLocation;
-		const bool bPlayerViewportRelative = false;
 		const bool bProjected = PlayerController->ProjectWorldLocationToScreenWithDistance(WorldLocation, PixelLocation, bPlayerViewportRelative);
 
 		if ( bProjected )
diff --git a/Engine/Source/Runtime/UMG/Public/Blueprint/WidgetLayoutLibrary.h b/Engine/Source/Runtime/UMG/Public/Blueprint/WidgetLayoutLibrary.h
--- a/Engine/Source/Runtime/UMG/Public/Blueprint/WidgetLayoutLibrary.h
+++ b/Engine/Source/Runtime/UMG/Public/Blueprint/WidgetLayoutLibrary.h
@@ -36,10 +36,11 @@ public:
 	 * @param PlayerController The player controller to project the position in the world to their screen.
 	 * @param WorldLocation The world location to project from.
 	 * @param ScreenPosition The position in the viewport with quality scale removed and DPI scale remove.
+	 * @param bPlayerViewportRelative Should this be relative to the player viewport subregion (useful when using player attached widgets in split screen or when aspect-ratio constrained)
 	 * @return true if the position projects onto the screen.
 	 */
 	UFUNCTION(BlueprintPure, BlueprintCosmetic, Category="Viewport")
-	static bool ProjectWorldLocationToWidgetPosition(APlayerController* PlayerController, FVector WorldLocation, FVector2D& ScreenPosition);
+	static bool ProjectWorldLocationToWidgetPosition(APlayerController* PlayerController, FVector WorldLocation, FVector2D& ScreenPosition, bool bPlayerViewportRelative);
 
 	/**
 	 * Convert a World Space 3D position into a 2D Widget Screen Space position, with distance from the camera the Z component.  This
@@ -47,7 +48,7 @@ public:
 	 *
 	 * @return true if the world coordinate was successfully projected to the screen.
 	 */
-	static bool ProjectWorldLocationToWidgetPositionWithDistance(APlayerController* PlayerController, FVector WorldLocation, FVector& ScreenPosition);
+	static bool ProjectWorldLocationToWidgetPositionWithDistance(APlayerController* PlayerController, FVector WorldLocation, FVector& ScreenPosition, bool bPlayerViewportRelative);
 
 	/**
 	 * Gets the current DPI Scale being applied to the viewport and all the Widgets.
