Correction to a bug where a NamedSlot placed in the hierarchy of the Parent class of a WidgetBlueprint is not visible in instances of the widget.

This is only supported if the child WidgetTree is empty and only the parent contains a WidgetTree.

diff --git a/Engine/Source/Runtime/UMG/Private/UserWidget.cpp b/Engine/Source/Runtime/UMG/Private/UserWidget.cpp
--- a/Engine/Source/Runtime/UMG/Private/UserWidget.cpp
+++ b/Engine/Source/Runtime/UMG/Private/UserWidget.cpp
@@ -88,7 +88,7 @@ UUserWidget::UUserWidget(const FObjectInitializer& ObjectInitializer)
 	}
 }
 
-UWidgetBlueprintGeneratedClass* UUserWidget::GetWidgetTreeOwningClass()
+UWidgetBlueprintGeneratedClass* UUserWidget::GetWidgetTreeOwningClass() const
 {
 	UWidgetBlueprintGeneratedClass* WidgetClass = Cast<UWidgetBlueprintGeneratedClass>(GetClass());
 	if (WidgetClass != nullptr)
@@ -953,7 +953,7 @@ UWidget* UUserWidget::GetWidgetFromName(const FName& Name) const
 void UUserWidget::GetSlotNames(TArray<FName>& SlotNames) const
 {
 	// Only do this if this widget is of a blueprint class
-	if ( UWidgetBlueprintGeneratedClass* BGClass = Cast<UWidgetBlueprintGeneratedClass>(GetClass()) )
+	if (UWidgetBlueprintGeneratedClass* BGClass = GetWidgetTreeOwningClass())
 	{
 		SlotNames.Append(BGClass->NamedSlots);
 	}
@@ -1997,7 +1997,7 @@ void UUserWidget::PreSave(const class ITargetPlatform* TargetPlatform)
 	}
 
 	// Remove bindings that are no longer contained in the class.
-	if ( UWidgetBlueprintGeneratedClass* BGClass = Cast<UWidgetBlueprintGeneratedClass>(GetClass()) )
+	if ( UWidgetBlueprintGeneratedClass* BGClass = GetWidgetTreeOwningClass())
 	{
 		RemoveObsoleteBindings(BGClass->NamedSlots);
 	}
diff --git a/Engine/Source/Runtime/UMG/Public/Blueprint/UserWidget.h b/Engine/Source/Runtime/UMG/Public/Blueprint/UserWidget.h
--- a/Engine/Source/Runtime/UMG/Public/Blueprint/UserWidget.h
+++ b/Engine/Source/Runtime/UMG/Public/Blueprint/UserWidget.h
@@ -226,7 +226,7 @@ public:
 
 	EWidgetTickFrequency GetDesiredTickFrequency() const { return TickFrequency; }
 
-	UWidgetBlueprintGeneratedClass* GetWidgetTreeOwningClass();
+	UWidgetBlueprintGeneratedClass* GetWidgetTreeOwningClass() const;
 
 protected:
 	virtual void TemplateInitInner();
