UMG - WidgetComponents should skip creating widgets if we're tearing down the world.  CreateWidget now denotes the name of the world that's being torn down.

diff --git a/Engine/Source/Runtime/UMG/Private/Components/WidgetComponent.cpp b/Engine/Source/Runtime/UMG/Private/Components/WidgetComponent.cpp
--- a/Engine/Source/Runtime/UMG/Private/Components/WidgetComponent.cpp
+++ b/Engine/Source/Runtime/UMG/Private/Components/WidgetComponent.cpp
@@ -1415,13 +1415,15 @@ void UWidgetComponent::InitWidget()
 	// Don't do any work if Slate is not initialized
 	if ( FSlateApplication::IsInitialized() )
 	{
-		if ( WidgetClass && Widget == nullptr && GetWorld() )
+		UWorld* World = GetWorld();
+
+		if ( WidgetClass && Widget == nullptr && World && !World->bIsTearingDown)
 		{
 			Widget = CreateWidget(GetWorld(), WidgetClass);
 		}
 		
 #if WITH_EDITOR
-		if ( Widget && !GetWorld()->IsGameWorld() && !bEditTimeUsable )
+		if ( Widget && !World->IsGameWorld() && !bEditTimeUsable )
 		{
 			if( !GEnableVREditorHacks )
 			{
@@ -1997,16 +1999,19 @@ void UWidgetComponent::SetWidgetClass(TSubclassOf<UUserWidget> InWidgetClass)
 	{
 		WidgetClass = InWidgetClass;
 
-		if(HasBegunPlay())
+		if (FSlateApplication::IsInitialized())
 		{
-			if (WidgetClass)
+			if (HasBegunPlay() && !GetWorld()->bIsTearingDown)
 			{
-				UUserWidget* NewWidget = CreateWidget(GetWorld(), WidgetClass);
-				SetWidget(NewWidget);
-			}
-			else
-			{
-				SetWidget(nullptr);
+				if (WidgetClass)
+				{
+					UUserWidget* NewWidget = CreateWidget(GetWorld(), WidgetClass);
+					SetWidget(NewWidget);
+				}
+				else
+				{
+					SetWidget(nullptr);
+				}
 			}
 		}
 	}
diff --git a/Engine/Source/Runtime/UMG/Private/UserWidget.cpp b/Engine/Source/Runtime/UMG/Private/UserWidget.cpp
--- a/Engine/Source/Runtime/UMG/Private/UserWidget.cpp
+++ b/Engine/Source/Runtime/UMG/Private/UserWidget.cpp
@@ -2136,6 +2136,8 @@ UUserWidget* UUserWidget::CreateWidgetInstance(UWorld& World, TSubclassOf<UUserW
 
 UUserWidget* UUserWidget::CreateInstanceInternal(UObject* Outer, TSubclassOf<UUserWidget> UserWidgetClass, FName InstanceName, UWorld* World, ULocalPlayer* LocalPlayer)
 {
+	//CSV_SCOPED_TIMING_STAT(Slate, CreateWidget);
+
 #if !(UE_BUILD_SHIPPING || UE_BUILD_TEST)
 	// Only do this on a non-shipping or test build.
 	if (!CreateWidgetHelpers::ValidateUserWidgetClass(UserWidgetClass))
@@ -2166,7 +2168,7 @@ UUserWidget* UUserWidget::CreateInstanceInternal(UObject* Outer, TSubclassOf<UUs
 	if (World)
 	{
 		// Look for indications that widgets are being created for a dead and dying world.
-		ensureMsgf(!World->bIsTearingDown, TEXT("Widget Class %s - Attempting to be created while tearing down the world."), *UserWidgetClass->GetName());
+		ensureMsgf(!World->bIsTearingDown, TEXT("Widget Class %s - Attempting to be created while tearing down the world '%s'"), *UserWidgetClass->GetName(), *World->GetName());
 		}
 #endif
 
