UMG - During BP Creation, when choosing the outer, and owning player controller for a widget, don't do editor specific logic.  Instead prioritize the incoming PC specified by the user first, then fallback to the implied, if created from within a player controller, then see if it's a widget, before finally falling back to using the incoming world context as just a world context.

diff --git a/Engine/Source/Runtime/UMG/Private/WidgetBlueprintLibrary.cpp b/Engine/Source/Runtime/UMG/Private/WidgetBlueprintLibrary.cpp
--- a/Engine/Source/Runtime/UMG/Private/WidgetBlueprintLibrary.cpp
+++ b/Engine/Source/Runtime/UMG/Private/WidgetBlueprintLibrary.cpp
@@ -41,18 +41,18 @@ UUserWidget* UWidgetBlueprintLibrary::Create(UObject* WorldContextObject, TSubcl
 		return nullptr;
 	}
 
-	if(GIsEditor)
-	{
-		if (UUserWidget* OwningWidget = Cast<UUserWidget>(WorldContextObject))
-		{
-			return CreateWidget(OwningWidget, WidgetType);
-		}
-	}
-
 	if (OwningPlayer)
 	{
 		return CreateWidget(OwningPlayer, WidgetType);
 	}
+	else if (APlayerController* ImpliedOwningPlayer = Cast<APlayerController>(WorldContextObject))
+	{
+		return CreateWidget(ImpliedOwningPlayer, WidgetType);
+	}
+	else if (UUserWidget* OwningWidget = Cast<UUserWidget>(WorldContextObject))
+	{
+		return CreateWidget(OwningWidget, WidgetType);
+	}
 	else if (UWorld* World = GEngine->GetWorldFromContextObject(WorldContextObject, EGetWorldErrorMode::LogAndReturnNull))
 	{
 		return CreateWidget(World, WidgetType);
