Expose bReturnFocusToSelection to STreeView

diff --git a/Engine/Source/Runtime/Slate/Public/Widgets/Views/STreeView.h b/Engine/Source/Runtime/Slate/Public/Widgets/Views/STreeView.h
--- a/Engine/Source/Runtime/Slate/Public/Widgets/Views/STreeView.h
+++ b/Engine/Source/Runtime/Slate/Public/Widgets/Views/STreeView.h
@@ -137,6 +137,7 @@ public:
 		, _HandleDirectionalNavigation(true)
 		, _AllowInvisibleItemSelection(false)
 		, _HighlightParentNodesForSelection(false)
+		, _ReturnFocusToSelection()
 		{
 			this->_Clipping = EWidgetClipping::ClipToBounds;
 		}
@@ -202,6 +203,8 @@ public:
 
 		SLATE_ARGUMENT(bool, HighlightParentNodesForSelection);
 
+		SLATE_ARGUMENT(bool, ReturnFocusToSelection)
+
 	SLATE_END_ARGS()
 
 		
@@ -248,6 +251,8 @@ public:
 		this->bAllowInvisibleItemSelection = InArgs._AllowInvisibleItemSelection;
 		this->bHighlightParentNodesForSelection = InArgs._HighlightParentNodesForSelection;
 
+		this->bReturnFocusToSelection = InArgs._ReturnFocusToSelection;
+
 		// Check for any parameters that the coder forgot to specify.
 		FString ErrorString;
 		{
diff --git a/Engine/Source/Runtime/UMG/Public/Components/ListViewBase.h b/Engine/Source/Runtime/UMG/Public/Components/ListViewBase.h
--- a/Engine/Source/Runtime/UMG/Public/Components/ListViewBase.h
+++ b/Engine/Source/Runtime/UMG/Public/Components/ListViewBase.h
@@ -285,6 +285,7 @@ protected:
 		ESelectionMode::Type SelectionMode = ESelectionMode::Single;
 		bool bClearSelectionOnClick = false;
 		EConsumeMouseWheel ConsumeMouseWheel = EConsumeMouseWheel::WhenScrollingPossible;
+		bool bReturnFocusToSelection = false;
 	};
 
 	template <template<typename> class TreeViewT = STreeView, typename UListViewBaseT>
@@ -299,6 +300,7 @@ protected:
 			.ClearSelectionOnClick(Args.bClearSelectionOnClick)
 			.ConsumeMouseWheel(Args.ConsumeMouseWheel)
 			.SelectionMode(Args.SelectionMode)
+			.ReturnFocusToSelection(Args.bReturnFocusToSelection)
 			.OnGenerateRow_UObject(Implementer, &UListViewBaseT::HandleGenerateRow)
 			.OnSelectionChanged_UObject(Implementer, &UListViewBaseT::HandleSelectionChanged)
 			.OnIsSelectableOrNavigable_UObject(Implementer, &UListViewBaseT::HandleIsSelectableOrNavigable)
diff --git a/Engine/Source/Runtime/UMG/Public/Components/TreeView.h b/Engine/Source/Runtime/UMG/Public/Components/TreeView.h
--- a/Engine/Source/Runtime/UMG/Public/Components/TreeView.h
+++ b/Engine/Source/Runtime/UMG/Public/Components/TreeView.h
@@ -64,6 +64,7 @@ protected:
 		Args.bClearSelectionOnClick = bClearSelectionOnClick;
 		Args.SelectionMode = SelectionMode;
 		Args.ConsumeMouseWheel = ConsumeMouseWheel;
+		Args.bReturnFocusToSelection = bReturnFocusToSelection;
 
 		MyListView = MyTreeView = ITypedUMGListView<UObject*>::ConstructTreeView<TreeViewT>(this, ListItems, Args);
 		MyTreeView->SetOnEntryInitialized(SListView<UObject*>::FOnEntryInitialized::CreateUObject(this, &UTreeView::HandleOnEntryInitializedInternal));
