Exposing Virtual Keyboard Trigger on EditableTextBox

diff --git a/Engine/Source/Runtime/UMG/Private/Components/EditableTextBox.cpp b/Engine/Source/Runtime/UMG/Private/Components/EditableTextBox.cpp
--- a/Engine/Source/Runtime/UMG/Private/Components/EditableTextBox.cpp
+++ b/Engine/Source/Runtime/UMG/Private/Components/EditableTextBox.cpp
@@ -77,6 +77,7 @@ TSharedRef<SWidget> UEditableTextBox::RebuildWidget()
 		.OnTextCommitted(BIND_UOBJECT_DELEGATE(FOnTextCommitted, HandleOnTextCommitted))
 		.VirtualKeyboardType(EVirtualKeyboardType::AsKeyboardType(KeyboardType.GetValue()))
 		.VirtualKeyboardOptions(VirtualKeyboardOptions)
+		.VirtualKeyboardTrigger(VirtualKeyboardTrigger)
 		.VirtualKeyboardDismissAction(VirtualKeyboardDismissAction)
 		.Justification(Justification);
 
diff --git a/Engine/Source/Runtime/UMG/Public/Components/EditableTextBox.h b/Engine/Source/Runtime/UMG/Public/Components/EditableTextBox.h
--- a/Engine/Source/Runtime/UMG/Public/Components/EditableTextBox.h
+++ b/Engine/Source/Runtime/UMG/Public/Components/EditableTextBox.h
@@ -123,6 +123,10 @@ public:
 	UPROPERTY(EditAnywhere, Category=Behavior, AdvancedDisplay)
 	FVirtualKeyboardOptions VirtualKeyboardOptions;
 
+	/** The type of event that will trigger the display of the virtual keyboard */
+	UPROPERTY(EditAnywhere, Category = Behavior, AdvancedDisplay)
+	EVirtualKeyboardTrigger VirtualKeyboardTrigger;
+
 	/** What action should be taken when the virtual keyboard is dismissed? */
 	UPROPERTY(EditAnywhere, Category=Behavior, AdvancedDisplay)
 	EVirtualKeyboardDismissAction VirtualKeyboardDismissAction;
