ListViewBase::GetOwningUserIndex -- get a Slate User Index, not a ControllerId. By default, SlateApplication maps ControllerId to UserIndex 1:1, but clients can override this mapping.

diff --git a/Engine/Source/Runtime/UMG/Public/Components/ListViewBase.h b/Engine/Source/Runtime/UMG/Public/Components/ListViewBase.h
--- a/Engine/Source/Runtime/UMG/Public/Components/ListViewBase.h
+++ b/Engine/Source/Runtime/UMG/Public/Components/ListViewBase.h
@@ -8,6 +8,7 @@
 #include "Widgets/Layout/SSpacer.h"
 #include "Widgets/Views/STileView.h"
 #include "Widgets/Views/STreeView.h"
+#include "Framework/Application/SlateApplication.h"
 
 #include "ListViewBase.generated.h"
 
@@ -730,7 +731,8 @@ protected:	\
 	virtual uint32 GetOwningUserIndex() const override \
 	{	\
 		const ULocalPlayer* LocalPlayer = GetOwningLocalPlayer();	\
-		return LocalPlayer ? LocalPlayer->GetControllerId() : 0;	\
+		int32 SlateUserIndex = LocalPlayer ? FSlateApplication::Get().GetUserIndexForController(LocalPlayer->GetControllerId()) : 0;	\
+		return SlateUserIndex >= 0 ? SlateUserIndex : 0;	\
 	}	\
 	virtual bool IsDesignerPreview() const override { return IsDesignTime(); }	\
 private:	\
