Improve performance of ClearChildren by removing them in reverse

diff --git a/Engine/Source/Runtime/UMG/Private/Components/PanelWidget.cpp b/Engine/Source/Runtime/UMG/Private/Components/PanelWidget.cpp
--- a/Engine/Source/Runtime/UMG/Private/Components/PanelWidget.cpp
+++ b/Engine/Source/Runtime/UMG/Private/Components/PanelWidget.cpp
@@ -228,10 +228,9 @@ bool UPanelWidget::HasAnyChildren() const
 
 void UPanelWidget::ClearChildren()
 {
-	int32 Children = GetChildrenCount();
-	for ( int32 ChildIndex = 0; ChildIndex < Children; ChildIndex++ )
+	for ( int32 ChildIndex = GetChildrenCount() - 1; ChildIndex >= 0; ChildIndex-- )
 	{
-		RemoveChildAt(0);
+		RemoveChildAt(ChildIndex);
 	}
 }
 
