Avoid warning when removing widgets from parent that have no slate parent either

diff --git a/Engine/Source/Runtime/UMG/Private/Components/Widget.cpp b/Engine/Source/Runtime/UMG/Private/Components/Widget.cpp
--- a/Engine/Source/Runtime/UMG/Private/Components/Widget.cpp
+++ b/Engine/Source/Runtime/UMG/Private/Components/Widget.cpp
@@ -733,7 +733,7 @@ void UWidget::RemoveFromParent()
 		else
 		{
 #if !(UE_BUILD_SHIPPING || UE_BUILD_TEST)
-			if (GetCachedWidget().IsValid() && !IsDesignTime())
+			if (GetCachedWidget().IsValid() && GetCachedWidget()->GetParentWidget().IsValid() && !IsDesignTime())
 			{
 				FText WarningMessage = FText::Format(LOCTEXT("RemoveFromParentWithNoParent", "UWidget::RemoveFromParent() called on '{0}' which has no UMG parent (if it was added directly to a native Slate widget via TakeWidget() then it must be removed explicitly rather than via RemoveFromParent())"), FText::AsCultureInvariant(GetPathName()));
 				// @todo: nickd - we need to switch this back to a warning in engine, but info for games
