Support using negative size on Retainer Widget

diff --git a/Engine/Source/Runtime/UMG/Private/Slate/SRetainerWidget.cpp b/Engine/Source/Runtime/UMG/Private/Slate/SRetainerWidget.cpp
--- a/Engine/Source/Runtime/UMG/Private/Slate/SRetainerWidget.cpp
+++ b/Engine/Source/Runtime/UMG/Private/Slate/SRetainerWidget.cpp
@@ -399,8 +399,9 @@ bool SRetainerWidget::PaintRetainedContent(const FSlateInvalidationContext& Cont
 		LastTickedFrame = GFrameCounter;
 		const double TimeSinceLastDraw = FApp::GetCurrentTime() - LastDrawTime;
 
-		const uint32 RenderTargetWidth  = FMath::RoundToInt(RenderSize.X);
-		const uint32 RenderTargetHeight = FMath::RoundToInt(RenderSize.Y);
+		// Size must be a positive integer to allocate the RenderTarget
+		const uint32 RenderTargetWidth  = FMath::RoundToInt(FMath::Abs(RenderSize.X));
+		const uint32 RenderTargetHeight = FMath::RoundToInt(FMath::Abs(RenderSize.Y));
 
 		const FVector2D ViewOffset = PaintGeometry.DrawPosition.RoundToVector();
 
