[RN] UUserWidget's will now properly display the cursor set in their Cursor property by default, if set.

diff --git a/Engine/Source/Runtime/UMG/Private/UserWidget.cpp b/Engine/Source/Runtime/UMG/Private/UserWidget.cpp
--- a/Engine/Source/Runtime/UMG/Private/UserWidget.cpp
+++ b/Engine/Source/Runtime/UMG/Private/UserWidget.cpp
@@ -1914,7 +1914,9 @@ FReply UUserWidget::NativeOnTouchForceChanged(const FGeometry& InGeometry, const
 
 FCursorReply UUserWidget::NativeOnCursorQuery( const FGeometry& InGeometry, const FPointerEvent& InCursorEvent )
 {
-	return FCursorReply::Unhandled();
+	return (bOverride_Cursor)
+		? FCursorReply::Cursor(Cursor)
+		: FCursorReply::Unhandled();
 }
 
 FNavigationReply UUserWidget::NativeOnNavigation(const FGeometry& InGeometry, const FNavigationEvent& InNavigationEvent)
