Extend ScreenToWidgetLocal and ScreenToWidgetAbsolute to support windowed game mode

diff --git a/Engine/Source/Runtime/UMG/Private/SlateBlueprintLibrary.cpp b/Engine/Source/Runtime/UMG/Private/SlateBlueprintLibrary.cpp
--- a/Engine/Source/Runtime/UMG/Private/SlateBlueprintLibrary.cpp
+++ b/Engine/Source/Runtime/UMG/Private/SlateBlueprintLibrary.cpp
@@ -12,6 +12,7 @@
 #include "Slate/SceneViewport.h"
 #include "Blueprint/WidgetLayoutLibrary.h"
 #include "Slate/SGameLayerManager.h"
+#include "Widgets/SWindow.h"
 
 #define LOCTEXT_NAMESPACE "UMG"
 
@@ -92,15 +93,15 @@ void USlateBlueprintLibrary::AbsoluteToViewport(UObject* WorldContextObject, FVe
 	ViewportPosition = FVector2D(0, 0);
 }
 
-void USlateBlueprintLibrary::ScreenToWidgetLocal(UObject* WorldContextObject, const FGeometry& Geometry, FVector2D ScreenPosition, FVector2D& LocalCoordinate)
+void USlateBlueprintLibrary::ScreenToWidgetLocal(UObject* WorldContextObject, const FGeometry& Geometry, FVector2D ScreenPosition, FVector2D& LocalCoordinate, bool bIncludeWindowPosition /*= false*/)
 {
 	FVector2D AbsoluteCoordinate;
-	ScreenToWidgetAbsolute(WorldContextObject, ScreenPosition, AbsoluteCoordinate);
+	ScreenToWidgetAbsolute(WorldContextObject, ScreenPosition, AbsoluteCoordinate, bIncludeWindowPosition);
 
 	LocalCoordinate = Geometry.AbsoluteToLocal(AbsoluteCoordinate);
 }
 
-void USlateBlueprintLibrary::ScreenToWidgetAbsolute(UObject* WorldContextObject, FVector2D ScreenPosition, FVector2D& AbsoluteCoordinate)
+void USlateBlueprintLibrary::ScreenToWidgetAbsolute(UObject* WorldContextObject, FVector2D ScreenPosition, FVector2D& AbsoluteCoordinate, bool bIncludeWindowPosition /*= false*/)
 {
 	UWorld* World = GEngine->GetWorldFromContextObject(WorldContextObject, EGetWorldErrorMode::LogAndReturnNull);
 	if ( World && World->IsGameWorld() )
@@ -117,6 +118,13 @@ void USlateBlueprintLibrary::ScreenToWidgetAbsolute(UObject* WorldContextObject,
 				const FVector2D ViewportPosition = ViewportGeometry.GetLocalSize() * (ScreenPosition / ViewportSize);
 
 				AbsoluteCoordinate = ViewportGeometry.LocalToAbsolute(ViewportPosition);
+				if (bIncludeWindowPosition)
+				{
+					if (SWindow* Window = ViewportClient->GetWindow().Get())
+					{
+						AbsoluteCoordinate -= Window->GetPositionInScreen();
+					}
+				}
 				return;
 			}
 		}
diff --git a/Engine/Source/Runtime/UMG/Public/Blueprint/SlateBlueprintLibrary.h b/Engine/Source/Runtime/UMG/Public/Blueprint/SlateBlueprintLibrary.h
--- a/Engine/Source/Runtime/UMG/Public/Blueprint/SlateBlueprintLibrary.h
+++ b/Engine/Source/Runtime/UMG/Public/Blueprint/SlateBlueprintLibrary.h
@@ -101,15 +101,17 @@ public:
 
 	/**
 	 * Translates a screen position in pixels into the local space of a widget with the given geometry. 
+	 * If bIncludeWindowPosition is true, then this method will also remove the game window's position (useful when in windowed mode).
 	 */
 	UFUNCTION(Category="User Interface|Geometry", meta=( WorldContext="WorldContextObject", DisplayName="ScreenToLocal" ))
-	static void ScreenToWidgetLocal(UObject* WorldContextObject, const FGeometry& Geometry, FVector2D ScreenPosition, FVector2D& LocalCoordinate);
+	static void ScreenToWidgetLocal(UObject* WorldContextObject, const FGeometry& Geometry, FVector2D ScreenPosition, FVector2D& LocalCoordinate, bool bIncludeWindowPosition = false);
 
 	/**
 	 * Translates a screen position in pixels into absolute application coordinates.
+	 * If bIncludeWindowPosition is true, then this method will also remove the game window's position (useful when in windowed mode).
 	 */
 	UFUNCTION(Category="User Interface|Geometry", meta=( WorldContext="WorldContextObject", DisplayName="ScreenToAbsolute" ))
-	static void ScreenToWidgetAbsolute(UObject* WorldContextObject, FVector2D ScreenPosition, FVector2D& AbsoluteCoordinate);
+	static void ScreenToWidgetAbsolute(UObject* WorldContextObject, FVector2D ScreenPosition, FVector2D& AbsoluteCoordinate, bool bIncludeWindowPosition = false);
 
 	/**
 	 * Translates a screen position in pixels into the local space of the viewport widget.
