Logging to track text style NDC issue.

diff --git a/Engine/Source/Runtime/SlateCore/Private/Fonts/SlateFontInfo.cpp b/Engine/Source/Runtime/SlateCore/Private/Fonts/SlateFontInfo.cpp
--- a/Engine/Source/Runtime/SlateCore/Private/Fonts/SlateFontInfo.cpp
+++ b/Engine/Source/Runtime/SlateCore/Private/Fonts/SlateFontInfo.cpp
@@ -6,6 +6,7 @@
 #include "Fonts/FontProviderInterface.h"
 #include "Fonts/LegacySlateFontInfoCache.h"
 #include "UObject/FortniteMainBranchObjectVersion.h"
+#include "UObject/UObjectThreadContext.h"
 
 /* FSlateFontInfo structors
  *****************************************************************************/
@@ -200,6 +201,21 @@ void FSlateFontInfo::PostSerialize(const FArchive& Ar)
 	{
 		UpgradeLegacyFontInfo(FontName_DEPRECATED, Hinting_DEPRECATED);
 	}
+
+	// @fixme Temporary logging to track NDC issue.
+	if (Ar.IsCooking() && Ar.IsSaving())
+	{
+		FUObjectSerializeContext* SaveContext = const_cast<FArchive*>(&Ar)->GetSerializeContext();
+		if (SaveContext && SaveContext->SerializedObject)
+		{
+			UE_LOG(LogSlate, Display, TEXT("TEXT_STYLE_NDC: Serialized SlateFontInfo in object %s (0x%016x). Font=%s, FontMaterial=%s, Size=%d"),
+				*(SaveContext->SerializedObject->GetPathName()),
+				SaveContext->SerializedObject,
+				(FontObject ? *FontObject->GetFullName() : TEXT("None")),
+				(FontMaterial ? *FontMaterial->GetFullName() : TEXT("None")),
+				Size);
+		}
+	}
 }
 #endif
 
diff --git a/Engine/Source/Runtime/UMG/Private/Components/TextBlock.cpp b/Engine/Source/Runtime/UMG/Private/Components/TextBlock.cpp
--- a/Engine/Source/Runtime/UMG/Private/Components/TextBlock.cpp
+++ b/Engine/Source/Runtime/UMG/Private/Components/TextBlock.cpp
@@ -9,6 +9,7 @@
 #include "Materials/MaterialInstanceDynamic.h"
 #include "Widgets/SOverlay.h"
 #include "Widgets/Images/SImage.h"
+#include "UMGPrivate.h"
 
 #define LOCTEXT_NAMESPACE "UMG"
 
@@ -92,6 +93,11 @@ void UTextBlock::SetShadowOffset(FVector2D InShadowOffset)
 
 void UTextBlock::SetFont(FSlateFontInfo InFontInfo)
 {
+#if WITH_EDITOR
+	// @fixme Temporary logging to track NDC issue.
+	LogFontChange(InFontInfo);
+#endif //if WITH_EDITOR
+
 	Font = InFontInfo;
 	if (MyTextBlock.IsValid())
 	{
@@ -361,7 +367,62 @@ bool UTextBlock::CanEditChange(const FProperty* InProperty) const
 	return Super::CanEditChange(InProperty);
 }
 
-#endif
+void UTextBlock::LogFontChange(const FSlateFontInfo& NewFont) const
+{
+	bool bFontChanged = false;
+	if (Font.FontObject != NewFont.FontObject)
+	{
+		bFontChanged = true;
+		UE_LOG(LogUMG, Display, TEXT("TEXT_STYLE_NDC: Changing font on TextBlock '%s' (0x%016x) in asset %s. From %s, to %s"),
+			*GetName(),
+			this,
+			*GetOutermost()->GetFName().ToString(),
+			(Font.FontObject ? *Font.FontObject->GetFullName() : TEXT("None")),
+			(NewFont.FontObject ? *NewFont.FontObject->GetFullName() : TEXT("None")));
+	}
+
+	if (Font.FontMaterial != NewFont.FontMaterial)
+	{
+		bFontChanged = true;
+		UE_LOG(LogUMG, Display, TEXT("TEXT_STYLE_NDC: Changing font material on TextBlock '%s' (0x%016x) in asset %s. From %s, to %s"),
+			*GetName(),
+			this,
+			*GetOutermost()->GetFName().ToString(),
+			(Font.FontMaterial ? *Font.FontMaterial->GetFullName() : TEXT("None")),
+			(NewFont.FontMaterial ? *NewFont.FontMaterial->GetFullName() : TEXT("None")));
+	}
+
+	if (Font.Size != NewFont.Size)
+	{
+		bFontChanged = true;
+		UE_LOG(LogUMG, Display, TEXT("TEXT_STYLE_NDC: Changing font size on TextBlock '%s' (0x%016x) in asset %s. From %d, to %d"),
+			*GetName(),
+			this,
+			*GetOutermost()->GetFName().ToString(),
+			Font.Size,
+			NewFont.Size);
+	}
+
+	if (bFontChanged)
+	{
+		LogCallstack(2);
+	}
+}
+
+void UTextBlock::LogCallstack(int32 IgnoreCount) const
+{
+	UE_LOG(LogUMG, Display, TEXT("TEXT_STYLE_NDC: Script callstack:"));
+	PrintScriptCallstack();
+
+	UE_LOG(LogUMG, Display, TEXT("TEXT_STYLE_NDC: Code callstack:"));
+	for (const FProgramCounterSymbolInfo& StackLine : FPlatformStackWalk::GetStack(IgnoreCount, 10))
+	{
+		UE_LOG(LogUMG, Display, TEXT("TEXT_STYLE_NDC:   %s"), ANSI_TO_TCHAR(StackLine.FunctionName));
+	}
+	UE_LOG(LogUMG, Display, TEXT("TEXT_STYLE_NDC:   ..."));
+}
+
+#endif //if WITH_EDITOR
 
 /////////////////////////////////////////////////////
 
diff --git a/Engine/Source/Runtime/UMG/Public/Components/TextBlock.h b/Engine/Source/Runtime/UMG/Public/Components/TextBlock.h
--- a/Engine/Source/Runtime/UMG/Public/Components/TextBlock.h
+++ b/Engine/Source/Runtime/UMG/Public/Components/TextBlock.h
@@ -212,6 +212,11 @@ protected:
 	virtual TSharedPtr<SWidget> GetAccessibleWidget() const override;
 #endif
 
+#if WITH_EDITOR
+	void LogFontChange(const FSlateFontInfo& NewFont) const;
+	void LogCallstack(int32 IgnoreCount=1) const;
+#endif //if WITH_EDITOR
+
 protected:
 
 	TSharedPtr<STextBlock> MyTextBlock;
