Fix non-unitty build

diff --git a/Engine/Source/Runtime/UMG/Private/Components/TextBlock.cpp b/Engine/Source/Runtime/UMG/Private/Components/TextBlock.cpp
--- a/Engine/Source/Runtime/UMG/Private/Components/TextBlock.cpp
+++ b/Engine/Source/Runtime/UMG/Private/Components/TextBlock.cpp
@@ -9,7 +9,11 @@
 #include "Materials/MaterialInstanceDynamic.h"
 #include "Widgets/SOverlay.h"
 #include "Widgets/Images/SImage.h"
+
+#if WITH_EDITOR
+#include "HAL/PlatformStackWalk.h"
 #include "UMGPrivate.h"
+#endif //if WITH_EDITOR
 
 #define LOCTEXT_NAMESPACE "UMG"
 
