Remove ensure to GetDynamicFontMaterial since we support it to be null.

diff --git a/Engine/Source/Runtime/UMG/Private/Components/TextBlock.cpp b/Engine/Source/Runtime/UMG/Private/Components/TextBlock.cpp
--- a/Engine/Source/Runtime/UMG/Private/Components/TextBlock.cpp
+++ b/Engine/Source/Runtime/UMG/Private/Components/TextBlock.cpp
@@ -137,7 +137,7 @@ void UTextBlock::SetAutoWrapText(bool InAutoWrapText)
 
 UMaterialInstanceDynamic* UTextBlock::GetDynamicFontMaterial()
 {
-	if (ensure(Font.FontMaterial))
+	if (Font.FontMaterial)
 	{
 		UMaterialInterface* Material = CastChecked<UMaterialInterface>(Font.FontMaterial);
 
@@ -159,7 +159,7 @@ UMaterialInstanceDynamic* UTextBlock::GetDynamicFontMaterial()
 
 UMaterialInstanceDynamic* UTextBlock::GetDynamicOutlineMaterial()
 {
-	if (ensure(Font.OutlineSettings.OutlineMaterial))
+	if (Font.OutlineSettings.OutlineMaterial)
 	{
 		UMaterialInterface* Material = CastChecked<UMaterialInterface>(Font.OutlineSettings.OutlineMaterial);
 
