Fix to WidgetComponent to correctly support Component using only SWidget. By resetiing the bRenderCleared flag et force a rerender when the SWisget goes from null to a valid value.

diff --git a/Engine/Source/Runtime/UMG/Private/Components/WidgetComponent.cpp b/Engine/Source/Runtime/UMG/Private/Components/WidgetComponent.cpp
--- a/Engine/Source/Runtime/UMG/Private/Components/WidgetComponent.cpp
+++ b/Engine/Source/Runtime/UMG/Private/Components/WidgetComponent.cpp
@@ -1537,6 +1537,7 @@ void UWidgetComponent::UpdateWidget()
 				{
 					CurrentSlateWidget = NewSlateWidget;
 					SlateWindow->SetContent(NewSlateWidget.ToSharedRef());
+					bRenderCleared = false;
 					bWidgetChanged = true;
 				}
 			}
@@ -1546,6 +1547,7 @@ void UWidgetComponent::UpdateWidget()
 				{
 					CurrentSlateWidget = SlateWidget;
 					SlateWindow->SetContent(SlateWidget.ToSharedRef());
+					bRenderCleared = false;
 					bWidgetChanged = true;
 				}
 			}
