Added functions to UserWidget for retrieving owning player's camera manager.

#rnx

diff --git a/Engine/Source/Runtime/UMG/Private/UserWidget.cpp b/Engine/Source/Runtime/UMG/Private/UserWidget.cpp
--- a/Engine/Source/Runtime/UMG/Private/UserWidget.cpp
+++ b/Engine/Source/Runtime/UMG/Private/UserWidget.cpp
@@ -1231,9 +1231,9 @@ void UUserWidget::SetOwningPlayer(APlayerController* LocalPlayerController)
 	}
 }
 
-class APawn* UUserWidget::GetOwningPlayerPawn() const
+APawn* UUserWidget::GetOwningPlayerPawn() const
 {
-	if ( APlayerController* PC = GetOwningPlayer() )
+	if (APlayerController* PC = GetOwningPlayer())
 	{
 		return PC->GetPawn();
 	}
@@ -1241,6 +1241,16 @@ class APawn* UUserWidget::GetOwningPlayerPawn() const
 	return nullptr;
 }
 
+APlayerCameraManager* UUserWidget::GetOwningPlayerCameraManager() const
+{
+	if (APlayerController* PC = GetOwningPlayer())
+	{
+		return PC->PlayerCameraManager;
+	}
+
+	return nullptr;
+}
+
 void UUserWidget::SetPositionInViewport(FVector2D Position, bool bRemoveDPIScale )
 {
 	if ( bRemoveDPIScale )
diff --git a/Engine/Source/Runtime/UMG/Public/Blueprint/UserWidget.h b/Engine/Source/Runtime/UMG/Public/Blueprint/UserWidget.h
--- a/Engine/Source/Runtime/UMG/Public/Blueprint/UserWidget.h
+++ b/Engine/Source/Runtime/UMG/Public/Blueprint/UserWidget.h
@@ -407,6 +407,24 @@ public:
 		return nullptr;
 	}
 
+	/**
+	 * Gets the player camera manager associated with this UI.
+	 * @return Gets the owning player camera manager that's owned by the player controller assigned to this widget.
+	 */
+	UFUNCTION(BlueprintCallable, BlueprintCosmetic, Category = "Player")
+	class APlayerCameraManager* GetOwningPlayerCameraManager() const;
+
+	/**
+	 * Gets the player camera manager associated with this UI cast to the template type.
+	 * @return Gets the owning player camera manager that's owned by the player controller assigned to this widget.
+	 * May be NULL if the cast fails.
+	 */
+	template <class T>
+	T* GetOwningPlayerCameraManager() const
+	{
+		return Cast<T>(GetOwningPlayerCameraManager());
+	}
+
 	/** 
 	 * Called once only at game time on non-template instances.
 	 * While Construct/Destruct pertain to the underlying Slate, this is called only once for the UUserWidget.
