Prevent the construction of Material when we destroy the WidgetComponent. Quick fix for the release, working on a more robust solution.

diff --git a/Engine/Source/Runtime/UMG/Private/Components/WidgetComponent.cpp b/Engine/Source/Runtime/UMG/Private/Components/WidgetComponent.cpp
--- a/Engine/Source/Runtime/UMG/Private/Components/WidgetComponent.cpp
+++ b/Engine/Source/Runtime/UMG/Private/Components/WidgetComponent.cpp
@@ -700,7 +700,7 @@ void UWidgetComponent::UpdateMaterialInstance()
 	MaterialInstance = UMaterialInstanceDynamic::Create(BaseMaterial, this);
 	if (MaterialInstance)
 	{
-			MaterialInstance->AddToCluster(this);
+		MaterialInstance->AddToCluster(this);
 	}
 	UpdateMaterialInstanceParameters();
 }
@@ -1505,7 +1505,7 @@ void UWidgetComponent::SetSlateWidget(const TSharedPtr<SWidget>& InSlateWidget)
 void UWidgetComponent::UpdateWidget()
 {
 	// Don't do any work if Slate is not initialized
-	if ( FSlateApplication::IsInitialized() )
+	if (FSlateApplication::IsInitialized() && !IsPendingKill())
 	{
 		if ( Space != EWidgetSpace::Screen )
 		{
@@ -1541,7 +1541,7 @@ void UWidgetComponent::UpdateWidget()
 					bWidgetChanged = true;
 				}
 			}
-			else if( SlateWidget.IsValid() )
+			else if ( SlateWidget.IsValid() )
 			{
 				if ( SlateWidget != CurrentSlateWidget || bNeededNewWindow )
 				{
