Fixed flashing widget.
- CurrentDrawSize was set to 0 on registration, calculating a unit box for bounds, breaking hardware occlusion, and causing false occlusion on the widget (flicker).

diff --git a/Engine/Source/Runtime/UMG/Private/Components/WidgetComponent.cpp b/Engine/Source/Runtime/UMG/Private/Components/WidgetComponent.cpp
--- a/Engine/Source/Runtime/UMG/Private/Components/WidgetComponent.cpp
+++ b/Engine/Source/Runtime/UMG/Private/Components/WidgetComponent.cpp
@@ -837,6 +837,9 @@ FCollisionShape UWidgetComponent::GetCollisionShape(float Inflation) const
 
 void UWidgetComponent::OnRegister()
 {
+	// Set this prior to registering the scene component so that bounds are calculated correctly.
+	CurrentDrawSize = DrawSize;
+
 	Super::OnRegister();
 
 #if !UE_SERVER
