PR #7001: Make PressKey return true if had any reply (Contributed by [REDACTED_PERSON])

diff --git a/Engine/Source/Runtime/UMG/Private/Components/WidgetInteractionComponent.cpp b/Engine/Source/Runtime/UMG/Private/Components/WidgetInteractionComponent.cpp
--- a/Engine/Source/Runtime/UMG/Private/Components/WidgetInteractionComponent.cpp
+++ b/Engine/Source/Runtime/UMG/Private/Components/WidgetInteractionComponent.cpp
@@ -501,15 +501,16 @@ bool UWidgetInteractionComponent::PressKey(FKey Key, bool bRepeat)
 	GetKeyAndCharCodes(Key, bHasKeyCode, KeyCode, bHasCharCode, CharCode);
 
 	FKeyEvent KeyEvent(Key, ModifierKeys, VirtualUser->GetUserIndex(), bRepeat, KeyCode, CharCode);
-	bool DownResult = FSlateApplication::Get().ProcessKeyDownEvent(KeyEvent);
+	bool bDownResult = FSlateApplication::Get().ProcessKeyDownEvent(KeyEvent);
 
+	bool bKeyCharResult = false;
 	if (bHasCharCode)
 	{
 		FCharacterEvent CharacterEvent(CharCode, ModifierKeys, VirtualUser->GetUserIndex(), bRepeat);
-		return FSlateApplication::Get().ProcessKeyCharEvent(CharacterEvent);
+		bKeyCharResult = FSlateApplication::Get().ProcessKeyCharEvent(CharacterEvent);
 	}
 
-	return DownResult;
+	return bDownResult || bKeyCharResult;
 }
 
 bool UWidgetInteractionComponent::ReleaseKey(FKey Key)
