Add "screen size" as keywords to GetViewportSize

diff --git a/Engine/Source/Runtime/UMG/Public/Blueprint/WidgetLayoutLibrary.h b/Engine/Source/Runtime/UMG/Public/Blueprint/WidgetLayoutLibrary.h
--- a/Engine/Source/Runtime/UMG/Public/Blueprint/WidgetLayoutLibrary.h
+++ b/Engine/Source/Runtime/UMG/Public/Blueprint/WidgetLayoutLibrary.h
@@ -64,7 +64,7 @@ public:
 	/**
 	 * Gets the size of the game viewport.
 	 */
-	UFUNCTION(BlueprintPure, BlueprintCosmetic, Category="Viewport", meta=( WorldContext="WorldContextObject" ))
+	UFUNCTION(BlueprintPure, BlueprintCosmetic, Category="Viewport", meta=( WorldContext="WorldContextObject", Keywords = "screen size"))
 	static FVector2D GetViewportSize(UObject* WorldContextObject);
 
 	/**
