UMG - Determinism during saving issues, removing the code that cleared out the WidgetTree on a generated class during save if we're not saving.

diff --git a/Engine/Source/Runtime/UMG/Private/WidgetBlueprintGeneratedClass.cpp b/Engine/Source/Runtime/UMG/Private/WidgetBlueprintGeneratedClass.cpp
--- a/Engine/Source/Runtime/UMG/Private/WidgetBlueprintGeneratedClass.cpp
+++ b/Engine/Source/Runtime/UMG/Private/WidgetBlueprintGeneratedClass.cpp
@@ -397,25 +397,6 @@ void UWidgetBlueprintGeneratedClass::SetWidgetTreeArchetype(UWidgetTree* InWidge
 	}
 }
 
-void UWidgetBlueprintGeneratedClass::PreSave(const class ITargetPlatform* TargetPlatform)
-{
-#if WITH_EDITOR
-	if (TargetPlatform == nullptr)
-	{
-		// If we're saving the generated class in the editor, should we allow it to preserve a shadow copy of the one in the
-		// blueprint?  Seems dangerous to have this potentially stale copy around, when really it should be the latest version
-		// that's compiled on load.
-		if (WidgetTree)
-		{
-			WidgetTree->SetFlags(RF_Transient);
-			WidgetTree = nullptr;
-		}
-	}
-#endif
-
-	Super::PreSave(TargetPlatform);
-}
-
 void UWidgetBlueprintGeneratedClass::Serialize(FArchive& Ar)
 {
 	Super::Serialize(Ar);
diff --git a/Engine/Source/Runtime/UMG/Public/Blueprint/WidgetBlueprintGeneratedClass.h b/Engine/Source/Runtime/UMG/Public/Blueprint/WidgetBlueprintGeneratedClass.h
--- a/Engine/Source/Runtime/UMG/Public/Blueprint/WidgetBlueprintGeneratedClass.h
+++ b/Engine/Source/Runtime/UMG/Public/Blueprint/WidgetBlueprintGeneratedClass.h
@@ -98,7 +98,6 @@ public:
 	UWidgetBlueprintGeneratedClass* FindWidgetTreeOwningClass();
 
 	// UObject interface
-	virtual void PreSave(const class ITargetPlatform* TargetPlatform) override;
 	virtual void Serialize(FArchive& Ar) override;
 
 	virtual void PostLoad() override;
