PR #6714: Slate changes (Contributed by [REDACTED_PERSON])

Adding SetForeGrounColor
Change layering sorting to a stable sort to keep order as much as we can.

diff --git a/Engine/Source/Runtime/Slate/Private/Widgets/Layout/SConstraintCanvas.cpp b/Engine/Source/Runtime/Slate/Private/Widgets/Layout/SConstraintCanvas.cpp
--- a/Engine/Source/Runtime/Slate/Private/Widgets/Layout/SConstraintCanvas.cpp
+++ b/Engine/Source/Runtime/Slate/Private/Widgets/Layout/SConstraintCanvas.cpp
@@ -99,7 +99,7 @@ void SConstraintCanvas::ArrangeLayeredChildren(const FGeometry& AllottedGeometry
 			SlotOrder.Add(Order);
 		}
 
-		SlotOrder.Sort(FSortSlotsByZOrder());
+		SlotOrder.StableSort(FSortSlotsByZOrder());
 		float LastZOrder = -FLT_MAX;
 
 		// Arrange the children now in their proper z-order.
diff --git a/Engine/Source/Runtime/UMG/Private/Components/EditableTextBox.cpp b/Engine/Source/Runtime/UMG/Private/Components/EditableTextBox.cpp
--- a/Engine/Source/Runtime/UMG/Private/Components/EditableTextBox.cpp
+++ b/Engine/Source/Runtime/UMG/Private/Components/EditableTextBox.cpp
@@ -137,6 +137,15 @@ void UEditableTextBox::SetHintText(FText InText)
 	}
 }
 
+void UEditableTextBox::SetForegroundColor(FLinearColor color)
+{
+	WidgetStyle.ForegroundColor = color;
+	if ( MyEditableTextBlock.IsValid() )
+	{
+		MyEditableTextBlock->SetForegroundColor(color);
+	}
+}
+
 void UEditableTextBox::SetError(FText InError)
 {
 	if ( MyEditableTextBlock.IsValid() )
diff --git a/Engine/Source/Runtime/UMG/Private/Components/MultiLineEditableTextBox.cpp b/Engine/Source/Runtime/UMG/Private/Components/MultiLineEditableTextBox.cpp
--- a/Engine/Source/Runtime/UMG/Private/Components/MultiLineEditableTextBox.cpp
+++ b/Engine/Source/Runtime/UMG/Private/Components/MultiLineEditableTextBox.cpp
@@ -188,6 +188,14 @@ void UMultiLineEditableTextBox::SetTextStyle(const FTextBlockStyle& InTextStyle)
 	}
 }
 
+void UMultiLineEditableTextBox::SetForegroundColor(FLinearColor color)
+{
+	if(MyEditableTextBlock.IsValid())
+	{
+		MyEditableTextBlock->SetForegroundColor(color);
+	}
+}
+
 void UMultiLineEditableTextBox::HandleOnTextChanged(const FText& InText)
 {
 	OnTextChanged.Broadcast(InText);
diff --git a/Engine/Source/Runtime/UMG/Public/Components/EditableTextBox.h b/Engine/Source/Runtime/UMG/Public/Components/EditableTextBox.h
--- a/Engine/Source/Runtime/UMG/Public/Components/EditableTextBox.h
+++ b/Engine/Source/Runtime/UMG/Public/Components/EditableTextBox.h
@@ -183,6 +183,9 @@ public:
 	UFUNCTION(BlueprintSetter)
 	void SetJustification(ETextJustify::Type InJustification);
 
+	UFUNCTION(BlueprintCallable, Category="Widget", meta=(DisplayName="SetForegroundColor (Text Box)"))
+	void SetForegroundColor(FLinearColor color);
+
 	//~ Begin UWidget Interface
 	virtual void SynchronizeProperties() override;
 	//~ End UWidget Interface
diff --git a/Engine/Source/Runtime/UMG/Public/Components/MultiLineEditableTextBox.h b/Engine/Source/Runtime/UMG/Public/Components/MultiLineEditableTextBox.h
--- a/Engine/Source/Runtime/UMG/Public/Components/MultiLineEditableTextBox.h
+++ b/Engine/Source/Runtime/UMG/Public/Components/MultiLineEditableTextBox.h
@@ -123,9 +123,11 @@ public:
 	UFUNCTION(BlueprintSetter)
 	void SetTextStyle(const FTextBlockStyle& InTextStyle);
 
+	UFUNCTION(BlueprintCallable, Category="Widget", meta=(DisplayName="SetForegroundColor (Multi-Line Text Box)"))
+	void SetForegroundColor(FLinearColor color);
+
 	//TODO UMG Add Set ReadOnlyForegroundColor
 	//TODO UMG Add Set BackgroundColor
-	//TODO UMG Add Set ForegroundColor
 	//TODO UMG Add Set Font
 
 public:
