Fix widget interaction components created in thumbnail scenes causing focus loss in the editor

diff --git a/Engine/Source/Runtime/UMG/Private/Components/WidgetInteractionComponent.cpp b/Engine/Source/Runtime/UMG/Private/Components/WidgetInteractionComponent.cpp
--- a/Engine/Source/Runtime/UMG/Private/Components/WidgetInteractionComponent.cpp
+++ b/Engine/Source/Runtime/UMG/Private/Components/WidgetInteractionComponent.cpp
@@ -60,7 +60,8 @@ void UWidgetInteractionComponent::Activate(bool bReset)
 {
 	Super::Activate(bReset);
 
-	if ( FSlateApplication::IsInitialized() )
+	// Only create another user in a real world. FindOrCreateVirtualUser changes focus
+	if ( FSlateApplication::IsInitialized() && !GetWorld()->IsPreviewWorld())
 	{
 		if ( !VirtualUser.IsValid() )
 		{
