PR #7296: Added missing 'BlueprintPure' directive to USlateBlueprintLibrary functions (Contributed by [REDACTED_PERSON])

diff --git a/Engine/Source/Runtime/UMG/Public/Blueprint/SlateBlueprintLibrary.h b/Engine/Source/Runtime/UMG/Public/Blueprint/SlateBlueprintLibrary.h
--- a/Engine/Source/Runtime/UMG/Public/Blueprint/SlateBlueprintLibrary.h
+++ b/Engine/Source/Runtime/UMG/Public/Blueprint/SlateBlueprintLibrary.h
@@ -103,19 +103,19 @@ public:
 	 * Translates a screen position in pixels into the local space of a widget with the given geometry. 
 	 * If bIncludeWindowPosition is true, then this method will also remove the game window's position (useful when in windowed mode).
 	 */
-	UFUNCTION(Category="User Interface|Geometry", meta=( WorldContext="WorldContextObject", DisplayName="ScreenToLocal" ))
+	UFUNCTION(BlueprintPure, Category="User Interface|Geometry", meta=( WorldContext="WorldContextObject", DisplayName="ScreenToLocal" ))
 	static void ScreenToWidgetLocal(UObject* WorldContextObject, const FGeometry& Geometry, FVector2D ScreenPosition, FVector2D& LocalCoordinate, bool bIncludeWindowPosition = false);
 
 	/**
 	 * Translates a screen position in pixels into absolute application coordinates.
 	 * If bIncludeWindowPosition is true, then this method will also remove the game window's position (useful when in windowed mode).
 	 */
-	UFUNCTION(Category="User Interface|Geometry", meta=( WorldContext="WorldContextObject", DisplayName="ScreenToAbsolute" ))
+	UFUNCTION(BlueprintPure, Category="User Interface|Geometry", meta=( WorldContext="WorldContextObject", DisplayName="ScreenToAbsolute" ))
 	static void ScreenToWidgetAbsolute(UObject* WorldContextObject, FVector2D ScreenPosition, FVector2D& AbsoluteCoordinate, bool bIncludeWindowPosition = false);
 
 	/**
 	 * Translates a screen position in pixels into the local space of the viewport widget.
 	 */
-	UFUNCTION(Category="User Interface|Geometry", meta=( WorldContext="WorldContextObject" ))
+	UFUNCTION(BlueprintPure, Category="User Interface|Geometry", meta=( WorldContext="WorldContextObject" ))
 	static void ScreenToViewport(UObject* WorldContextObject, FVector2D ScreenPosition, FVector2D& ViewportPosition);
 };
