Slate: Prevent negative size texture in SRetainerWidget.

diff --git a/Engine/Source/Runtime/UMG/Private/Slate/SRetainerWidget.cpp b/Engine/Source/Runtime/UMG/Private/Slate/SRetainerWidget.cpp
--- a/Engine/Source/Runtime/UMG/Private/Slate/SRetainerWidget.cpp
+++ b/Engine/Source/Runtime/UMG/Private/Slate/SRetainerWidget.cpp
@@ -417,17 +417,25 @@ SRetainerWidget::EPaintRetainedContentResult SRetainerWidget::PaintRetainedConte
 		// Size must be a positive integer to allocate the RenderTarget
 		const uint32 RenderTargetWidth  = FMath::RoundToInt(FMath::Abs(RenderSize.X));
 		const uint32 RenderTargetHeight = FMath::RoundToInt(FMath::Abs(RenderSize.Y));
-		const bool bTextureTooLarge = FMath::Max(RenderTargetWidth, RenderTargetHeight) > GetMax2DTextureDimension();
+		const bool bTextureIsTooLarge = FMath::Max(RenderTargetWidth, RenderTargetHeight) > GetMax2DTextureDimension();
+		const bool bTextureSizeZero = (RenderTargetWidth == 0 || RenderTargetHeight == 0);
 
-		if ( bTextureTooLarge )
+		if ( bTextureIsTooLarge || bTextureSizeZero )
 		{
 			// if bTextureTooLarge then the user probably have a layout issue. Warn the user.
 			if ( !bInvalidSizeLogged )
 			{
 				bInvalidSizeLogged = true;
-				UE_LOG(LogUMG, Error, TEXT("The requested size for SRetainerWidget is too large. W:%i H:%i"), RenderTargetWidth, RenderTargetHeight);
+				if (bTextureIsTooLarge)
+				{
+					UE_LOG(LogUMG, Error, TEXT("The requested size for SRetainerWidget is too large. W:%i H:%i"), RenderTargetWidth, RenderTargetHeight);
+				}
+				else
+				{
+					UE_LOG(LogUMG, Error, TEXT("The requested size for SRetainerWidget is 0. W:%i H:%i"), RenderTargetWidth, RenderTargetHeight);
+				}
 			}
-			return EPaintRetainedContentResult::InvalidSize;
+			return bTextureIsTooLarge ? EPaintRetainedContentResult::TextureSizeTooBig : EPaintRetainedContentResult::TextureSizeZero;
 		}
 		bInvalidSizeLogged = false;
 
@@ -516,16 +524,22 @@ int32 SRetainerWidget::OnPaint(const FPaintArgs& Args, const FGeometry& Allotted
 		EPaintRetainedContentResult PaintResult = MutableThis->PaintRetainedContentImpl(Context, AllottedGeometry);
 
 #if WITH_SLATE_DEBUGGING
-		if (PaintResult == EPaintRetainedContentResult::NotPainted || PaintResult == EPaintRetainedContentResult::InvalidSize)
+		if (PaintResult == EPaintRetainedContentResult::NotPainted
+			|| PaintResult == EPaintRetainedContentResult::TextureSizeZero
+			|| PaintResult == EPaintRetainedContentResult::TextureSizeTooBig)
 		{
 			MutableThis->SetLastPaintType(ESlateInvalidationPaintType::None);
 		}
 #endif
 
-		if (PaintResult == EPaintRetainedContentResult::InvalidSize)
+		if (PaintResult == EPaintRetainedContentResult::TextureSizeTooBig)
 		{
 			return SCompoundWidget::OnPaint(Args, AllottedGeometry, MyCullingRect, OutDrawElements, LayerId, InWidgetStyle, bParentEnabled);
 		}
+		else if (PaintResult == EPaintRetainedContentResult::TextureSizeZero)
+		{
+			return GetCachedMaxLayerId();
+		}
 		else
 		{
 			UTextureRenderTarget2D* RenderTarget = RenderingResources->RenderTarget;
diff --git a/Engine/Source/Runtime/UMG/Public/Slate/SRetainerWidget.h b/Engine/Source/Runtime/UMG/Public/Slate/SRetainerWidget.h
--- a/Engine/Source/Runtime/UMG/Public/Slate/SRetainerWidget.h
+++ b/Engine/Source/Runtime/UMG/Public/Slate/SRetainerWidget.h
@@ -103,7 +103,8 @@ protected:
 		NotPainted,
 		Painted,
 		Queued,
-		InvalidSize,
+		TextureSizeTooBig,
+		TextureSizeZero,
 	};
 	EPaintRetainedContentResult PaintRetainedContentImpl(const FSlateInvalidationContext& Context, const FGeometry& AllottedGeometry);
 
