Updated ListView to still allow input to items when selection mode is set to None

See UDN: https://udn.unrealengine.com/s/question/0D52L00004qJQysSAG/bug-slistview-does-not-get-onmousebuttonclicked-event-when-set-to-eselectionmodenone

diff --git a/Engine/Source/Runtime/Slate/Public/Widgets/Views/STableRow.h b/Engine/Source/Runtime/Slate/Public/Widgets/Views/STableRow.h
--- a/Engine/Source/Runtime/Slate/Public/Widgets/Views/STableRow.h
+++ b/Engine/Source/Runtime/Slate/Public/Widgets/Views/STableRow.h
@@ -531,7 +531,8 @@ public:
 			}
 		}
 
-		return FReply::Unhandled();
+		return FReply::Unhandled()
+			.CaptureMouse(SharedThis(this));
 	}
 	
 	/**
diff --git a/Engine/Source/Runtime/UMG/Public/Slate/SObjectTableRow.h b/Engine/Source/Runtime/UMG/Public/Slate/SObjectTableRow.h
--- a/Engine/Source/Runtime/UMG/Public/Slate/SObjectTableRow.h
+++ b/Engine/Source/Runtime/UMG/Public/Slate/SObjectTableRow.h
@@ -324,6 +324,7 @@ public:
 		bChangedSelectionOnMouseDown = false;
 
 		FReply Reply = SObjectWidget::OnMouseButtonDown(MyGeometry, MouseEvent);
+		Reply.CaptureMouse(SharedThis(this));
 		if (!Reply.IsEventHandled())
 		{
 			TSharedRef<ITypedTableView<ItemType>> OwnerTable = OwnerTablePtr.Pin().ToSharedRef();
