Revert a fix to support OnItemClicked in LisView when the list is not selectable (CL 14788542)

diff --git a/Engine/Source/Runtime/Slate/Public/Widgets/Views/STableRow.h b/Engine/Source/Runtime/Slate/Public/Widgets/Views/STableRow.h
--- a/Engine/Source/Runtime/Slate/Public/Widgets/Views/STableRow.h
+++ b/Engine/Source/Runtime/Slate/Public/Widgets/Views/STableRow.h
@@ -531,8 +531,7 @@ public:
 			}
 		}
 
-		return FReply::Unhandled()
-			.CaptureMouse(SharedThis(this));
+		return FReply::Unhandled();
 	}
 	
 	/**
diff --git a/Engine/Source/Runtime/UMG/Public/Slate/SObjectTableRow.h b/Engine/Source/Runtime/UMG/Public/Slate/SObjectTableRow.h
--- a/Engine/Source/Runtime/UMG/Public/Slate/SObjectTableRow.h
+++ b/Engine/Source/Runtime/UMG/Public/Slate/SObjectTableRow.h
@@ -324,7 +324,6 @@ public:
 		bChangedSelectionOnMouseDown = false;
 
 		FReply Reply = SObjectWidget::OnMouseButtonDown(MyGeometry, MouseEvent);
-		Reply.CaptureMouse(SharedThis(this));
 		if (!Reply.IsEventHandled())
 		{
 			TSharedRef<ITypedTableView<ItemType>> OwnerTable = OwnerTablePtr.Pin().ToSharedRef();
@@ -360,8 +359,6 @@ public:
 						Reply.SetUserFocus(OwnerTable->AsWidget(), EFocusCause::Mouse);
 					}
 				}
-
-		
 			}
 		}
 		
