Automated wrapped object pointer upgrade for Engine + ShooterGame + ShooterGame referenced plugins

diff --git a/Engine/Source/Runtime/MovieScene/Public/EntitySystem/MovieSceneEntitySystemGraphs.h b/Engine/Source/Runtime/MovieScene/Public/EntitySystem/MovieSceneEntitySystemGraphs.h
--- a/Engine/Source/Runtime/MovieScene/Public/EntitySystem/MovieSceneEntitySystemGraphs.h
+++ b/Engine/Source/Runtime/MovieScene/Public/EntitySystem/MovieSceneEntitySystemGraphs.h
@@ -45,7 +45,7 @@ struct FMovieSceneEntitySystemGraphNode
 	TSharedPtr<UE::MovieScene::FSystemTaskPrerequisites> SubsequentTasks;
 
 	UPROPERTY()
-	UMovieSceneEntitySystem* System;
+	TObjectPtr<UMovieSceneEntitySystem> System;
 };
 
 USTRUCT()
diff --git a/Engine/Source/Runtime/MovieScene/Public/EntitySystem/TrackInstance/MovieSceneTrackInstance.h b/Engine/Source/Runtime/MovieScene/Public/EntitySystem/TrackInstance/MovieSceneTrackInstance.h
--- a/Engine/Source/Runtime/MovieScene/Public/EntitySystem/TrackInstance/MovieSceneTrackInstance.h
+++ b/Engine/Source/Runtime/MovieScene/Public/EntitySystem/TrackInstance/MovieSceneTrackInstance.h
@@ -17,7 +17,7 @@ struct FMovieSceneTrackInstanceInput
 
 	/** The section that acts as an input for this animator */
 	UPROPERTY()
-	UMovieSceneSection* Section = nullptr;
+	TObjectPtr<UMovieSceneSection> Section = nullptr;
 
 	/** The instance that this input relates to */
 	UE::MovieScene::FInstanceHandle InstanceHandle;
@@ -157,13 +157,13 @@ private:
 private:
 
 	UPROPERTY()
-	UObject* AnimatedObject;
+	TObjectPtr<UObject> AnimatedObject;
 
 	UPROPERTY()
 	bool bIsMasterTrackInstance;
 
 	UPROPERTY()
-	UMovieSceneEntitySystemLinker* Linker;
+	TObjectPtr<UMovieSceneEntitySystemLinker> Linker;
 
 	UPROPERTY()
 	TArray<FMovieSceneTrackInstanceInput> Inputs;
diff --git a/Engine/Source/Runtime/MovieScene/Public/EntitySystem/TrackInstance/MovieSceneTrackInstanceSystem.h b/Engine/Source/Runtime/MovieScene/Public/EntitySystem/TrackInstance/MovieSceneTrackInstanceSystem.h
--- a/Engine/Source/Runtime/MovieScene/Public/EntitySystem/TrackInstance/MovieSceneTrackInstanceSystem.h
+++ b/Engine/Source/Runtime/MovieScene/Public/EntitySystem/TrackInstance/MovieSceneTrackInstanceSystem.h
@@ -18,10 +18,10 @@ struct FMovieSceneTrackInstanceEntry
 	GENERATED_BODY()
 
 	UPROPERTY()
-	UObject* BoundObject = nullptr;
+	TObjectPtr<UObject> BoundObject = nullptr;
 
 	UPROPERTY()
-	UMovieSceneTrackInstance* TrackInstance = nullptr;
+	TObjectPtr<UMovieSceneTrackInstance> TrackInstance = nullptr;
 };
 
 
@@ -77,7 +77,7 @@ private:
 	virtual void OnRun(FSystemTaskPrerequisites& InPrerequisites, FSystemSubsequentTasks& Subsequents) override final;
 
 	UPROPERTY()
-	UMovieSceneTrackInstanceInstantiator* Instantiator;
+	TObjectPtr<UMovieSceneTrackInstanceInstantiator> Instantiator;
 };
 
 
diff --git a/Engine/Source/Runtime/MovieScene/Public/Evaluation/MovieSceneEvaluationTemplateInstance.h b/Engine/Source/Runtime/MovieScene/Public/Evaluation/MovieSceneEvaluationTemplateInstance.h
--- a/Engine/Source/Runtime/MovieScene/Public/Evaluation/MovieSceneEvaluationTemplateInstance.h
+++ b/Engine/Source/Runtime/MovieScene/Public/Evaluation/MovieSceneEvaluationTemplateInstance.h
@@ -131,18 +131,18 @@ private:
 	TWeakObjectPtr<UMovieSceneSequence> WeakRootSequence;
 
 	UPROPERTY()
-	UMovieSceneCompiledDataManager* CompiledDataManager;
+	TObjectPtr<UMovieSceneCompiledDataManager> CompiledDataManager;
 
 	UE::MovieScene::FInstanceHandle RootInstanceHandle;
 
 	UPROPERTY()
-	UMovieSceneEntitySystemLinker* EntitySystemLinker;
+	TObjectPtr<UMovieSceneEntitySystemLinker> EntitySystemLinker;
 
 	FMovieSceneEntitySystemRunner EntitySystemRunner;
 
 	/** Map of director instances by sequence ID. Kept alive by this map assuming this struct is reference collected */
 	UPROPERTY()
-	TMap<FMovieSceneSequenceID, UObject*> DirectorInstances;
+	TMap<FMovieSceneSequenceID, TObjectPtr<UObject>> DirectorInstances;
 
 	FMovieSceneSequenceID RootID;
 
diff --git a/Engine/Source/Runtime/MovieScene/Public/Generators/MovieSceneEasingCurves.h b/Engine/Source/Runtime/MovieScene/Public/Generators/MovieSceneEasingCurves.h
--- a/Engine/Source/Runtime/MovieScene/Public/Generators/MovieSceneEasingCurves.h
+++ b/Engine/Source/Runtime/MovieScene/Public/Generators/MovieSceneEasingCurves.h
@@ -57,5 +57,5 @@ public:
 
 	/** Curve data */
 	UPROPERTY(EditAnywhere, Category=Easing)
-	UCurveFloat* Curve;
+	TObjectPtr<UCurveFloat> Curve;
 };
diff --git a/Engine/Source/Runtime/MovieScene/Public/MovieScene.h b/Engine/Source/Runtime/MovieScene/Public/MovieScene.h
--- a/Engine/Source/Runtime/MovieScene/Public/MovieScene.h
+++ b/Engine/Source/Runtime/MovieScene/Public/MovieScene.h
@@ -343,7 +343,7 @@ public:
 
 private:
 	UPROPERTY()
-	TArray<UMovieSceneNodeGroup*> NodeGroups;
+	TArray<TObjectPtr<UMovieSceneNodeGroup>> NodeGroups;
 
 	bool bAnyActiveFilter;
 
@@ -1160,11 +1160,11 @@ private:
 
 	/** Master tracks which are not bound to spawned or possessed objects */
 	UPROPERTY(Instanced)
-	TArray<UMovieSceneTrack*> MasterTracks;
+	TArray<TObjectPtr<UMovieSceneTrack>> MasterTracks;
 
 	/** The camera cut track is a specialized track for switching between cameras on a cinematic */
 	UPROPERTY(Instanced)
-	UMovieSceneTrack* CameraCutTrack;
+	TObjectPtr<UMovieSceneTrack> CameraCutTrack;
 
 	/** User-defined selection range. */
 	UPROPERTY()
@@ -1220,7 +1220,7 @@ private:
 
 	/** The root folders for this movie scene. */
 	UPROPERTY()
-	TArray<UMovieSceneFolder*> RootFolders;
+	TArray<TObjectPtr<UMovieSceneFolder>> RootFolders;
 
 	/** Nodes currently marked Solo, stored as node tree paths */
 	UPROPERTY()
@@ -1236,7 +1236,7 @@ private:
 
 	/** Collection of user-defined groups */
 	UPROPERTY()
-	UMovieSceneNodeGroupCollection* NodeGroupCollection;
+	TObjectPtr<UMovieSceneNodeGroupCollection> NodeGroupCollection;
 
 	/** Whether this scene's marked frames should be shown globally */
 	bool bGloballyShowMarkedFrames;
diff --git a/Engine/Source/Runtime/MovieScene/Public/MovieSceneBinding.h b/Engine/Source/Runtime/MovieScene/Public/MovieSceneBinding.h
--- a/Engine/Source/Runtime/MovieScene/Public/MovieSceneBinding.h
+++ b/Engine/Source/Runtime/MovieScene/Public/MovieSceneBinding.h
@@ -169,7 +169,7 @@ private:
 
 	/** All tracks in this binding */
 	UPROPERTY(Instanced)
-	TArray<UMovieSceneTrack*> Tracks;
+	TArray<TObjectPtr<UMovieSceneTrack>> Tracks;
 
 #if WITH_EDITORONLY_DATA
 	/** The desired sorting order for this binding in Sequencer */
diff --git a/Engine/Source/Runtime/MovieScene/Public/MovieSceneFolder.h b/Engine/Source/Runtime/MovieScene/Public/MovieSceneFolder.h
--- a/Engine/Source/Runtime/MovieScene/Public/MovieSceneFolder.h
+++ b/Engine/Source/Runtime/MovieScene/Public/MovieSceneFolder.h
@@ -106,11 +106,11 @@ private:
 
 	/** The folders contained by this folder. */
 	UPROPERTY()
-	TArray<UMovieSceneFolder*> ChildFolders;
+	TArray<TObjectPtr<UMovieSceneFolder>> ChildFolders;
 
 	/** The master tracks contained by this folder. */
 	UPROPERTY()
-	TArray<UMovieSceneTrack*> ChildMasterTracks;
+	TArray<TObjectPtr<UMovieSceneTrack>> ChildMasterTracks;
 
 	/** The guid strings used to serialize the guids for the object bindings contained by this folder. */
 	UPROPERTY()
diff --git a/Engine/Source/Runtime/MovieScene/Public/MovieScenePossessable.h b/Engine/Source/Runtime/MovieScene/Public/MovieScenePossessable.h
--- a/Engine/Source/Runtime/MovieScene/Public/MovieScenePossessable.h
+++ b/Engine/Source/Runtime/MovieScene/Public/MovieScenePossessable.h
@@ -131,7 +131,7 @@ private:
 	// possessable "slot" though.  Or we could use it to generate a "preview" spawnable puppet when previewing with no
 	// possessable object available.
 	UPROPERTY()
-	UClass* PossessedObjectClass;
+	TObjectPtr<UClass> PossessedObjectClass;
 
 	/** GUID relating to this possessable's parent, if applicable. */
 	UPROPERTY()
diff --git a/Engine/Source/Runtime/MovieScene/Public/MovieSceneSequence.h b/Engine/Source/Runtime/MovieScene/Public/MovieSceneSequence.h
--- a/Engine/Source/Runtime/MovieScene/Public/MovieSceneSequence.h
+++ b/Engine/Source/Runtime/MovieScene/Public/MovieSceneSequence.h
@@ -266,7 +266,7 @@ private:
 
 	/** Serialized compiled data - should only be used through UMovieSceneCompiledDataManager */
 	UPROPERTY(Instanced)
-	UMovieSceneCompiledData* CompiledData;
+	TObjectPtr<UMovieSceneCompiledData> CompiledData;
 
 public:
 
diff --git a/Engine/Source/Runtime/MovieScene/Public/MovieSceneSequencePlayer.h b/Engine/Source/Runtime/MovieScene/Public/MovieSceneSequencePlayer.h
--- a/Engine/Source/Runtime/MovieScene/Public/MovieSceneSequencePlayer.h
+++ b/Engine/Source/Runtime/MovieScene/Public/MovieSceneSequencePlayer.h
@@ -617,7 +617,7 @@ protected:
 
 	/** The sequence to play back */
 	UPROPERTY(transient)
-	UMovieSceneSequence* Sequence;
+	TObjectPtr<UMovieSceneSequence> Sequence;
 
 	/** Time (in playback frames) at which to start playing the sequence (defaults to the lower bound of the sequence's play range) */
 	UPROPERTY(replicated)
@@ -657,7 +657,7 @@ protected:
 
 	/** Global tick manager, held here to keep it alive while world sequences are in play */
 	UPROPERTY(transient)
-	UMovieSceneSequenceTickManager* TickManager;
+	TObjectPtr<UMovieSceneSequenceTickManager> TickManager;
 
 	/** Local latent action manager for when we're running a blocking sequence */
 	FMovieSceneLatentActionManager LatentActionManager;
diff --git a/Engine/Source/Runtime/MovieScene/Public/MovieSceneSequenceTickManager.h b/Engine/Source/Runtime/MovieScene/Public/MovieSceneSequenceTickManager.h
--- a/Engine/Source/Runtime/MovieScene/Public/MovieSceneSequenceTickManager.h
+++ b/Engine/Source/Runtime/MovieScene/Public/MovieSceneSequenceTickManager.h
@@ -57,11 +57,11 @@ private:
 
 public:
 	UPROPERTY(transient)
-	TArray<AActor*> SequenceActors;
+	TArray<TObjectPtr<AActor>> SequenceActors;
 
 private:
 	UPROPERTY(transient)
-	UMovieSceneEntitySystemLinker* Linker;
+	TObjectPtr<UMovieSceneEntitySystemLinker> Linker;
 
 	FMovieSceneEntitySystemRunner Runner;
 
diff --git a/Engine/Source/Runtime/MovieScene/Public/MovieSceneSpawnable.h b/Engine/Source/Runtime/MovieScene/Public/MovieSceneSpawnable.h
--- a/Engine/Source/Runtime/MovieScene/Public/MovieSceneSpawnable.h
+++ b/Engine/Source/Runtime/MovieScene/Public/MovieSceneSpawnable.h
@@ -266,7 +266,7 @@ private:
 	FString Name;
 
 	UPROPERTY()
-	UObject* ObjectTemplate;
+	TObjectPtr<UObject> ObjectTemplate;
 
 	/** Set of GUIDs to possessable object bindings that are bound to an object inside this spawnable */
 	// @todo sequencer: This should be a TSet, but they don't duplicate correctly atm
@@ -281,7 +281,7 @@ private:
 public:
 	/** Deprecated generated class */
 	UPROPERTY()
-	UClass* GeneratedClass_DEPRECATED;
+	TObjectPtr<UClass> GeneratedClass_DEPRECATED;
 #endif
 
 	/** Name of level to spawn into */
diff --git a/Engine/Source/Runtime/MovieScene/Public/MovieSceneTrackEvaluationField.h b/Engine/Source/Runtime/MovieScene/Public/MovieSceneTrackEvaluationField.h
--- a/Engine/Source/Runtime/MovieScene/Public/MovieSceneTrackEvaluationField.h
+++ b/Engine/Source/Runtime/MovieScene/Public/MovieSceneTrackEvaluationField.h
@@ -21,7 +21,7 @@ struct FMovieSceneTrackEvaluationFieldEntry
 	GENERATED_BODY()
 
 	UPROPERTY()
-	UMovieSceneSection* Section = nullptr;
+	TObjectPtr<UMovieSceneSection> Section = nullptr;
 
 	UPROPERTY()
 	FFrameNumberRange Range;
diff --git a/Engine/Source/Runtime/MovieScene/Public/Sections/MovieSceneSubSection.h b/Engine/Source/Runtime/MovieScene/Public/Sections/MovieSceneSubSection.h
--- a/Engine/Source/Runtime/MovieScene/Public/Sections/MovieSceneSubSection.h
+++ b/Engine/Source/Runtime/MovieScene/Public/Sections/MovieSceneSubSection.h
@@ -197,7 +197,7 @@ protected:
 	 * @todo Sequencer: Should this be lazy loaded?
 	 */
 	UPROPERTY(EditAnywhere, Category="Sequence")
-	UMovieSceneSequence* SubSequence;
+	TObjectPtr<UMovieSceneSequence> SubSequence;
 
 	/** Target actor to record */
 	UPROPERTY(EditAnywhere, Category="Sequence Recording")
diff --git a/Engine/Source/Runtime/MovieSceneCapture/Public/MovieSceneCapture.h b/Engine/Source/Runtime/MovieSceneCapture/Public/MovieSceneCapture.h
--- a/Engine/Source/Runtime/MovieSceneCapture/Public/MovieSceneCapture.h
+++ b/Engine/Source/Runtime/MovieSceneCapture/Public/MovieSceneCapture.h
@@ -89,10 +89,10 @@ public:
 
 	/** Capture protocol responsible for actually capturing frame data */
 	UPROPERTY(VisibleAnywhere, Category=CaptureSettings, Transient, Instanced)
-	UMovieSceneImageCaptureProtocolBase* ImageCaptureProtocol;
+	TObjectPtr<UMovieSceneImageCaptureProtocolBase> ImageCaptureProtocol;
 	
 	UPROPERTY(VisibleAnywhere, Category = CaptureSettings, Transient, Instanced)
-	UMovieSceneAudioCaptureProtocolBase* AudioCaptureProtocol;
+	TObjectPtr<UMovieSceneAudioCaptureProtocolBase> AudioCaptureProtocol;
 
 	/** Settings that define how to capture */
 	UPROPERTY(config, EditAnywhere, BlueprintReadWrite, Category=CaptureSettings, meta=(ShowOnlyInnerProperties))
diff --git a/Engine/Source/Runtime/MovieSceneCapture/Public/Protocols/CompositionGraphCaptureProtocol.h b/Engine/Source/Runtime/MovieSceneCapture/Public/Protocols/CompositionGraphCaptureProtocol.h
--- a/Engine/Source/Runtime/MovieSceneCapture/Public/Protocols/CompositionGraphCaptureProtocol.h
+++ b/Engine/Source/Runtime/MovieSceneCapture/Public/Protocols/CompositionGraphCaptureProtocol.h
@@ -87,7 +87,7 @@ public:
 private:
 
 	UPROPERTY(transient)
-	UMaterialInterface* PostProcessingMaterialPtr;
+	TObjectPtr<UMaterialInterface> PostProcessingMaterialPtr;
 
 	/** The viewport we are capturing from */
 	TWeakPtr<FSceneViewport> SceneViewport;
diff --git a/Engine/Source/Runtime/MovieSceneCapture/Public/Protocols/UserDefinedCaptureProtocol.h b/Engine/Source/Runtime/MovieSceneCapture/Public/Protocols/UserDefinedCaptureProtocol.h
--- a/Engine/Source/Runtime/MovieSceneCapture/Public/Protocols/UserDefinedCaptureProtocol.h
+++ b/Engine/Source/Runtime/MovieSceneCapture/Public/Protocols/UserDefinedCaptureProtocol.h
@@ -207,7 +207,7 @@ protected:
 
 	/** World pointer assigned on Setup */
 	UPROPERTY(transient, BlueprintReadOnly, Category=Capture)
-	UWorld* World;
+	TObjectPtr<UWorld> World;
 
 protected:
 
diff --git a/Engine/Source/Runtime/MovieSceneTracks/Private/Evaluation/MovieSceneAudioTemplate.h b/Engine/Source/Runtime/MovieSceneTracks/Private/Evaluation/MovieSceneAudioTemplate.h
--- a/Engine/Source/Runtime/MovieSceneTracks/Private/Evaluation/MovieSceneAudioTemplate.h
+++ b/Engine/Source/Runtime/MovieSceneTracks/Private/Evaluation/MovieSceneAudioTemplate.h
@@ -23,7 +23,7 @@ struct FMovieSceneAudioSectionTemplate : public FMovieSceneEvalTemplate
 	FMovieSceneAudioSectionTemplate(const UMovieSceneAudioSection& Section);
 
 	UPROPERTY()
-	const UMovieSceneAudioSection* AudioSection;
+	TObjectPtr<const UMovieSceneAudioSection> AudioSection;
 
 private:
 
diff --git a/Engine/Source/Runtime/MovieSceneTracks/Private/Evaluation/MovieSceneMaterialParameterCollectionTemplate.h b/Engine/Source/Runtime/MovieSceneTracks/Private/Evaluation/MovieSceneMaterialParameterCollectionTemplate.h
--- a/Engine/Source/Runtime/MovieSceneTracks/Private/Evaluation/MovieSceneMaterialParameterCollectionTemplate.h
+++ b/Engine/Source/Runtime/MovieSceneTracks/Private/Evaluation/MovieSceneMaterialParameterCollectionTemplate.h
@@ -24,5 +24,5 @@ private:
 	virtual void Evaluate(const FMovieSceneEvaluationOperand& Operand, const FMovieSceneContext& Context, const FPersistentEvaluationData& PersistentData, FMovieSceneExecutionTokens& ExecutionTokens) const override;
 
 	UPROPERTY()
-	UMaterialParameterCollection* MPC;
+	TObjectPtr<UMaterialParameterCollection> MPC;
 };
diff --git a/Engine/Source/Runtime/MovieSceneTracks/Public/Channels/MovieSceneEvent.h b/Engine/Source/Runtime/MovieSceneTracks/Public/Channels/MovieSceneEvent.h
--- a/Engine/Source/Runtime/MovieSceneTracks/Public/Channels/MovieSceneEvent.h
+++ b/Engine/Source/Runtime/MovieSceneTracks/Public/Channels/MovieSceneEvent.h
@@ -35,7 +35,7 @@ struct FMovieSceneEventPtrs
 	{}
 
 	UPROPERTY()
-	UFunction* Function;
+	TObjectPtr<UFunction> Function;
 
 	UPROPERTY()
 	TFieldPath<FProperty> BoundObjectProperty;
diff --git a/Engine/Source/Runtime/MovieSceneTracks/Public/MovieSceneTracksComponentTypes.h b/Engine/Source/Runtime/MovieSceneTracks/Public/MovieSceneTracksComponentTypes.h
--- a/Engine/Source/Runtime/MovieSceneTracks/Public/MovieSceneTracksComponentTypes.h
+++ b/Engine/Source/Runtime/MovieSceneTracks/Public/MovieSceneTracksComponentTypes.h
@@ -21,7 +21,7 @@ struct FLevelVisibilityComponentData
 	GENERATED_BODY()
 
 	UPROPERTY()
-	const UMovieSceneLevelVisibilitySection* Section = nullptr;
+	TObjectPtr<const UMovieSceneLevelVisibilitySection> Section = nullptr;
 };
 
 namespace UE
diff --git a/Engine/Source/Runtime/MovieSceneTracks/Public/Sections/MovieSceneAudioSection.h b/Engine/Source/Runtime/MovieSceneTracks/Public/Sections/MovieSceneAudioSection.h
--- a/Engine/Source/Runtime/MovieSceneTracks/Public/Sections/MovieSceneAudioSection.h
+++ b/Engine/Source/Runtime/MovieSceneTracks/Public/Sections/MovieSceneAudioSection.h
@@ -173,7 +173,7 @@ private:
 
 	/** The sound cue or wave that this section plays */
 	UPROPERTY(EditAnywhere, Category="Audio")
-	USoundBase* Sound;
+	TObjectPtr<USoundBase> Sound;
 
 	/** The offset into the beginning of the audio clip */
 	UPROPERTY(EditAnywhere, Category="Audio")
@@ -219,7 +219,7 @@ private:
 
 	/** The attenuation settings to use. */
 	UPROPERTY( EditAnywhere, Category="Attenuation" )
-	class USoundAttenuation* AttenuationSettings;
+	TObjectPtr<class USoundAttenuation> AttenuationSettings;
 
 	/** Called when subtitles are sent to the SubtitleManager.  Set this delegate if you want to hijack the subtitles for other purposes */
 	UPROPERTY()
diff --git a/Engine/Source/Runtime/MovieSceneTracks/Public/Sections/MovieSceneCameraAnimSection.h b/Engine/Source/Runtime/MovieSceneTracks/Public/Sections/MovieSceneCameraAnimSection.h
--- a/Engine/Source/Runtime/MovieSceneTracks/Public/Sections/MovieSceneCameraAnimSection.h
+++ b/Engine/Source/Runtime/MovieSceneTracks/Public/Sections/MovieSceneCameraAnimSection.h
@@ -28,7 +28,7 @@ struct FMovieSceneCameraAnimSectionData
 
 	/** The camera anim to play */
 	UPROPERTY(EditAnywhere, Category = "Camera Anim")
-	UCameraAnim* CameraAnim;
+	TObjectPtr<UCameraAnim> CameraAnim;
 
 	/** How fast to play back the animation. */
 	UPROPERTY(EditAnywhere, Category = "Camera Anim")
@@ -73,7 +73,7 @@ private:
 
 	/** Deprecated members */
 	UPROPERTY()
-	UCameraAnim* CameraAnim_DEPRECATED;
+	TObjectPtr<UCameraAnim> CameraAnim_DEPRECATED;
 
 	UPROPERTY()
 	float PlayRate_DEPRECATED;
diff --git a/Engine/Source/Runtime/MovieSceneTracks/Public/Sections/MovieSceneSkeletalAnimationSection.h b/Engine/Source/Runtime/MovieSceneTracks/Public/Sections/MovieSceneSkeletalAnimationSection.h
--- a/Engine/Source/Runtime/MovieSceneTracks/Public/Sections/MovieSceneSkeletalAnimationSection.h
+++ b/Engine/Source/Runtime/MovieSceneTracks/Public/Sections/MovieSceneSkeletalAnimationSection.h
@@ -27,7 +27,7 @@ struct FMovieSceneSkeletalAnimationParams
 
 	/** The animation this section plays */
 	UPROPERTY(BlueprintReadWrite, EditAnywhere, Category="Animation", meta=(AllowedClasses = "AnimSequence,AnimComposite,AnimStreamable"))
-	UAnimSequenceBase* Animation;
+	TObjectPtr<UAnimSequenceBase> Animation;
 
 	/** The offset into the beginning of the animation clip for the first loop of play. */
 	UPROPERTY(BlueprintReadWrite, EditAnywhere, Category="Animation")
@@ -128,10 +128,10 @@ private:
 private:
 
 	UPROPERTY()
-	class UAnimSequence* AnimSequence_DEPRECATED;
+	TObjectPtr<class UAnimSequence> AnimSequence_DEPRECATED;
 
 	UPROPERTY()
-	UAnimSequenceBase* Animation_DEPRECATED;
+	TObjectPtr<UAnimSequenceBase> Animation_DEPRECATED;
 
 	UPROPERTY()
 	float StartOffset_DEPRECATED;
diff --git a/Engine/Source/Runtime/MovieSceneTracks/Public/Systems/MovieScenePropertySystem.h b/Engine/Source/Runtime/MovieSceneTracks/Public/Systems/MovieScenePropertySystem.h
--- a/Engine/Source/Runtime/MovieSceneTracks/Public/Systems/MovieScenePropertySystem.h
+++ b/Engine/Source/Runtime/MovieSceneTracks/Public/Systems/MovieScenePropertySystem.h
@@ -38,7 +38,7 @@ protected:
 
 	/** Pointer to the property instantiator system for retrieving property stats */
 	UPROPERTY()
-	UMovieScenePropertyInstantiatorSystem* InstantiatorSystem;
+	TObjectPtr<UMovieScenePropertyInstantiatorSystem> InstantiatorSystem;
 
 	/** Must be set on construction - the composite type of the property this system operates with */
 	UE::MovieScene::FCompositePropertyTypeID CompositePropertyID;
diff --git a/Engine/Source/Runtime/MovieSceneTracks/Public/Tracks/MovieScene3DConstraintTrack.h b/Engine/Source/Runtime/MovieSceneTracks/Public/Tracks/MovieScene3DConstraintTrack.h
--- a/Engine/Source/Runtime/MovieSceneTracks/Public/Tracks/MovieScene3DConstraintTrack.h
+++ b/Engine/Source/Runtime/MovieSceneTracks/Public/Tracks/MovieScene3DConstraintTrack.h
@@ -47,5 +47,5 @@ protected:
 
 	/** List of all constraint sections. */
 	UPROPERTY()
-	TArray<UMovieSceneSection*> ConstraintSections;
+	TArray<TObjectPtr<UMovieSceneSection>> ConstraintSections;
 };
diff --git a/Engine/Source/Runtime/MovieSceneTracks/Public/Tracks/MovieSceneAudioTrack.h b/Engine/Source/Runtime/MovieSceneTracks/Public/Tracks/MovieSceneAudioTrack.h
--- a/Engine/Source/Runtime/MovieSceneTracks/Public/Tracks/MovieSceneAudioTrack.h
+++ b/Engine/Source/Runtime/MovieSceneTracks/Public/Tracks/MovieSceneAudioTrack.h
@@ -65,7 +65,7 @@ private:
 
 	/** List of all master audio sections */
 	UPROPERTY()
-	TArray<UMovieSceneSection*> AudioSections;
+	TArray<TObjectPtr<UMovieSceneSection>> AudioSections;
 
 #if WITH_EDITORONLY_DATA
 
diff --git a/Engine/Source/Runtime/MovieSceneTracks/Public/Tracks/MovieSceneCameraAnimTrack.h b/Engine/Source/Runtime/MovieSceneTracks/Public/Tracks/MovieSceneCameraAnimTrack.h
--- a/Engine/Source/Runtime/MovieSceneTracks/Public/Tracks/MovieSceneCameraAnimTrack.h
+++ b/Engine/Source/Runtime/MovieSceneTracks/Public/Tracks/MovieSceneCameraAnimTrack.h
@@ -48,5 +48,5 @@ public:
 private:
 	/** List of all sections */
 	UPROPERTY()
-	TArray<UMovieSceneSection*> CameraAnimSections;
+	TArray<TObjectPtr<UMovieSceneSection>> CameraAnimSections;
 };
diff --git a/Engine/Source/Runtime/MovieSceneTracks/Public/Tracks/MovieSceneCameraShakeSourceShakeTrack.h b/Engine/Source/Runtime/MovieSceneTracks/Public/Tracks/MovieSceneCameraShakeSourceShakeTrack.h
--- a/Engine/Source/Runtime/MovieSceneTracks/Public/Tracks/MovieSceneCameraShakeSourceShakeTrack.h
+++ b/Engine/Source/Runtime/MovieSceneTracks/Public/Tracks/MovieSceneCameraShakeSourceShakeTrack.h
@@ -49,6 +49,6 @@ public:
 private:
 	/** List of all sections */
 	UPROPERTY()
-	TArray<UMovieSceneSection*> CameraShakeSections;
+	TArray<TObjectPtr<UMovieSceneSection>> CameraShakeSections;
 };
 
diff --git a/Engine/Source/Runtime/MovieSceneTracks/Public/Tracks/MovieSceneCameraShakeTrack.h b/Engine/Source/Runtime/MovieSceneTracks/Public/Tracks/MovieSceneCameraShakeTrack.h
--- a/Engine/Source/Runtime/MovieSceneTracks/Public/Tracks/MovieSceneCameraShakeTrack.h
+++ b/Engine/Source/Runtime/MovieSceneTracks/Public/Tracks/MovieSceneCameraShakeTrack.h
@@ -49,6 +49,6 @@ public:
 private:
 	/** List of all sections */
 	UPROPERTY()
-	TArray<UMovieSceneSection*> CameraShakeSections;
+	TArray<TObjectPtr<UMovieSceneSection>> CameraShakeSections;
 
 };
diff --git a/Engine/Source/Runtime/MovieSceneTracks/Public/Tracks/MovieSceneMaterialParameterCollectionTrack.h b/Engine/Source/Runtime/MovieSceneTracks/Public/Tracks/MovieSceneMaterialParameterCollectionTrack.h
--- a/Engine/Source/Runtime/MovieSceneTracks/Public/Tracks/MovieSceneMaterialParameterCollectionTrack.h
+++ b/Engine/Source/Runtime/MovieSceneTracks/Public/Tracks/MovieSceneMaterialParameterCollectionTrack.h
@@ -23,7 +23,7 @@ public:
 
 	/** The material parameter collection to manipulate */
 	UPROPERTY(EditAnywhere, Category=General, DisplayName="Material Parameter Collection")
-	UMaterialParameterCollection* MPC;
+	TObjectPtr<UMaterialParameterCollection> MPC;
 
 	UMovieSceneMaterialParameterCollectionTrack(const FObjectInitializer& ObjectInitializer);
 
diff --git a/Engine/Source/Runtime/MovieSceneTracks/Public/Tracks/MovieSceneObjectPropertyTrack.h b/Engine/Source/Runtime/MovieSceneTracks/Public/Tracks/MovieSceneObjectPropertyTrack.h
--- a/Engine/Source/Runtime/MovieSceneTracks/Public/Tracks/MovieSceneObjectPropertyTrack.h
+++ b/Engine/Source/Runtime/MovieSceneTracks/Public/Tracks/MovieSceneObjectPropertyTrack.h
@@ -15,7 +15,7 @@ public:
 	GENERATED_BODY()
 
 	UPROPERTY()
-	UClass* PropertyClass;
+	TObjectPtr<UClass> PropertyClass;
 
 	UMovieSceneObjectPropertyTrack(const FObjectInitializer& ObjInit);
 
diff --git a/Engine/Source/Runtime/MovieSceneTracks/Public/Tracks/MovieSceneParticleTrack.h b/Engine/Source/Runtime/MovieSceneTracks/Public/Tracks/MovieSceneParticleTrack.h
--- a/Engine/Source/Runtime/MovieSceneTracks/Public/Tracks/MovieSceneParticleTrack.h
+++ b/Engine/Source/Runtime/MovieSceneTracks/Public/Tracks/MovieSceneParticleTrack.h
@@ -56,5 +56,5 @@ private:
 
 	/** List of all particle sections. */
 	UPROPERTY()
-	TArray<UMovieSceneSection*> ParticleSections;
+	TArray<TObjectPtr<UMovieSceneSection>> ParticleSections;
 };
diff --git a/Engine/Source/Runtime/MovieSceneTracks/Public/Tracks/MovieScenePropertyTrack.h b/Engine/Source/Runtime/MovieSceneTracks/Public/Tracks/MovieScenePropertyTrack.h
--- a/Engine/Source/Runtime/MovieSceneTracks/Public/Tracks/MovieScenePropertyTrack.h
+++ b/Engine/Source/Runtime/MovieSceneTracks/Public/Tracks/MovieScenePropertyTrack.h
@@ -130,7 +130,7 @@ public:
 private:
 	/** Section we should Key */
 	UPROPERTY()
-	UMovieSceneSection* SectionToKey;
+	TObjectPtr<UMovieSceneSection> SectionToKey;
 
 protected:
 
@@ -139,5 +139,5 @@ protected:
 
 	/** All the sections in this list */
 	UPROPERTY()
-	TArray<UMovieSceneSection*> Sections;
+	TArray<TObjectPtr<UMovieSceneSection>> Sections;
 };
diff --git a/Engine/Source/Runtime/NavigationSystem/Public/NavGraph/NavigationGraph.h b/Engine/Source/Runtime/NavigationSystem/Public/NavGraph/NavigationGraph.h
--- a/Engine/Source/Runtime/NavigationSystem/Public/NavGraph/NavigationGraph.h
+++ b/Engine/Source/Runtime/NavigationSystem/Public/NavGraph/NavigationGraph.h
@@ -41,7 +41,7 @@ struct FNavGraphNode
 
 	/** Who's this node referring to? This will most commonly point to an actor or a component */
 	UPROPERTY()
-	UObject* Owner;
+	TObjectPtr<UObject> Owner;
 
 	enum {
 		InitialEdgesCount = 4
diff --git a/Engine/Source/Runtime/NavigationSystem/Public/NavGraph/NavigationGraphNodeComponent.h b/Engine/Source/Runtime/NavigationSystem/Public/NavGraph/NavigationGraphNodeComponent.h
--- a/Engine/Source/Runtime/NavigationSystem/Public/NavGraph/NavigationGraphNodeComponent.h
+++ b/Engine/Source/Runtime/NavigationSystem/Public/NavGraph/NavigationGraphNodeComponent.h
@@ -17,10 +17,10 @@ class UNavigationGraphNodeComponent : public USceneComponent
 	FNavGraphNode Node;
 
 	UPROPERTY()
-	UNavigationGraphNodeComponent* NextNodeComponent;
+	TObjectPtr<UNavigationGraphNodeComponent> NextNodeComponent;
 
 	UPROPERTY()
-	UNavigationGraphNodeComponent* PrevNodeComponent;
+	TObjectPtr<UNavigationGraphNodeComponent> PrevNodeComponent;
 
 	//~ Begin UObject Interface.
 	virtual void BeginDestroy() override;
diff --git a/Engine/Source/Runtime/NavigationSystem/Public/NavRelevantComponent.h b/Engine/Source/Runtime/NavigationSystem/Public/NavRelevantComponent.h
--- a/Engine/Source/Runtime/NavigationSystem/Public/NavRelevantComponent.h
+++ b/Engine/Source/Runtime/NavigationSystem/Public/NavRelevantComponent.h
@@ -49,5 +49,5 @@ protected:
 	uint32 bNavParentCacheInitialized : 1;
 
 	UPROPERTY(transient)
-	UObject* CachedNavParent;
+	TObjectPtr<UObject> CachedNavParent;
 };
diff --git a/Engine/Source/Runtime/NavigationSystem/Public/NavSystemConfigOverride.h b/Engine/Source/Runtime/NavigationSystem/Public/NavSystemConfigOverride.h
--- a/Engine/Source/Runtime/NavigationSystem/Public/NavSystemConfigOverride.h
+++ b/Engine/Source/Runtime/NavigationSystem/Public/NavSystemConfigOverride.h
@@ -28,12 +28,12 @@ class NAVIGATIONSYSTEM_API ANavSystemConfigOverride : public AActor
 #if WITH_EDITORONLY_DATA
 private:
 	UPROPERTY()
-	class UBillboardComponent* SpriteComponent;
+	TObjectPtr<class UBillboardComponent> SpriteComponent;
 #endif // WITH_EDITORONLY_DATA
 
 protected:
 	UPROPERTY(EditAnywhere, BlueprintReadOnly, Category = Navigation, Instanced,  meta = (NoResetToDefault))
-	UNavigationSystemConfig* NavigationSystemConfig;
+	TObjectPtr<UNavigationSystemConfig> NavigationSystemConfig;
 
 	/** If there's already a NavigationSystem instance in the world how should this nav override behave */
 	UPROPERTY(EditAnywhere, BlueprintReadOnly, Category = Navigation)
diff --git a/Engine/Source/Runtime/NavigationSystem/Public/NavigationData.h b/Engine/Source/Runtime/NavigationSystem/Public/NavigationData.h
--- a/Engine/Source/Runtime/NavigationSystem/Public/NavigationData.h
+++ b/Engine/Source/Runtime/NavigationSystem/Public/NavigationData.h
@@ -36,7 +36,7 @@ struct NAVIGATIONSYSTEM_API FSupportedAreaData
 	int32 AreaID;
 
 	UPROPERTY(transient)
-	const UClass* AreaClass;
+	TObjectPtr<const UClass> AreaClass;
 
 	FSupportedAreaData(TSubclassOf<UNavArea> NavAreaClass = NULL, int32 InAreaID = INDEX_NONE);
 };
@@ -496,7 +496,7 @@ class NAVIGATIONSYSTEM_API ANavigationData : public AActor, public INavigationDa
 	GENERATED_UCLASS_BODY()
 	
 	UPROPERTY(transient, duplicatetransient)
-	UPrimitiveComponent* RenderingComp;
+	TObjectPtr<UPrimitiveComponent> RenderingComp;
 
 protected:
 	UPROPERTY()
diff --git a/Engine/Source/Runtime/NavigationSystem/Public/NavigationSystem.h b/Engine/Source/Runtime/NavigationSystem/Public/NavigationSystem.h
--- a/Engine/Source/Runtime/NavigationSystem/Public/NavigationSystem.h
+++ b/Engine/Source/Runtime/NavigationSystem/Public/NavigationSystem.h
@@ -203,11 +203,11 @@ public:
 	virtual ~UNavigationSystemV1();
 
 	UPROPERTY(Transient)
-	ANavigationData* MainNavData;
+	TObjectPtr<ANavigationData> MainNavData;
 
 	/** special navigation data for managing direct paths, not part of NavDataSet! */
 	UPROPERTY(Transient)
-	ANavigationData* AbstractNavData;
+	TObjectPtr<ANavigationData> AbstractNavData;
 
 protected:
 	/** If not None indicates which of navigation datas and supported agents are
@@ -304,10 +304,10 @@ protected:
 public:
 
 	UPROPERTY(Transient)
-	TArray<ANavigationData*> NavDataSet;
+	TArray<TObjectPtr<ANavigationData>> NavDataSet;
 
 	UPROPERTY(Transient)
-	TArray<ANavigationData*> NavDataRegistrationQueue;
+	TArray<TObjectPtr<ANavigationData>> NavDataRegistrationQueue;
 
 	// List of pending navigation bounds update requests (add, remove, update size)
 	TArray<FNavigationBoundsUpdateRequest> PendingNavBoundsUpdates;
diff --git a/Engine/Source/Runtime/NavigationSystem/Public/NavigationTestingActor.h b/Engine/Source/Runtime/NavigationSystem/Public/NavigationTestingActor.h
--- a/Engine/Source/Runtime/NavigationSystem/Public/NavigationTestingActor.h
+++ b/Engine/Source/Runtime/NavigationSystem/Public/NavigationTestingActor.h
@@ -47,16 +47,16 @@ class NAVIGATIONSYSTEM_API ANavigationTestingActor : public AActor, public INavA
 
 private:
 	UPROPERTY()
-	class UCapsuleComponent* CapsuleComponent;
+	TObjectPtr<class UCapsuleComponent> CapsuleComponent;
 
 #if WITH_EDITORONLY_DATA
 	/** Editor Preview */
 	UPROPERTY()
-	class UNavTestRenderingComponent* EdRenderComp;
+	TObjectPtr<class UNavTestRenderingComponent> EdRenderComp;
 #endif // WITH_EDITORONLY_DATA
 
 	UPROPERTY(EditAnywhere, Category = Navigation, meta=(EditCondition="bActAsNavigationInvoker"))
-	UNavigationInvokerComponent* InvokerComponent;
+	TObjectPtr<UNavigationInvokerComponent> InvokerComponent;
 
 	UPROPERTY(EditAnywhere, Category = Navigation, meta=(InlineEditConditionToggle))
 	uint32 bActAsNavigationInvoker : 1;
@@ -71,7 +71,7 @@ public:
 	FVector QueryingExtent;
 
 	UPROPERTY(transient)
-	ANavigationData* MyNavData;
+	TObjectPtr<ANavigationData> MyNavData;
 
 	UPROPERTY(VisibleAnywhere, BlueprintReadOnly, Category=AgentStatus)
 	FVector ProjectedLocation;
@@ -153,7 +153,7 @@ public:
 	int32 PathfindingSteps;
 
 	UPROPERTY(EditAnywhere, Category=Pathfinding)
-	ANavigationTestingActor* OtherActor;
+	TObjectPtr<ANavigationTestingActor> OtherActor;
 
 	/** "None" will result in default filter being used */
 	UPROPERTY(EditAnywhere, Category=Pathfinding)
diff --git a/Engine/Source/Runtime/Overlay/Public/BasicOverlays.h b/Engine/Source/Runtime/Overlay/Public/BasicOverlays.h
--- a/Engine/Source/Runtime/Overlay/Public/BasicOverlays.h
+++ b/Engine/Source/Runtime/Overlay/Public/BasicOverlays.h
@@ -32,7 +32,7 @@ public:
 
 	/** The import data used to make this overlays asset */
 	UPROPERTY(VisibleAnywhere, Instanced, Category="Import Settings")
-	UAssetImportData* AssetImportData;
+	TObjectPtr<UAssetImportData> AssetImportData;
 
 #endif	// WITH_EDITORONLY_DATA
 
diff --git a/Engine/Source/Runtime/Overlay/Public/LocalizedOverlays.h b/Engine/Source/Runtime/Overlay/Public/LocalizedOverlays.h
--- a/Engine/Source/Runtime/Overlay/Public/LocalizedOverlays.h
+++ b/Engine/Source/Runtime/Overlay/Public/LocalizedOverlays.h
@@ -29,7 +29,7 @@ public:
 
 	/** The overlays to use if no overlays are found for the current culture */
 	UPROPERTY(EditAnywhere, Category="Overlay Data")
-	UBasicOverlays* DefaultOverlays;
+	TObjectPtr<UBasicOverlays> DefaultOverlays;
 
 	/**
 	 * Maps a set of cultures to specific BasicOverlays assets.
@@ -39,13 +39,13 @@ public:
 	 *		An optional two-letter ISO 3166-1 country code  (e.g., "CN")
 	 */
 	UPROPERTY(EditAnywhere, Category="Overlay Data")
-	TMap<FString, UBasicOverlays*> LocaleToOverlaysMap;
+	TMap<FString, TObjectPtr<UBasicOverlays>> LocaleToOverlaysMap;
 
 #if WITH_EDITORONLY_DATA
 
 	/** The import data used to make this overlays asset */
 	UPROPERTY(VisibleAnywhere, Instanced, Category="Import Settings")
-	UAssetImportData* AssetImportData;
+	TObjectPtr<UAssetImportData> AssetImportData;
 
 #endif	// WITH_EDITORONLY_DATA
 
diff --git a/Engine/Source/Runtime/PhysicsCore/Public/PhysicalMaterials/PhysicalMaterial.h b/Engine/Source/Runtime/PhysicsCore/Public/PhysicalMaterials/PhysicalMaterial.h
--- a/Engine/Source/Runtime/PhysicsCore/Public/PhysicalMaterials/PhysicalMaterial.h
+++ b/Engine/Source/Runtime/PhysicsCore/Public/PhysicalMaterials/PhysicalMaterial.h
@@ -93,7 +93,7 @@ class PHYSICSCORE_API UPhysicalMaterial : public UObject
 	float DestructibleDamageThresholdScale;
 
 	UPROPERTY()
-	class UDEPRECATED_PhysicalMaterialPropertyBase* PhysicalMaterialProperty_DEPRECATED;
+	TObjectPtr<class UDEPRECATED_PhysicalMaterialPropertyBase> PhysicalMaterialProperty_DEPRECATED;
 
 	/**
 	 * To edit surface type for your project, use ProjectSettings/Physics/PhysicalSurface section
diff --git a/Engine/Source/Runtime/PropertyAccess/Public/PropertyAccess.h b/Engine/Source/Runtime/PropertyAccess/Public/PropertyAccess.h
--- a/Engine/Source/Runtime/PropertyAccess/Public/PropertyAccess.h
+++ b/Engine/Source/Runtime/PropertyAccess/Public/PropertyAccess.h
@@ -98,7 +98,7 @@ private:
 
 	// Function if this is a script of native function indirection
 	UPROPERTY()
-	UFunction* Function = nullptr;
+	TObjectPtr<UFunction> Function = nullptr;
 
 	// Return buffer size if this is a script of native function indirection
 	UPROPERTY()
@@ -217,7 +217,7 @@ private:
 
 	/** The Class or ScriptStruct that was used last to resolve Name to a property. */
 	UPROPERTY()
-	UStruct* Struct = nullptr;
+	TObjectPtr<UStruct> Struct = nullptr;
 
 	/** The cached property on the Struct that this Name resolved to at compile time. If this is a Function segment, then this is the return property of the function. */
 	UPROPERTY()
@@ -225,7 +225,7 @@ private:
 
 	/** If this segment is a function, EPropertyAccessSegmentFlags::Function flag will be present and this value will be valid */
 	UPROPERTY()
-	UFunction* Function = nullptr;
+	TObjectPtr<UFunction> Function = nullptr;
 
 	/** The optional array index. */
 	UPROPERTY()
diff --git a/Engine/Source/Runtime/PropertyPath/Public/PropertyPathHelpers.h b/Engine/Source/Runtime/PropertyPath/Public/PropertyPathHelpers.h
--- a/Engine/Source/Runtime/PropertyPath/Public/PropertyPathHelpers.h
+++ b/Engine/Source/Runtime/PropertyPath/Public/PropertyPathHelpers.h
@@ -62,7 +62,7 @@ private:
 
 	/** The cached Class or ScriptStruct that was used last to resolve Name to a property. */
 	UPROPERTY(Transient)
-	mutable UStruct* Struct;
+	mutable TObjectPtr<UStruct> Struct;
 
 	/**
 	 * The cached property on the Struct that this Name resolved to on it last time Resolve was called, if 
@@ -196,7 +196,7 @@ private:
 
 	/** Cached function for function-terminated paths */
 	UPROPERTY()
-	mutable UFunction* CachedFunction;
+	mutable TObjectPtr<UFunction> CachedFunction;
 
 #if DO_CHECK
 	/** Cached container, used for checking purposes only */
diff --git a/Engine/Source/Runtime/Serialization/Private/Tests/StructSerializerTestTypes.h b/Engine/Source/Runtime/Serialization/Private/Tests/StructSerializerTestTypes.h
--- a/Engine/Source/Runtime/Serialization/Private/Tests/StructSerializerTestTypes.h
+++ b/Engine/Source/Runtime/Serialization/Private/Tests/StructSerializerTestTypes.h
@@ -131,7 +131,7 @@ struct FStructSerializerObjectTestStruct
 	GENERATED_BODY()
 
 	UPROPERTY()
-	UClass* Class;
+	TObjectPtr<UClass> Class;
 
 	UPROPERTY()
 	TSubclassOf<class UMetaData> SubClass;
@@ -140,7 +140,7 @@ struct FStructSerializerObjectTestStruct
 	TSoftClassPtr<class UMetaData> SoftClass;
 
 	UPROPERTY()
-	class UObject* Object;
+	TObjectPtr<class UObject> Object;
 
 	UPROPERTY()
 	TWeakObjectPtr<class UMetaData> WeakObject;
diff --git a/Engine/Source/Runtime/SlateCore/Public/Fonts/CompositeFont.h b/Engine/Source/Runtime/SlateCore/Public/Fonts/CompositeFont.h
--- a/Engine/Source/Runtime/SlateCore/Public/Fonts/CompositeFont.h
+++ b/Engine/Source/Runtime/SlateCore/Public/Fonts/CompositeFont.h
@@ -260,7 +260,7 @@ private:
 	 * Font data v3. This points to a font face asset.
 	 */
 	UPROPERTY()
-	const UObject* FontFaceAsset;
+	TObjectPtr<const UObject> FontFaceAsset;
 
 #if WITH_EDITORONLY_DATA
 	/**
@@ -268,7 +268,7 @@ private:
 	 * This can be removed once we no longer support loading packages older than FEditorObjectVersion::AddedFontFaceAssets (as can UFontBulkData itself).
 	 */
 	UPROPERTY()
-	UFontBulkData* BulkDataPtr_DEPRECATED;
+	TObjectPtr<UFontBulkData> BulkDataPtr_DEPRECATED;
 
 	/**
 	 * Legacy font data v1. This used to be where font data was stored prior to font bulk data.
diff --git a/Engine/Source/Runtime/SlateCore/Public/Fonts/SlateFontInfo.h b/Engine/Source/Runtime/SlateCore/Public/Fonts/SlateFontInfo.h
--- a/Engine/Source/Runtime/SlateCore/Public/Fonts/SlateFontInfo.h
+++ b/Engine/Source/Runtime/SlateCore/Public/Fonts/SlateFontInfo.h
@@ -51,7 +51,7 @@ struct SLATECORE_API FFontOutlineSettings
 
 	/** Optional material to apply to the outline */
 	UPROPERTY(EditAnywhere, BlueprintReadWrite, Category=SlateStyleRules, meta=(AllowedClasses="MaterialInterface"))
-	UObject* OutlineMaterial;
+	TObjectPtr<UObject> OutlineMaterial;
 
 	/** The color of the outline for any character in this font */
 	UPROPERTY(EditAnywhere, BlueprintReadWrite, Category=OutlineSettings)
@@ -135,11 +135,11 @@ struct SLATECORE_API FSlateFontInfo
 
 	/** The font object (valid when used from UMG or a Slate widget style asset) */
 	UPROPERTY(EditAnywhere, BlueprintReadWrite, Category=SlateStyleRules, meta=(AllowedClasses="Font", DisplayName="Font Family"))
-	const UObject* FontObject;
+	TObjectPtr<const UObject> FontObject;
 
 	/** The material to use when rendering this font */
 	UPROPERTY(EditAnywhere, BlueprintReadWrite, Category=SlateStyleRules, meta=(AllowedClasses="MaterialInterface"))
-	UObject* FontMaterial;
+	TObjectPtr<UObject> FontMaterial;
 
 	/** Settings for applying an outline to a font */
 	UPROPERTY(EditAnywhere, BlueprintReadWrite, Category=SlateStyleRules)
diff --git a/Engine/Source/Runtime/SlateCore/Public/Sound/SlateSound.h b/Engine/Source/Runtime/SlateCore/Public/Sound/SlateSound.h
--- a/Engine/Source/Runtime/SlateCore/Public/Sound/SlateSound.h
+++ b/Engine/Source/Runtime/SlateCore/Public/Sound/SlateSound.h
@@ -68,7 +68,7 @@ protected:
 	 * Edited via FSlateSoundStructCustomization to ensure you can only set USoundBase assets on it.
 	 */
 	UPROPERTY(EditAnywhere, BlueprintReadWrite, Category=Sound, meta=(DisplayName="Sound", AllowedClasses="SoundBase"))
-	UObject* ResourceObject;
+	TObjectPtr<UObject> ResourceObject;
 
 	/** The legacy resource name; only used by sounds that have been set-up in code, or otherwise upgraded from old FName properties, set to NAME_None in non-legacy instances */
 	FName LegacyResourceName_DEPRECATED;
diff --git a/Engine/Source/Runtime/SlateCore/Public/Styling/SlateBrush.h b/Engine/Source/Runtime/SlateCore/Public/Styling/SlateBrush.h
--- a/Engine/Source/Runtime/SlateCore/Public/Styling/SlateBrush.h
+++ b/Engine/Source/Runtime/SlateCore/Public/Styling/SlateBrush.h
@@ -438,7 +438,7 @@ private:
 	 * the AtlasedTextureInterface. 
 	 */
 	UPROPERTY(EditAnywhere, BlueprintReadWrite, Category=Brush, meta=( AllowPrivateAccess="true", DisplayThumbnail="true", DisplayName="Image", AllowedClasses="Texture,MaterialInterface,SlateTextureAtlasInterface", DisallowedClasses = "MediaTexture"))
-	UObject* ResourceObject;
+	TObjectPtr<UObject> ResourceObject;
 
 protected:
 	/** The name of the rendering resource to use */
diff --git a/Engine/Source/Runtime/SlateCore/Public/Styling/SlateWidgetStyleAsset.h b/Engine/Source/Runtime/SlateCore/Public/Styling/SlateWidgetStyleAsset.h
--- a/Engine/Source/Runtime/SlateCore/Public/Styling/SlateWidgetStyleAsset.h
+++ b/Engine/Source/Runtime/SlateCore/Public/Styling/SlateWidgetStyleAsset.h
@@ -21,7 +21,7 @@ class SLATECORE_API USlateWidgetStyleAsset : public UObject
 public:  
 	/**  */
 	UPROPERTY(Category=Appearance, EditAnywhere, Instanced)
-	USlateWidgetStyleContainerBase* CustomStyle;
+	TObjectPtr<USlateWidgetStyleContainerBase> CustomStyle;
 
 	template< class WidgetStyleType >            
 	const WidgetStyleType* GetStyle() const 
diff --git a/Engine/Source/Runtime/TypedElementRuntime/Public/Elements/Framework/TypedElementSelectionSet.h b/Engine/Source/Runtime/TypedElementRuntime/Public/Elements/Framework/TypedElementSelectionSet.h
--- a/Engine/Source/Runtime/TypedElementRuntime/Public/Elements/Framework/TypedElementSelectionSet.h
+++ b/Engine/Source/Runtime/TypedElementRuntime/Public/Elements/Framework/TypedElementSelectionSet.h
@@ -440,7 +440,7 @@ private:
 
 	/** Underlying element list holding the selection state. */
 	UPROPERTY()
-	UTypedElementList* ElementList = nullptr;
+	TObjectPtr<UTypedElementList> ElementList = nullptr;
 
 	/** Delegate that is invoked whenever the underlying element list is potentially about to change. */
 	FOnPreChange OnPreChangeDelegate;
diff --git a/Engine/Source/Runtime/UMG/Public/Animation/UMGSequencePlayer.h b/Engine/Source/Runtime/UMG/Public/Animation/UMGSequencePlayer.h
--- a/Engine/Source/Runtime/UMG/Public/Animation/UMGSequencePlayer.h
+++ b/Engine/Source/Runtime/UMG/Public/Animation/UMGSequencePlayer.h
@@ -101,7 +101,7 @@ private:
 
 	/** Animation being played */
 	UPROPERTY()
-	UWidgetAnimation* Animation;
+	TObjectPtr<UWidgetAnimation> Animation;
 
 	/** The user widget this sequence is animating */
 	TWeakObjectPtr<UUserWidget> UserWidget;
diff --git a/Engine/Source/Runtime/UMG/Public/Animation/UMGSequenceTickManager.h b/Engine/Source/Runtime/UMG/Public/Animation/UMGSequenceTickManager.h
--- a/Engine/Source/Runtime/UMG/Public/Animation/UMGSequenceTickManager.h
+++ b/Engine/Source/Runtime/UMG/Public/Animation/UMGSequenceTickManager.h
@@ -49,7 +49,7 @@ private:
 	TSet<TWeakObjectPtr<UUserWidget>> WeakUserWidgets;
 
 	UPROPERTY(transient)
-	UMovieSceneEntitySystemLinker* Linker;
+	TObjectPtr<UMovieSceneEntitySystemLinker> Linker;
 
 	FMovieSceneEntitySystemRunner Runner;
 
diff --git a/Engine/Source/Runtime/UMG/Public/Blueprint/DragDropOperation.h b/Engine/Source/Runtime/UMG/Public/Blueprint/DragDropOperation.h
--- a/Engine/Source/Runtime/UMG/Public/Blueprint/DragDropOperation.h
+++ b/Engine/Source/Runtime/UMG/Public/Blueprint/DragDropOperation.h
@@ -52,14 +52,14 @@ public:
 	 * were building an inventory screen this would be the UObject representing the item being moved to another slot.
 	 */
 	UPROPERTY(EditAnywhere, BlueprintReadWrite, Category="Drag and Drop", meta=( ExposeOnSpawn="true" ))
-	UObject* Payload;
+	TObjectPtr<UObject> Payload;
 	
 	/**
 	 * The Drag Visual is the widget to display when dragging the item.  Normally people create a new widget to represent the 
 	 * temporary drag.
 	 */
 	UPROPERTY(EditAnywhere, BlueprintReadOnly, Category="Drag and Drop", meta=( ExposeOnSpawn="true", DisplayName="Drag Visual" ))
-	class UWidget* DefaultDragVisual;
+	TObjectPtr<class UWidget> DefaultDragVisual;
 
 	/**
 	 * Controls where the drag widget visual will appear when dragged relative to the pointer performing
diff --git a/Engine/Source/Runtime/UMG/Public/Blueprint/UserWidget.h b/Engine/Source/Runtime/UMG/Public/Blueprint/UserWidget.h
--- a/Engine/Source/Runtime/UMG/Public/Blueprint/UserWidget.h
+++ b/Engine/Source/Runtime/UMG/Public/Blueprint/UserWidget.h
@@ -78,7 +78,7 @@ public:
 
 	/** The animation to look for. */
 	UPROPERTY()
-	UWidgetAnimation* Animation;
+	TObjectPtr<UWidgetAnimation> Animation;
 
 	/** The callback. */
 	UPROPERTY()
@@ -156,7 +156,7 @@ public:
 	FName Name;
 
 	UPROPERTY(Instanced)
-	UWidget* Content;
+	TObjectPtr<UWidget> Content;
 };
 
 class UUMGSequencePlayer;
@@ -1164,14 +1164,14 @@ public:
 
 	/** All the sequence players currently playing */
 	UPROPERTY(Transient)
-	TArray<UUMGSequencePlayer*> ActiveSequencePlayers;
+	TArray<TObjectPtr<UUMGSequencePlayer>> ActiveSequencePlayers;
 
 	UPROPERTY(Transient)
-	UUMGSequenceTickManager* AnimationTickManager;
+	TObjectPtr<UUMGSequenceTickManager> AnimationTickManager;
 
 	/** List of sequence players to cache and clean up when safe */
 	UPROPERTY(Transient)
-	TArray<UUMGSequencePlayer*> StoppedSequencePlayers;
+	TArray<TObjectPtr<UUMGSequencePlayer>> StoppedSequencePlayers;
 
 private:
 	/** Stores the widgets being assigned to named slots */
@@ -1181,7 +1181,7 @@ private:
 public:
 	/** The widget tree contained inside this user widget initialized by the blueprint */
 	UPROPERTY(Transient, DuplicateTransient, TextExportTransient)
-	UWidgetTree* WidgetTree;
+	TObjectPtr<UWidgetTree> WidgetTree;
 
 public:
 
@@ -1203,7 +1203,7 @@ public:
 	 * a texture with a screenshot of your game in it, for example if you were designing a HUD.
 	 */
 	UPROPERTY(EditDefaultsOnly, Category="Designer")
-	UTexture2D* PreviewBackground;
+	TObjectPtr<UTexture2D> PreviewBackground;
 
 #endif
 
@@ -1399,7 +1399,7 @@ private:
 
 protected:
 	UPROPERTY(Transient, DuplicateTransient)
-	class UInputComponent* InputComponent;
+	TObjectPtr<class UInputComponent> InputComponent;
 
 protected:
 	UPROPERTY(Transient, DuplicateTransient)
diff --git a/Engine/Source/Runtime/UMG/Public/Blueprint/UserWidgetPool.h b/Engine/Source/Runtime/UMG/Public/Blueprint/UserWidgetPool.h
--- a/Engine/Source/Runtime/UMG/Public/Blueprint/UserWidgetPool.h
+++ b/Engine/Source/Runtime/UMG/Public/Blueprint/UserWidgetPool.h
@@ -139,10 +139,10 @@ private:
 	}
 
 	UPROPERTY(Transient)
-	TArray<UUserWidget*> ActiveWidgets;
+	TArray<TObjectPtr<UUserWidget>> ActiveWidgets;
 	
 	UPROPERTY(Transient)
-	TArray<UUserWidget*> InactiveWidgets;
+	TArray<TObjectPtr<UUserWidget>> InactiveWidgets;
 
 	TWeakObjectPtr<UWidget> OwningWidget;
 	TWeakObjectPtr<UWorld> OwningWorld;
diff --git a/Engine/Source/Runtime/UMG/Public/Blueprint/WidgetBlueprintGeneratedClass.h b/Engine/Source/Runtime/UMG/Public/Blueprint/WidgetBlueprintGeneratedClass.h
--- a/Engine/Source/Runtime/UMG/Public/Blueprint/WidgetBlueprintGeneratedClass.h
+++ b/Engine/Source/Runtime/UMG/Public/Blueprint/WidgetBlueprintGeneratedClass.h
@@ -62,7 +62,7 @@ private:
 
 	/** A tree of the widget templates to be created */
 	UPROPERTY()
-	UWidgetTree* WidgetTree;
+	TObjectPtr<UWidgetTree> WidgetTree;
 
 #if WITH_EDITORONLY_DATA
 
@@ -84,7 +84,7 @@ public:
 	TArray< FDelegateRuntimeBinding > Bindings;
 
 	UPROPERTY()
-	TArray< UWidgetAnimation* > Animations;
+	TArray< TObjectPtr<UWidgetAnimation> > Animations;
 
 	UPROPERTY()
 	TArray< FName > NamedSlots;
diff --git a/Engine/Source/Runtime/UMG/Public/Blueprint/WidgetTree.h b/Engine/Source/Runtime/UMG/Public/Blueprint/WidgetTree.h
--- a/Engine/Source/Runtime/UMG/Public/Blueprint/WidgetTree.h
+++ b/Engine/Source/Runtime/UMG/Public/Blueprint/WidgetTree.h
@@ -108,12 +108,12 @@ public:
 public:
 	/** The root widget of the tree */
 	UPROPERTY(Instanced)
-	UWidget* RootWidget;
+	TObjectPtr<UWidget> RootWidget;
 
 protected:
 
 #if WITH_EDITORONLY_DATA
 	UPROPERTY(Instanced)
-	TArray< UWidget* > AllWidgets;
+	TArray< TObjectPtr<UWidget> > AllWidgets;
 #endif
 };
diff --git a/Engine/Source/Runtime/UMG/Public/Components/Border.h b/Engine/Source/Runtime/UMG/Public/Components/Border.h
--- a/Engine/Source/Runtime/UMG/Public/Components/Border.h
+++ b/Engine/Source/Runtime/UMG/Public/Components/Border.h
@@ -197,7 +197,7 @@ protected:
 protected:
 	/** Image to use for the border */
 	UPROPERTY()
-	USlateBrushAsset* Brush_DEPRECATED;
+	TObjectPtr<USlateBrushAsset> Brush_DEPRECATED;
 #endif
 
 	PROPERTY_BINDING_IMPLEMENTATION(FLinearColor, ContentColorAndOpacity)
diff --git a/Engine/Source/Runtime/UMG/Public/Components/Button.h b/Engine/Source/Runtime/UMG/Public/Components/Button.h
--- a/Engine/Source/Runtime/UMG/Public/Components/Button.h
+++ b/Engine/Source/Runtime/UMG/Public/Components/Button.h
@@ -35,7 +35,7 @@ class UMG_API UButton : public UContentWidget
 public:
 	/** The template style asset, used to seed the mutable instance of the style. */
 	UPROPERTY()
-	USlateWidgetStyleAsset* Style_DEPRECATED;
+	TObjectPtr<USlateWidgetStyleAsset> Style_DEPRECATED;
 
 	/** The button style used at runtime */
 	UPROPERTY(EditAnywhere, BlueprintReadWrite, Category="Appearance", meta=( DisplayName="Style" ))
diff --git a/Engine/Source/Runtime/UMG/Public/Components/CheckBox.h b/Engine/Source/Runtime/UMG/Public/Components/CheckBox.h
--- a/Engine/Source/Runtime/UMG/Public/Components/CheckBox.h
+++ b/Engine/Source/Runtime/UMG/Public/Components/CheckBox.h
@@ -47,43 +47,43 @@ public:
 
 	/** Style of the check box */
 	UPROPERTY()
-	USlateWidgetStyleAsset* Style_DEPRECATED;
+	TObjectPtr<USlateWidgetStyleAsset> Style_DEPRECATED;
 
 	/** Image to use when the checkbox is unchecked */
 	UPROPERTY()
-	USlateBrushAsset* UncheckedImage_DEPRECATED;
+	TObjectPtr<USlateBrushAsset> UncheckedImage_DEPRECATED;
 	
 	/** Image to use when the checkbox is unchecked and hovered */
 	UPROPERTY()
-	USlateBrushAsset* UncheckedHoveredImage_DEPRECATED;
+	TObjectPtr<USlateBrushAsset> UncheckedHoveredImage_DEPRECATED;
 	
 	/** Image to use when the checkbox is unchecked and pressed */
 	UPROPERTY()
-	USlateBrushAsset* UncheckedPressedImage_DEPRECATED;
+	TObjectPtr<USlateBrushAsset> UncheckedPressedImage_DEPRECATED;
 	
 	/** Image to use when the checkbox is checked */
 	UPROPERTY()
-	USlateBrushAsset* CheckedImage_DEPRECATED;
+	TObjectPtr<USlateBrushAsset> CheckedImage_DEPRECATED;
 	
 	/** Image to use when the checkbox is checked and hovered */
 	UPROPERTY()
-	USlateBrushAsset* CheckedHoveredImage_DEPRECATED;
+	TObjectPtr<USlateBrushAsset> CheckedHoveredImage_DEPRECATED;
 	
 	/** Image to use when the checkbox is checked and pressed */
 	UPROPERTY()
-	USlateBrushAsset* CheckedPressedImage_DEPRECATED;
+	TObjectPtr<USlateBrushAsset> CheckedPressedImage_DEPRECATED;
 	
 	/** Image to use when the checkbox is in an ambiguous state and hovered */
 	UPROPERTY()
-	USlateBrushAsset* UndeterminedImage_DEPRECATED;
+	TObjectPtr<USlateBrushAsset> UndeterminedImage_DEPRECATED;
 	
 	/** Image to use when the checkbox is checked and hovered */
 	UPROPERTY()
-	USlateBrushAsset* UndeterminedHoveredImage_DEPRECATED;
+	TObjectPtr<USlateBrushAsset> UndeterminedHoveredImage_DEPRECATED;
 	
 	/** Image to use when the checkbox is in an ambiguous state and pressed */
 	UPROPERTY()
-	USlateBrushAsset* UndeterminedPressedImage_DEPRECATED;
+	TObjectPtr<USlateBrushAsset> UndeterminedPressedImage_DEPRECATED;
 
 	/** How the content of the toggle button should align within the given space */
 	UPROPERTY(EditAnywhere, BlueprintReadOnly, Category=Appearance)
diff --git a/Engine/Source/Runtime/UMG/Public/Components/CircularThrobber.h b/Engine/Source/Runtime/UMG/Public/Components/CircularThrobber.h
--- a/Engine/Source/Runtime/UMG/Public/Components/CircularThrobber.h
+++ b/Engine/Source/Runtime/UMG/Public/Components/CircularThrobber.h
@@ -39,7 +39,7 @@ public:
 
 	/** Image to use for each segment of the throbber */
 	UPROPERTY()
-	USlateBrushAsset* PieceImage_DEPRECATED;
+	TObjectPtr<USlateBrushAsset> PieceImage_DEPRECATED;
 
 	UPROPERTY(EditAnywhere, BlueprintReadOnly, Category=Appearance)
 	FSlateBrush Image;
diff --git a/Engine/Source/Runtime/UMG/Public/Components/ComboBox.h b/Engine/Source/Runtime/UMG/Public/Components/ComboBox.h
--- a/Engine/Source/Runtime/UMG/Public/Components/ComboBox.h
+++ b/Engine/Source/Runtime/UMG/Public/Components/ComboBox.h
@@ -21,7 +21,7 @@ public:
 
 	/** The list of items to be displayed on the combobox. */
 	UPROPERTY(EditAnywhere, BlueprintReadWrite, Category=Content)
-	TArray<UObject*> Items;
+	TArray<TObjectPtr<UObject>> Items;
 
 	/** Called when the widget is needed for the item. */
 	UPROPERTY(EditAnywhere, Category=Events, meta=( IsBindableEvent="True" ))
diff --git a/Engine/Source/Runtime/UMG/Public/Components/EditableText.h b/Engine/Source/Runtime/UMG/Public/Components/EditableText.h
--- a/Engine/Source/Runtime/UMG/Public/Components/EditableText.h
+++ b/Engine/Source/Runtime/UMG/Public/Components/EditableText.h
@@ -55,19 +55,19 @@ public:
 
 	/** Text style */
 	UPROPERTY()
-	USlateWidgetStyleAsset* Style_DEPRECATED;
+	TObjectPtr<USlateWidgetStyleAsset> Style_DEPRECATED;
 
 	/** Background image for the selected text (overrides Style) */
 	UPROPERTY()
-	USlateBrushAsset* BackgroundImageSelected_DEPRECATED;
+	TObjectPtr<USlateBrushAsset> BackgroundImageSelected_DEPRECATED;
 
 	/** Background image for the composing text (overrides Style) */
 	UPROPERTY()
-	USlateBrushAsset* BackgroundImageComposing_DEPRECATED;
+	TObjectPtr<USlateBrushAsset> BackgroundImageComposing_DEPRECATED;
 
 	/** Image brush used for the caret (overrides Style) */
 	UPROPERTY()
-	USlateBrushAsset* CaretImage_DEPRECATED;
+	TObjectPtr<USlateBrushAsset> CaretImage_DEPRECATED;
 
 	/** Font color and opacity (overrides Style) */
 	UPROPERTY()
diff --git a/Engine/Source/Runtime/UMG/Public/Components/EditableTextBox.h b/Engine/Source/Runtime/UMG/Public/Components/EditableTextBox.h
--- a/Engine/Source/Runtime/UMG/Public/Components/EditableTextBox.h
+++ b/Engine/Source/Runtime/UMG/Public/Components/EditableTextBox.h
@@ -49,7 +49,7 @@ public:
 
 	/** Style used for the text box */
 	UPROPERTY()
-	USlateWidgetStyleAsset* Style_DEPRECATED;
+	TObjectPtr<USlateWidgetStyleAsset> Style_DEPRECATED;
 
 	/** Hint text that appears when there is no text in the text box */
 	UPROPERTY(EditAnywhere, BlueprintReadOnly, Category=Content)
diff --git a/Engine/Source/Runtime/UMG/Public/Components/ExpandableArea.h b/Engine/Source/Runtime/UMG/Public/Components/ExpandableArea.h
--- a/Engine/Source/Runtime/UMG/Public/Components/ExpandableArea.h
+++ b/Engine/Source/Runtime/UMG/Public/Components/ExpandableArea.h
@@ -98,10 +98,10 @@ protected:
 
 protected:
 	UPROPERTY()
-	UWidget* HeaderContent;
+	TObjectPtr<UWidget> HeaderContent;
 
 	UPROPERTY()
-	UWidget* BodyContent;
+	TObjectPtr<UWidget> BodyContent;
 
 	TSharedPtr<SExpandableArea> MyExpandableArea;
 };
diff --git a/Engine/Source/Runtime/UMG/Public/Components/Image.h b/Engine/Source/Runtime/UMG/Public/Components/Image.h
--- a/Engine/Source/Runtime/UMG/Public/Components/Image.h
+++ b/Engine/Source/Runtime/UMG/Public/Components/Image.h
@@ -36,7 +36,7 @@ public:
 #if WITH_EDITORONLY_DATA
 	/** Image to draw */
 	UPROPERTY()
-	USlateBrushAsset* Image_DEPRECATED;
+	TObjectPtr<USlateBrushAsset> Image_DEPRECATED;
 #endif
 
 	/** Image to draw */
diff --git a/Engine/Source/Runtime/UMG/Public/Components/ListView.h b/Engine/Source/Runtime/UMG/Public/Components/ListView.h
--- a/Engine/Source/Runtime/UMG/Public/Components/ListView.h
+++ b/Engine/Source/Runtime/UMG/Public/Components/ListView.h
@@ -181,7 +181,7 @@ protected:
 	bool bReturnFocusToSelection = false;
 
 	UPROPERTY(Transient)
-	TArray<UObject*> ListItems;
+	TArray<TObjectPtr<UObject>> ListItems;
 
 	TSharedPtr<SListView<UObject*>> MyListView;
 
diff --git a/Engine/Source/Runtime/UMG/Public/Components/PanelSlot.h b/Engine/Source/Runtime/UMG/Public/Components/PanelSlot.h
--- a/Engine/Source/Runtime/UMG/Public/Components/PanelSlot.h
+++ b/Engine/Source/Runtime/UMG/Public/Components/PanelSlot.h
@@ -16,10 +16,10 @@ class UMG_API UPanelSlot : public UVisual
 public:
 
 	UPROPERTY(Instanced)
-	class UPanelWidget* Parent;
+	TObjectPtr<class UPanelWidget> Parent;
 
 	UPROPERTY(Instanced)
-	class UWidget* Content;
+	TObjectPtr<class UWidget> Content;
 	
 #if WITH_EDITOR
 	bool IsDesignTime() const;
diff --git a/Engine/Source/Runtime/UMG/Public/Components/PanelWidget.h b/Engine/Source/Runtime/UMG/Public/Components/PanelWidget.h
--- a/Engine/Source/Runtime/UMG/Public/Components/PanelWidget.h
+++ b/Engine/Source/Runtime/UMG/Public/Components/PanelWidget.h
@@ -19,7 +19,7 @@ protected:
 
 	/** The slots in the widget holding the child widgets of this panel. */
 	UPROPERTY(Instanced)
-	TArray<UPanelSlot*> Slots;
+	TArray<TObjectPtr<UPanelSlot>> Slots;
 
 public:
 
diff --git a/Engine/Source/Runtime/UMG/Public/Components/ProgressBar.h b/Engine/Source/Runtime/UMG/Public/Components/ProgressBar.h
--- a/Engine/Source/Runtime/UMG/Public/Components/ProgressBar.h
+++ b/Engine/Source/Runtime/UMG/Public/Components/ProgressBar.h
@@ -30,19 +30,19 @@ public:
 
 	/** Style used for the progress bar */
 	UPROPERTY()
-	USlateWidgetStyleAsset* Style_DEPRECATED;
+	TObjectPtr<USlateWidgetStyleAsset> Style_DEPRECATED;
 
 	/** The brush to use as the background of the progress bar */
 	UPROPERTY()
-	USlateBrushAsset* BackgroundImage_DEPRECATED;
+	TObjectPtr<USlateBrushAsset> BackgroundImage_DEPRECATED;
 	
 	/** The brush to use as the fill image */
 	UPROPERTY()
-	USlateBrushAsset* FillImage_DEPRECATED;
+	TObjectPtr<USlateBrushAsset> FillImage_DEPRECATED;
 	
 	/** The brush to use as the marquee image */
 	UPROPERTY()
-	USlateBrushAsset* MarqueeImage_DEPRECATED;
+	TObjectPtr<USlateBrushAsset> MarqueeImage_DEPRECATED;
 
 	/** Used to determine the fill position of the progress bar ranging 0..1 */
 	UPROPERTY(EditAnywhere, BlueprintReadOnly, Category=Progress, meta=( UIMin = "0", UIMax = "1" ))
diff --git a/Engine/Source/Runtime/UMG/Public/Components/RetainerBox.h b/Engine/Source/Runtime/UMG/Public/Components/RetainerBox.h
--- a/Engine/Source/Runtime/UMG/Public/Components/RetainerBox.h
+++ b/Engine/Source/Runtime/UMG/Public/Components/RetainerBox.h
@@ -122,7 +122,7 @@ protected:
 	 * you won't see the expected color.
 	 */
 	UPROPERTY(EditAnywhere, BlueprintReadOnly, Category="Effect")
-	UMaterialInterface* EffectMaterial;
+	TObjectPtr<UMaterialInterface> EffectMaterial;
 
 	/**
 	 * The texture sampler parameter of the @EffectMaterial, that we'll set to the render target.
diff --git a/Engine/Source/Runtime/UMG/Public/Components/RichTextBlock.h b/Engine/Source/Runtime/UMG/Public/Components/RichTextBlock.h
--- a/Engine/Source/Runtime/UMG/Public/Components/RichTextBlock.h
+++ b/Engine/Source/Runtime/UMG/Public/Components/RichTextBlock.h
@@ -172,7 +172,7 @@ protected:
 
 	/**  */
 	UPROPERTY(EditAnywhere, Category=Appearance, meta=(RequiredAssetDataTags = "RowStructure=RichTextStyleRow"))
-	class UDataTable* TextStyleSet;
+	TObjectPtr<class UDataTable> TextStyleSet;
 
 	/**  */
 	UPROPERTY(EditAnywhere, Category=Appearance)
@@ -198,7 +198,7 @@ protected:
 	FTextBlockStyle DefaultTextStyle;
 
 	UPROPERTY(Transient)
-	TArray<URichTextBlockDecorator*> InstanceDecorators;
+	TArray<TObjectPtr<URichTextBlockDecorator>> InstanceDecorators;
 	TSharedPtr<class FSlateStyleSet> StyleInstance;
 	TSharedPtr<SRichTextBlock> MyRichTextBlock;
 };
diff --git a/Engine/Source/Runtime/UMG/Public/Components/RichTextBlockImageDecorator.h b/Engine/Source/Runtime/UMG/Public/Components/RichTextBlockImageDecorator.h
--- a/Engine/Source/Runtime/UMG/Public/Components/RichTextBlockImageDecorator.h
+++ b/Engine/Source/Runtime/UMG/Public/Components/RichTextBlockImageDecorator.h
@@ -52,5 +52,5 @@ protected:
 	FRichImageRow* FindImageRow(FName TagOrId, bool bWarnIfMissing);
 
 	UPROPERTY(EditAnywhere, Category=Appearance, meta = (RequiredAssetDataTags = "RowStructure=RichImageRow"))
-	class UDataTable* ImageSet;
+	TObjectPtr<class UDataTable> ImageSet;
 };
diff --git a/Engine/Source/Runtime/UMG/Public/Components/ScrollBox.h b/Engine/Source/Runtime/UMG/Public/Components/ScrollBox.h
--- a/Engine/Source/Runtime/UMG/Public/Components/ScrollBox.h
+++ b/Engine/Source/Runtime/UMG/Public/Components/ScrollBox.h
@@ -31,10 +31,10 @@ public:
 	FScrollBarStyle WidgetBarStyle;
 
 	UPROPERTY()
-	USlateWidgetStyleAsset* Style_DEPRECATED;
+	TObjectPtr<USlateWidgetStyleAsset> Style_DEPRECATED;
 
 	UPROPERTY()
-	USlateWidgetStyleAsset* BarStyle_DEPRECATED;
+	TObjectPtr<USlateWidgetStyleAsset> BarStyle_DEPRECATED;
 
 	/** The orientation of the scrolling and stacking in the box. */
 	UPROPERTY(EditAnywhere, BlueprintReadOnly, Category="Scroll")
diff --git a/Engine/Source/Runtime/UMG/Public/Components/Throbber.h b/Engine/Source/Runtime/UMG/Public/Components/Throbber.h
--- a/Engine/Source/Runtime/UMG/Public/Components/Throbber.h
+++ b/Engine/Source/Runtime/UMG/Public/Components/Throbber.h
@@ -40,7 +40,7 @@ public:
 
 	/** Image to use for each segment of the throbber */
 	UPROPERTY()
-	USlateBrushAsset* PieceImage_DEPRECATED;
+	TObjectPtr<USlateBrushAsset> PieceImage_DEPRECATED;
 
 	UPROPERTY(EditAnywhere, BlueprintReadOnly, Category=Appearance)
 	FSlateBrush Image;
diff --git a/Engine/Source/Runtime/UMG/Public/Components/Widget.h b/Engine/Source/Runtime/UMG/Public/Components/Widget.h
--- a/Engine/Source/Runtime/UMG/Public/Components/Widget.h
+++ b/Engine/Source/Runtime/UMG/Public/Components/Widget.h
@@ -250,7 +250,7 @@ public:
 	 * The parent slot of the UWidget.  Allows us to easily inline edit the layout controlling this widget.
 	 */
 	UPROPERTY(Instanced, TextExportTransient, EditAnywhere, BlueprintReadOnly, Category=Layout, meta=(ShowOnlyInnerProperties))
-	UPanelSlot* Slot;
+	TObjectPtr<UPanelSlot> Slot;
 
 	/** A bindable delegate for bIsEnabled */
 	UPROPERTY()
@@ -266,7 +266,7 @@ public:
 
 	/** Tooltip widget to show when the user hovers over the widget with the mouse */
 	UPROPERTY(VisibleAnywhere, BlueprintReadOnly, Category="Behavior", AdvancedDisplay)
-	UWidget* ToolTipWidget;
+	TObjectPtr<UWidget> ToolTipWidget;
 
 	/** A bindable delegate for ToolTipWidget */
 	UPROPERTY()
@@ -362,7 +362,7 @@ public:
 private:
 	/** A custom set of accessibility rules for this widget. If null, default rules for the widget are used. */
 	UPROPERTY(Instanced)
-	USlateAccessibleWidgetData* AccessibleWidgetData;
+	TObjectPtr<USlateAccessibleWidgetData> AccessibleWidgetData;
 
 protected:
 
@@ -420,7 +420,7 @@ public:
 	 * can occur between widgets.
 	 */
 	UPROPERTY(Instanced, EditAnywhere, BlueprintReadOnly, Category="Navigation")
-	class UWidgetNavigation* Navigation;
+	TObjectPtr<class UWidgetNavigation> Navigation;
 
 	/** Allows you to set a new flow direction */
 	UPROPERTY(EditAnywhere, Category = "Localization")
@@ -1052,7 +1052,7 @@ protected:
 
 	/** Native property bindings. */
 	UPROPERTY(Transient)
-	TArray<UPropertyBinding*> NativeBindings;
+	TArray<TObjectPtr<UPropertyBinding>> NativeBindings;
 
 	static TArray<TSubclassOf<UPropertyBinding>> BinderClasses;
 
diff --git a/Engine/Source/Runtime/UMG/Public/Components/WidgetComponent.h b/Engine/Source/Runtime/UMG/Public/Components/WidgetComponent.h
--- a/Engine/Source/Runtime/UMG/Public/Components/WidgetComponent.h
+++ b/Engine/Source/Runtime/UMG/Public/Components/WidgetComponent.h
@@ -519,7 +519,7 @@ protected:
 	 * what player's screen it appears on for split screen, if not set, users player 0.
 	 */
 	UPROPERTY()
-	ULocalPlayer* OwnerPlayer;
+	TObjectPtr<ULocalPlayer> OwnerPlayer;
 
 	/** The background color of the component */
 	UPROPERTY(EditAnywhere, Category=Rendering)
@@ -547,39 +547,39 @@ protected:
 
 	/** The body setup of the displayed quad */
 	UPROPERTY(Transient, DuplicateTransient)
-	class UBodySetup* BodySetup;
+	TObjectPtr<class UBodySetup> BodySetup;
 
 	/** The material instance for translucent widget components */
 	UPROPERTY()
-	UMaterialInterface* TranslucentMaterial;
+	TObjectPtr<UMaterialInterface> TranslucentMaterial;
 
 	/** The material instance for translucent, one-sided widget components */
 	UPROPERTY()
-	UMaterialInterface* TranslucentMaterial_OneSided;
+	TObjectPtr<UMaterialInterface> TranslucentMaterial_OneSided;
 
 	/** The material instance for opaque widget components */
 	UPROPERTY()
-	UMaterialInterface* OpaqueMaterial;
+	TObjectPtr<UMaterialInterface> OpaqueMaterial;
 
 	/** The material instance for opaque, one-sided widget components */
 	UPROPERTY()
-	UMaterialInterface* OpaqueMaterial_OneSided;
+	TObjectPtr<UMaterialInterface> OpaqueMaterial_OneSided;
 
 	/** The material instance for masked widget components. */
 	UPROPERTY()
-	UMaterialInterface* MaskedMaterial;
+	TObjectPtr<UMaterialInterface> MaskedMaterial;
 
 	/** The material instance for masked, one-sided widget components. */
 	UPROPERTY()
-	UMaterialInterface* MaskedMaterial_OneSided;
+	TObjectPtr<UMaterialInterface> MaskedMaterial_OneSided;
 
 	/** The target to which the user widget is rendered */
 	UPROPERTY(Transient, DuplicateTransient)
-	UTextureRenderTarget2D* RenderTarget;
+	TObjectPtr<UTextureRenderTarget2D> RenderTarget;
 
 	/** The dynamic instance of the material that the render target is attached to */
 	UPROPERTY(Transient, DuplicateTransient)
-	UMaterialInstanceDynamic* MaterialInstance;
+	TObjectPtr<UMaterialInstanceDynamic> MaterialInstance;
 
 	UPROPERTY(Transient, DuplicateTransient)
 	bool bAddedToScreen;
@@ -627,7 +627,7 @@ private:
 
 	/** The User Widget object displayed and managed by this component */
 	UPROPERTY(Transient, DuplicateTransient)
-	UUserWidget* Widget;
+	TObjectPtr<UUserWidget> Widget;
 
 	/** The Slate widget to be displayed by this component.  Only one of either Widget or SlateWidget can be used */
 	TSharedPtr<SWidget> SlateWidget;
diff --git a/Engine/Source/Runtime/UMG/Public/Components/WidgetInteractionComponent.h b/Engine/Source/Runtime/UMG/Public/Components/WidgetInteractionComponent.h
--- a/Engine/Source/Runtime/UMG/Public/Components/WidgetInteractionComponent.h
+++ b/Engine/Source/Runtime/UMG/Public/Components/WidgetInteractionComponent.h
@@ -337,7 +337,7 @@ protected:
 
 	/** The widget component we're currently hovering over. */
 	UPROPERTY(Transient)
-	UWidgetComponent* HoveredWidgetComponent;
+	TObjectPtr<UWidgetComponent> HoveredWidgetComponent;
 
 	/** The last hit result we used. */
 	UPROPERTY(Transient)
@@ -365,7 +365,7 @@ private:
 
 	/** The arrow component we show at editor time. */
 	UPROPERTY()
-	class UArrowComponent* ArrowComponent;
+	TObjectPtr<class UArrowComponent> ArrowComponent;
 
 #endif
 
diff --git a/Engine/Source/Runtime/UMG/Public/Slate/SlateVectorArtData.h b/Engine/Source/Runtime/UMG/Public/Slate/SlateVectorArtData.h
--- a/Engine/Source/Runtime/UMG/Public/Slate/SlateVectorArtData.h
+++ b/Engine/Source/Runtime/UMG/Public/Slate/SlateVectorArtData.h
@@ -118,11 +118,11 @@ private:
 
 	/** The mesh data asset from which the vector art is sourced */
 	UPROPERTY(EditAnywhere, Category="Vector Art" )
-	UStaticMesh* MeshAsset;
+	TObjectPtr<UStaticMesh> MeshAsset;
 
 	/** The material which we are using, or the material from with the MIC was constructed. */
 	UPROPERTY(Transient)
-	UMaterialInterface* SourceMaterial;
+	TObjectPtr<UMaterialInterface> SourceMaterial;
 #endif
 
 	/** @see GetVertexData() */
@@ -135,7 +135,7 @@ private:
 
 	/** @see GetMaterial() */
 	UPROPERTY()
-	UMaterialInterface* Material;
+	TObjectPtr<UMaterialInterface> Material;
 
 	UPROPERTY()
 	FVector2D ExtentMin;
