Fixing static analysis warnings.

diff --git a/Engine/Source/Runtime/UMG/Private/Components/WidgetComponent.cpp b/Engine/Source/Runtime/UMG/Private/Components/WidgetComponent.cpp
--- a/Engine/Source/Runtime/UMG/Private/Components/WidgetComponent.cpp
+++ b/Engine/Source/Runtime/UMG/Private/Components/WidgetComponent.cpp
@@ -1495,24 +1495,25 @@ void UWidgetComponent::InitWidget()
 	// Don't do any work if Slate is not initialized
 	if ( FSlateApplication::IsInitialized() )
 	{
-		UWorld* World = GetWorld();
-
-		if ( WidgetClass && Widget == nullptr && World && !World->bIsTearingDown)
+		if (UWorld* World = GetWorld())
 		{
-			Widget = CreateWidget(GetWorld(), WidgetClass);
-			SetTickMode(TickMode);
-		}
-		
+			if (WidgetClass && Widget == nullptr && !World->bIsTearingDown)
+			{
+				Widget = CreateWidget(World, WidgetClass);
+				SetTickMode(TickMode);
+			}
+
 #if WITH_EDITOR
-		if ( Widget && !World->IsGameWorld() && !bEditTimeUsable )
-		{
-			if( !GEnableVREditorHacks )
+			if (Widget && !World->IsGameWorld() && !bEditTimeUsable)
 			{
-				// Prevent native ticking of editor component previews
-				Widget->SetDesignerFlags(EWidgetDesignFlags::Designing);
+				if (!GEnableVREditorHacks)
+				{
+					// Prevent native ticking of editor component previews
+					Widget->SetDesignerFlags(EWidgetDesignFlags::Designing);
+				}
 			}
-		}
 #endif
+		}
 	}
 }
 
