Fix widgets with animations ticking after animations complete & before init.

diff --git a/Engine/Source/Runtime/UMG/Private/Animation/UMGSequenceTickManager.cpp b/Engine/Source/Runtime/UMG/Private/Animation/UMGSequenceTickManager.cpp
--- a/Engine/Source/Runtime/UMG/Private/Animation/UMGSequenceTickManager.cpp
+++ b/Engine/Source/Runtime/UMG/Private/Animation/UMGSequenceTickManager.cpp
@@ -142,6 +142,7 @@ void UUMGSequenceTickManager::TickWidgetAnimations(float DeltaSeconds)
 			// If this widget no longer has any animations playing, it doesn't need to be ticked any more
 			if (UserWidget->ActiveSequencePlayers.Num() == 0)
 			{
+				UserWidget->UpdateCanTick();
 				UserWidget->AnimationTickManager = nullptr;
 				WidgetIter.RemoveCurrent();
 			}
