PR #7786: Fix Comment in WidgetComponent (Contributed by [REDACTED_PERSON])

diff --git a/Engine/Source/Runtime/UMG/Public/Components/WidgetComponent.h b/Engine/Source/Runtime/UMG/Public/Components/WidgetComponent.h
--- a/Engine/Source/Runtime/UMG/Public/Components/WidgetComponent.h
+++ b/Engine/Source/Runtime/UMG/Public/Components/WidgetComponent.h
@@ -192,10 +192,7 @@ public:
 	/** Returns the window containing the user widget content */
 	TSharedPtr<SWindow> GetSlateWindow() const;
 
-	/**  
-	 *  Sets the widget to use directly. This function will keep track of the widget till the next time it's called
-	 *	with either a newer widget or a nullptr
-	 */ 
+	/** Returns the user widget displayed by this component */
 	UFUNCTION(BlueprintCallable, Category=UserInterface)
 	virtual UUserWidget* GetWidget() const;
 
