Correction to WidgetCompoenent to correctly handle Screen Space and Automatic tick mode.

diff --git a/Engine/Source/Runtime/UMG/Private/Components/WidgetComponent.cpp b/Engine/Source/Runtime/UMG/Private/Components/WidgetComponent.cpp
--- a/Engine/Source/Runtime/UMG/Private/Components/WidgetComponent.cpp
+++ b/Engine/Source/Runtime/UMG/Private/Components/WidgetComponent.cpp
@@ -978,17 +978,20 @@ void UWidgetComponent::SetTickMode(ETickMode InTickMode)
 }
 
 bool UWidgetComponent::IsWidgetVisible() const
-{	
-	if (!SlateWindow.IsValid() || !SlateWindow->GetVisibility().IsVisible())
+{
+	//  If we are not in Screen Space, if the SlateWindow is not visible the Widget is not visible.
+	if (Space != EWidgetSpace::Screen && (!SlateWindow.IsValid() || !SlateWindow->GetVisibility().IsVisible()))
 	{
 		return false;
 	}	
 	
+	// If we have a UUserWidget check its visibility
 	if (Widget)
 	{
 		return Widget->IsVisible();
 	}
 
+	// If we use a SlateWidget check its visibility
 	return SlateWidget.IsValid() && SlateWidget->GetVisibility().IsVisible();
 }
 
@@ -1153,12 +1156,22 @@ void UWidgetComponent::TickComponent(float DeltaTime, enum ELevelTick TickType,
 	{
 		UpdateWidget();
 
+		// There is no Widget set and we already rendered an empty widget. No need to continue.
 		if (Widget == nullptr && !SlateWidget.IsValid() && bRenderCleared)
 		{
-			// We will enter here if the WidgetClass is empty and we already renderered an empty widget. No need to continue.
 			return;	
 		}
+
+		// Tick Mode is Disabled, we stop here and Disable the Component Tick
+		if (TickMode == ETickMode::Disabled)
+		{
+			SetComponentTickEnabled(false);
+			return;
+		}
 		
+		UpdateWidgetOnScreen();
+
+		// We have a Widget, it's invisible and we are in automatic TickMode, we disable ticking and register a callback to know if visibility changes.
 		if (Widget && TickMode == ETickMode::Automatic && !IsWidgetVisible())
 		{
 			SetComponentTickEnabled(false);
@@ -1170,11 +1183,6 @@ void UWidgetComponent::TickComponent(float DeltaTime, enum ELevelTick TickType,
 			return;
 		}
 
-		if (TickMode == ETickMode::Disabled)
-		{ 
-			SetComponentTickEnabled(false);
-			return;
-		}
 
 	    if ( Space != EWidgetSpace::Screen )
 	    {
@@ -1193,63 +1201,61 @@ void UWidgetComponent::TickComponent(float DeltaTime, enum ELevelTick TickType,
 				}
 		    }
 	    }
-	    else
-	    {
-			if ( ( Widget && !Widget->IsDesignTime() ) || SlateWidget.IsValid() )
-		    {
-				UWorld* ThisWorld = GetWorld();
+	    
+	}
+#endif // !UE_SERVER
+}
 
-				ULocalPlayer* TargetPlayer = GetOwnerPlayer();
-				APlayerController* PlayerController = TargetPlayer ? ToRawPtr(TargetPlayer->PlayerController) : nullptr;
+void UWidgetComponent::UpdateWidgetOnScreen()
+{	
+	if (Space == EWidgetSpace::Screen && !bAddedToScreen)
+	{
+		ULocalPlayer* TargetPlayer = GetOwnerPlayer();
+		APlayerController* PlayerController = TargetPlayer ? ToRawPtr(TargetPlayer->PlayerController) : nullptr;
+		if (TargetPlayer && PlayerController && IsVisible() && !(GetOwner()->IsHidden()))
+		{
+			AddWidgetToScreen(TargetPlayer);
+		}
+	}
+	else if(bAddedToScreen && Space == EWidgetSpace::World)
+	{
+		RemoveWidgetFromScreen();
+	}
+}
 
-				if ( TargetPlayer && PlayerController && IsVisible() && !(GetOwner()->IsHidden()))
-				{
-					if ( !bAddedToScreen )
-					{
-						if ( ThisWorld->IsGameWorld() )
-						{
-							if ( UGameViewportClient* ViewportClient = ThisWorld->GetGameViewport() )
-							{
-								TSharedPtr<IGameLayerManager> LayerManager = ViewportClient->GetGameLayerManager();
-								if ( LayerManager.IsValid() )
-								{
-									TSharedPtr<FWorldWidgetScreenLayer> ScreenLayer;
+void UWidgetComponent::AddWidgetToScreen(ULocalPlayer* TargetPlayer)
+{	
+	UWorld* ThisWorld = GetWorld();
+	if (ThisWorld->IsGameWorld())
+	{
+		if (UGameViewportClient* ViewportClient = ThisWorld->GetGameViewport())
+		{
+			TSharedPtr<IGameLayerManager> LayerManager = ViewportClient->GetGameLayerManager();
+			if (LayerManager.IsValid())
+			{
+				TSharedPtr<FWorldWidgetScreenLayer> ScreenLayer;
 
-									FLocalPlayerContext PlayerContext(TargetPlayer, ThisWorld);
+				FLocalPlayerContext PlayerContext(TargetPlayer, ThisWorld);
 
-									TSharedPtr<IGameLayer> Layer = LayerManager->FindLayerForPlayer(TargetPlayer, SharedLayerName);
-									if ( !Layer.IsValid() )
-									{
-										TSharedRef<FWorldWidgetScreenLayer> NewScreenLayer = MakeShareable(new FWorldWidgetScreenLayer(PlayerContext));
-										LayerManager->AddLayerForPlayer(TargetPlayer, SharedLayerName, NewScreenLayer, LayerZOrder);
-										ScreenLayer = NewScreenLayer;
-									}
-									else
-									{
-										ScreenLayer = StaticCastSharedPtr<FWorldWidgetScreenLayer>(Layer);
-									}
-								
-									bAddedToScreen = true;
-								
-									if (Widget && Widget->IsValidLowLevel())
-									{
-										Widget->SetPlayerContext(PlayerContext);
-									}
-									
-									ScreenLayer->AddComponent(this);
-								}
-							}
-						}
-					}
+				TSharedPtr<IGameLayer> Layer = LayerManager->FindLayerForPlayer(TargetPlayer, SharedLayerName);
+				if (!Layer.IsValid())
+				{
+					TSharedRef<FWorldWidgetScreenLayer> NewScreenLayer = MakeShareable(new FWorldWidgetScreenLayer(PlayerContext));
+					LayerManager->AddLayerForPlayer(TargetPlayer, SharedLayerName, NewScreenLayer, LayerZOrder);
+					ScreenLayer = NewScreenLayer;
 				}
-				else if ( bAddedToScreen )
+				else
 				{
-					RemoveWidgetFromScreen();
+					ScreenLayer = StaticCastSharedPtr<FWorldWidgetScreenLayer>(Layer);
 				}
+
+				bAddedToScreen = true;
+
+
+				ScreenLayer->AddComponent(this);
 			}
 		}
 	}
-#endif // !UE_SERVER
 }
 
 bool UWidgetComponent::ShouldDrawWidget() const
@@ -1655,17 +1661,27 @@ void UWidgetComponent::UpdateWidget()
 	// Don't do any work if Slate is not initialized
 	if (FSlateApplication::IsInitialized() && !IsPendingKill())
 	{
-		if ( Space != EWidgetSpace::Screen )
+		// Look for a UMG widget set
+		TSharedPtr<SWidget> NewSlateWidget;
+		if (Widget)
 		{
-			// Look for a UMG widget set
-			TSharedPtr<SWidget> NewSlateWidget;
-			if (Widget)
+			NewSlateWidget = Widget->TakeWidget();
+
+			if (Space == EWidgetSpace::Screen && Widget->IsValidLowLevel())
 			{
-				NewSlateWidget = Widget->TakeWidget();
+				UWorld* ThisWorld = GetWorld();
+				ULocalPlayer* TargetPlayer = GetOwnerPlayer();
+				FLocalPlayerContext PlayerContext(TargetPlayer, ThisWorld);
+				Widget->SetPlayerContext(PlayerContext);
 			}
 
-			// Create the SlateWindow if it doesn't exists
-			bool bNeededNewWindow = false;
+		}
+
+		// Create the SlateWindow if it doesn't exists
+		bool bNeededNewWindow = false;
+
+		if (Space != EWidgetSpace::Screen)
+		{
 			if (!SlateWindow.IsValid())
 			{
 				UpdateMaterialInstance();
@@ -1679,52 +1695,62 @@ void UWidgetComponent::UpdateWidget()
 			}
 
 			SlateWindow->Resize(CurrentDrawSize);
+		}
+		else
+		{
+			bRenderCleared = false;
+			UnregisterWindow();
+		}
 
-			// Add the UMG or SlateWidget to the Component
-			bool bWidgetChanged = false;
-			
-			// We Get here if we have a UMG Widget
-			if (NewSlateWidget.IsValid())
+		// Add the UMG or SlateWidget to the Component
+		bool bWidgetChanged = false;
+
+		// We Get here if we have a UMG Widget
+		if (NewSlateWidget.IsValid())
+		{
+			if (NewSlateWidget != CurrentSlateWidget)
 			{
-				if (NewSlateWidget != CurrentSlateWidget || bNeededNewWindow)
+				CurrentSlateWidget = NewSlateWidget;
+				if (SlateWindow.IsValid() && bNeededNewWindow)
 				{
-					CurrentSlateWidget = NewSlateWidget;
 					SlateWindow->SetContent(NewSlateWidget.ToSharedRef());
-					bRenderCleared = false;
-					bWidgetChanged = true;
 				}
+				bRenderCleared = false;
+				bWidgetChanged = true;
 			}
-			// If we don't have one, we look for a Slate Widget
-			else if (SlateWidget.IsValid())
+		}
+		// If we don't have one, we look for a Slate Widget
+		else if (SlateWidget.IsValid())
+		{
+			if (SlateWidget != CurrentSlateWidget || bNeededNewWindow)
 			{
-				if (SlateWidget != CurrentSlateWidget || bNeededNewWindow)
+				CurrentSlateWidget = SlateWidget;
+				if (SlateWindow.IsValid())
 				{
-					CurrentSlateWidget = SlateWidget;
 					SlateWindow->SetContent(SlateWidget.ToSharedRef());
-					bRenderCleared = false;
-					bWidgetChanged = true;
 				}
+				bRenderCleared = false;
+				bWidgetChanged = true;
 			}
-			else
+		}
+		else
+		{
+			if (CurrentSlateWidget != SNullWidget::NullWidget)
 			{
-				if (CurrentSlateWidget != SNullWidget::NullWidget)
-				{
-					CurrentSlateWidget = SNullWidget::NullWidget;
-					bRenderCleared = false;
-					bWidgetChanged = true;
-				}
-				SlateWindow->SetContent(SNullWidget::NullWidget);
+				CurrentSlateWidget = SNullWidget::NullWidget;
+				bRenderCleared = false;
+				bWidgetChanged = true;
 			}
-		
-			if (bNeededNewWindow || bWidgetChanged)
+			if (SlateWindow.IsValid())
 			{
-				MarkRenderStateDirty();
-				SetComponentTickEnabled(true);
+				SlateWindow->SetContent(SNullWidget::NullWidget);
 			}
 		}
-		else
+
+		if (SlateWindow.IsValid() && (bNeededNewWindow || bWidgetChanged))
 		{
-			UnregisterWindow();
+			MarkRenderStateDirty();
+			SetComponentTickEnabled(true);
 		}
 	}
 }
diff --git a/Engine/Source/Runtime/UMG/Public/Components/WidgetComponent.h b/Engine/Source/Runtime/UMG/Public/Components/WidgetComponent.h
--- a/Engine/Source/Runtime/UMG/Public/Components/WidgetComponent.h
+++ b/Engine/Source/Runtime/UMG/Public/Components/WidgetComponent.h
@@ -639,6 +639,9 @@ private:
 	/** The slate widget currently being drawn. */
 	TWeakPtr<SWidget> CurrentSlateWidget;
 
+	void UpdateWidgetOnScreen();
+	void AddWidgetToScreen(ULocalPlayer* TargetPlayer);
+
 	static EVisibility ConvertWindowVisibilityToVisibility(EWindowVisibility visibility);
 
 	void OnWidgetVisibilityChanged(ESlateVisibility InVisibility);
