Avoid usage of FRawPtrProxy implementation detail.

diff --git a/Engine/Source/Runtime/UMG/Private/Components/RichTextBlock.cpp b/Engine/Source/Runtime/UMG/Private/Components/RichTextBlock.cpp
--- a/Engine/Source/Runtime/UMG/Private/Components/RichTextBlock.cpp
+++ b/Engine/Source/Runtime/UMG/Private/Components/RichTextBlock.cpp
@@ -38,7 +38,7 @@ private:
 };
 
 template< class ObjectType >
-FORCEINLINE SharedPointerInternals::FRawPtrProxy< ObjectType > MakeShareableDeferredCleanup(ObjectType* InObject)
+FORCEINLINE TSharedPtr< ObjectType > MakeShareableDeferredCleanup(ObjectType* InObject)
 {
 	return MakeShareable(InObject, [](ObjectType* ObjectToDelete) { BeginCleanup(new FDeferredDeletor<ObjectType>(ObjectToDelete)); });
 }
