Slate: Remove temporary variable in STextBlock

diff --git a/Engine/Source/Runtime/UMG/Private/Components/TextBlock.cpp b/Engine/Source/Runtime/UMG/Private/Components/TextBlock.cpp
--- a/Engine/Source/Runtime/UMG/Private/Components/TextBlock.cpp
+++ b/Engine/Source/Runtime/UMG/Private/Components/TextBlock.cpp
@@ -312,8 +312,7 @@ void UTextBlock::SetText(FText InText)
 	TextDelegate.Unbind();
 	if ( MyTextBlock.IsValid() )
 	{
-		TAttribute<FText> TextBinding = GetDisplayText();
-		MyTextBlock->SetText(TextBinding);
+		MyTextBlock->SetText(GetDisplayText());
 	}
 }
 
