Slate: Make the Invalidation only Layout Volatile (remove Paint Volatile).

diff --git a/Engine/Source/Runtime/Slate/Public/Widgets/SInvalidationPanel.h b/Engine/Source/Runtime/Slate/Public/Widgets/SInvalidationPanel.h
--- a/Engine/Source/Runtime/Slate/Public/Widgets/SInvalidationPanel.h
+++ b/Engine/Source/Runtime/Slate/Public/Widgets/SInvalidationPanel.h
@@ -56,7 +56,6 @@ public:
 
 protected:
 	virtual bool CustomPrepass(float LayoutScaleMultiplier) override;
-	virtual bool ComputeVolatility() const { return GetCanCache() || SCompoundWidget::ComputeVolatility(); }
 	virtual bool Advanced_IsInvalidationRoot() const override;
 	virtual const FSlateInvalidationRoot* Advanced_AsInvalidationRoot() const override;
 	virtual TSharedRef<SWidget> GetRootWidget() override;
diff --git a/Engine/Source/Runtime/UMG/Public/Slate/SRetainerWidget.h b/Engine/Source/Runtime/UMG/Public/Slate/SRetainerWidget.h
--- a/Engine/Source/Runtime/UMG/Public/Slate/SRetainerWidget.h
+++ b/Engine/Source/Runtime/UMG/Public/Slate/SRetainerWidget.h
@@ -93,7 +93,6 @@ protected:
 	/** SCompoundWidget interface */
 	virtual int32 OnPaint(const FPaintArgs& Args, const FGeometry& AllottedGeometry, const FSlateRect& MyCullingRect, FSlateWindowElementList& OutDrawElements, int32 LayerId, const FWidgetStyle& InWidgetStyle, bool bParentEnabled) const override;
 	virtual FVector2D ComputeDesiredSize(float Scale) const override;
-	virtual bool ComputeVolatility() const { return bEnableRetainedRendering || SCompoundWidget::ComputeVolatility(); }
 	virtual bool Advanced_IsInvalidationRoot() const { return bEnableRetainedRendering; }
 	virtual const FSlateInvalidationRoot* Advanced_AsInvalidationRoot() const override { return bEnableRetainedRendering ? this : nullptr; }
 	virtual bool CustomPrepass(float LayoutScaleMultiplier) override;
