Change to EditorUtilityWidgetBlueprint to use the Transient package as it's Outer. This prevents the World from being marked dirty when opening or changing a property on the EditorUtilityWidget object.

diff --git a/Engine/Source/Editor/Blutility/Private/EditorUtilityWidgetBlueprint.cpp b/Engine/Source/Editor/Blutility/Private/EditorUtilityWidgetBlueprint.cpp
--- a/Engine/Source/Editor/Blutility/Private/EditorUtilityWidgetBlueprint.cpp
+++ b/Engine/Source/Editor/Blutility/Private/EditorUtilityWidgetBlueprint.cpp
@@ -77,7 +77,7 @@ TSharedRef<SWidget> UEditorUtilityWidgetBlueprint::CreateUtilityWidget()
 		{
 			CreatedUMGWidget->Rename(nullptr, GetTransientPackage(), REN_DoNotDirty);
 		}
-		CreatedUMGWidget = CreateWidget<UEditorUtilityWidget>(World, WidgetClass);
+		CreatedUMGWidget = CreateWidget<UEditorUtilityWidget>(GetTransientPackage(), WidgetClass);
 	}
 
 	if (CreatedUMGWidget)
diff --git a/Engine/Source/Runtime/UMG/Private/UserWidget.cpp b/Engine/Source/Runtime/UMG/Private/UserWidget.cpp
--- a/Engine/Source/Runtime/UMG/Private/UserWidget.cpp
+++ b/Engine/Source/Runtime/UMG/Private/UserWidget.cpp
@@ -1971,6 +1971,11 @@ UUserWidget* UUserWidget::CreateWidgetInstance(UWorld& World, TSubclassOf<UUserW
 	return CreateInstanceInternal(&World, UserWidgetClass, WidgetName, &World, World.GetFirstLocalPlayerFromController());
 }
 
+UUserWidget* UUserWidget::CreateWidgetInstance(UPackage& Package, TSubclassOf<UUserWidget> UserWidgetClass, FName WidgetName)
+{
+	return CreateInstanceInternal(&Package, UserWidgetClass, WidgetName, nullptr, nullptr);
+}
+
 UUserWidget* UUserWidget::CreateInstanceInternal(UObject* Outer, TSubclassOf<UUserWidget> UserWidgetClass, FName InstanceName, UWorld* World, ULocalPlayer* LocalPlayer)
 {
 	//CSV_SCOPED_TIMING_STAT(Slate, CreateWidget);
diff --git a/Engine/Source/Runtime/UMG/Public/Blueprint/UserWidget.h b/Engine/Source/Runtime/UMG/Public/Blueprint/UserWidget.h
--- a/Engine/Source/Runtime/UMG/Public/Blueprint/UserWidget.h
+++ b/Engine/Source/Runtime/UMG/Public/Blueprint/UserWidget.h
@@ -1135,6 +1135,7 @@ public:
 	static UUserWidget* CreateWidgetInstance(APlayerController& OwnerPC, TSubclassOf<UUserWidget> UserWidgetClass, FName WidgetName);
 	static UUserWidget* CreateWidgetInstance(UGameInstance& GameInstance, TSubclassOf<UUserWidget> UserWidgetClass, FName WidgetName);
 	static UUserWidget* CreateWidgetInstance(UWorld& World, TSubclassOf<UUserWidget> UserWidgetClass, FName WidgetName);
+	static UUserWidget* CreateWidgetInstance(UPackage& Package, TSubclassOf<UUserWidget> UserWidgetClass, FName WidgetName);
 
 private:
 	static UUserWidget* CreateInstanceInternal(UObject* Outer, TSubclassOf<UUserWidget> UserWidgetClass, FName WidgetName, UWorld* World, ULocalPlayer* LocalPlayer);
@@ -1471,7 +1472,9 @@ WidgetT* CreateWidget(OwnerT* OwningObject, TSubclassOf<UUserWidget> UserWidgetC
 		|| TIsDerivedFrom<OwnerT, UWidgetTree>::IsDerived
 		|| TIsDerivedFrom<OwnerT, APlayerController>::IsDerived
 		|| TIsDerivedFrom<OwnerT, UGameInstance>::IsDerived
-		|| TIsDerivedFrom<OwnerT, UWorld>::IsDerived, "The given OwningObject is not of a supported type for use with CreateWidget.");
+		|| TIsDerivedFrom<OwnerT, UWorld>::IsDerived 
+		|| TIsDerivedFrom<OwnerT, UPackage>::IsDerived, // We want to support using the Transient package for a Editor Utility Widget.
+		"The given OwningObject is not of a supported type for use with CreateWidget.");
 
 	SCOPE_CYCLE_COUNTER(STAT_CreateWidget);
 
