PR #8393: URichTextBlock::RefreshTextLayout - new method to support refreshing text layout on demand (Contributed by [REDACTED_PERSON])

diff --git a/Engine/Source/Runtime/UMG/Private/Components/RichTextBlock.cpp b/Engine/Source/Runtime/UMG/Private/Components/RichTextBlock.cpp
--- a/Engine/Source/Runtime/UMG/Private/Components/RichTextBlock.cpp
+++ b/Engine/Source/Runtime/UMG/Private/Components/RichTextBlock.cpp
@@ -357,6 +357,14 @@ void URichTextBlock::ApplyUpdatedDefaultTextStyle()
 	}
 }
 
+void URichTextBlock::RefreshTextLayout()
+{
+	if (MyRichTextBlock.IsValid())
+	{
+		MyRichTextBlock->Refresh();
+	}
+}
+
 /////////////////////////////////////////////////////
 
 #undef LOCTEXT_NAMESPACE
diff --git a/Engine/Source/Runtime/UMG/Public/Components/RichTextBlock.h b/Engine/Source/Runtime/UMG/Public/Components/RichTextBlock.h
--- a/Engine/Source/Runtime/UMG/Public/Components/RichTextBlock.h
+++ b/Engine/Source/Runtime/UMG/Public/Components/RichTextBlock.h
@@ -153,6 +153,12 @@ public:
 	UFUNCTION(BlueprintCallable, Category = "Widget")
 	URichTextBlockDecorator* GetDecoratorByClass(TSubclassOf<URichTextBlockDecorator> DecoratorClass);
 
+	/**
+	 * Causes the text to reflow it's layout and re-evaluate any decorators
+	 */
+	UFUNCTION(BlueprintCallable, Category = "Widget")
+	void RefreshTextLayout();
+
 protected:
 	virtual TSharedRef<SWidget> RebuildWidget() override;
 	
