Fix foreground color spinbox

diff --git a/Engine/Source/Runtime/UMG/Private/Components/SpinBox.cpp b/Engine/Source/Runtime/UMG/Private/Components/SpinBox.cpp
--- a/Engine/Source/Runtime/UMG/Private/Components/SpinBox.cpp
+++ b/Engine/Source/Runtime/UMG/Private/Components/SpinBox.cpp
@@ -38,7 +38,6 @@ USpinBox::USpinBox(const FObjectInitializer& ObjectInitializer)
 	MinDesiredWidth = 0;
 	ClearKeyboardFocusOnCommit = false;
 	SelectAllTextOnCommit = true;
-	ForegroundColor = FSlateColor(FLinearColor::Black);
 
 	if (DefaultSpinBoxStyle == nullptr)
 	{
@@ -65,6 +64,7 @@ USpinBox::USpinBox(const FObjectInitializer& ObjectInitializer)
 	}
 #endif // WITH_EDITOR
 
+	ForegroundColor = WidgetStyle.ForegroundColor;
 }
 
 void USpinBox::ReleaseSlateResources(bool bReleaseChildren)
