Remove unnecessary error warning firing when dragging in widget archetypes derived from C++ UserWidgets.
Warning no longer relevant with changes to widget template generation.

diff --git a/Engine/Source/Runtime/UMG/Private/WidgetBlueprintGeneratedClass.cpp b/Engine/Source/Runtime/UMG/Private/WidgetBlueprintGeneratedClass.cpp
--- a/Engine/Source/Runtime/UMG/Private/WidgetBlueprintGeneratedClass.cpp
+++ b/Engine/Source/Runtime/UMG/Private/WidgetBlueprintGeneratedClass.cpp
@@ -189,7 +189,6 @@ void UWidgetBlueprintGeneratedClass::InitializeWidgetStatic(UUserWidget* UserWid
 
 	if ( UserWidget->IsTemplate() )
 	{
-		UE_LOG(LogUMG, Error, TEXT("Widget Class %s - Running Initialize On Archetype, %s."), *InClass->GetName(), *UserWidget->GetName());
 		return;
 	}
 
