Fix checkbox style on editor utility widgets.

diff --git a/Engine/Source/Editor/Blutility/Classes/EditorUtilityWidget.h b/Engine/Source/Editor/Blutility/Classes/EditorUtilityWidget.h
--- a/Engine/Source/Editor/Blutility/Classes/EditorUtilityWidget.h
+++ b/Engine/Source/Editor/Blutility/Classes/EditorUtilityWidget.h
@@ -39,6 +39,8 @@ public:
 		return bAutoRunDefaultAction;
 	}
 
+	virtual bool IsEditorUtility() const override { return true; }
+
 protected:
 	UPROPERTY(Category = Config, EditDefaultsOnly, BlueprintReadWrite, AssetRegistrySearchable)
 	FString HelpText;
diff --git a/Engine/Source/Runtime/SlateCore/Public/Styling/SlateTypes.h b/Engine/Source/Runtime/SlateCore/Public/Styling/SlateTypes.h
--- a/Engine/Source/Runtime/SlateCore/Public/Styling/SlateTypes.h
+++ b/Engine/Source/Runtime/SlateCore/Public/Styling/SlateTypes.h
@@ -262,7 +262,17 @@ struct SLATECORE_API FCheckBoxStyle : public FSlateWidgetStyle
 		UndeterminedImage.UnlinkColors();
 		UndeterminedHoveredImage.UnlinkColors();
 		UndeterminedPressedImage.UnlinkColors();
+		BackgroundImage.UnlinkColors();
+		BackgroundHoveredImage.UnlinkColors();
+		BackgroundPressedImage.UnlinkColors();
+
 		ForegroundColor.Unlink();
+		HoveredForeground.Unlink();
+		PressedForeground.Unlink();
+		CheckedForeground.Unlink();
+		CheckedHoveredForeground.Unlink();
+		CheckedPressedForeground.Unlink();
+		UndeterminedForeground.Unlink();
 		BorderBackgroundColor.Unlink();
 	}
 };
diff --git a/Engine/Source/Runtime/UMG/Private/Components/Button.cpp b/Engine/Source/Runtime/UMG/Private/Components/Button.cpp
--- a/Engine/Source/Runtime/UMG/Private/Components/Button.cpp
+++ b/Engine/Source/Runtime/UMG/Private/Components/Button.cpp
@@ -44,6 +44,9 @@ UButton::UButton(const FObjectInitializer& ObjectInitializer)
 	if (IsEditorWidget())
 	{
 		WidgetStyle = *EditorButtonStyle;
+
+		// The CDO isn't an editor widget and thus won't use the editor style, call post edit change to mark difference from CDO
+		PostEditChange();
 	}
 #endif // WITH_EDITOR
 
diff --git a/Engine/Source/Runtime/UMG/Private/Components/CheckBox.cpp b/Engine/Source/Runtime/UMG/Private/Components/CheckBox.cpp
--- a/Engine/Source/Runtime/UMG/Private/Components/CheckBox.cpp
+++ b/Engine/Source/Runtime/UMG/Private/Components/CheckBox.cpp
@@ -43,6 +43,9 @@ UCheckBox::UCheckBox(const FObjectInitializer& ObjectInitializer)
 	if (IsEditorWidget())
 	{
 		WidgetStyle = *EditorCheckboxStyle;
+
+		// The CDO isn't an editor widget and thus won't use the editor style, call post edit change to mark difference from CDO
+		PostEditChange();
 	}
 #endif // WITH_EDITOR
 
diff --git a/Engine/Source/Runtime/UMG/Private/Components/CircularThrobber.cpp b/Engine/Source/Runtime/UMG/Private/Components/CircularThrobber.cpp
--- a/Engine/Source/Runtime/UMG/Private/Components/CircularThrobber.cpp
+++ b/Engine/Source/Runtime/UMG/Private/Components/CircularThrobber.cpp
@@ -44,6 +44,9 @@ UCircularThrobber::UCircularThrobber(const FObjectInitializer& ObjectInitializer
 	if (IsEditorWidget())
 	{
 		Image = *EditorCircularThrobberBrushStyle;
+		
+		// The CDO isn't an editor widget and thus won't use the editor style, call post edit change to mark difference from CDO
+		PostEditChange();
 	}
 #endif // WITH_EDITOR
 
diff --git a/Engine/Source/Runtime/UMG/Private/Components/ComboBoxKey.cpp b/Engine/Source/Runtime/UMG/Private/Components/ComboBoxKey.cpp
--- a/Engine/Source/Runtime/UMG/Private/Components/ComboBoxKey.cpp
+++ b/Engine/Source/Runtime/UMG/Private/Components/ComboBoxKey.cpp
@@ -22,6 +22,12 @@ UComboBoxKey::UComboBoxKey()
 		const ISlateStyle& SlateStyle = IsEditorWidget() ? FCoreStyle::Get() : FUMGCoreStyle::Get();
 		WidgetStyle = SlateStyle.GetWidgetStyle<FComboBoxStyle>("ComboBox");
 		ItemStyle = SlateStyle.GetWidgetStyle<FTableRowStyle>("TableView.Row");
+
+		if (IsEditorWidget())
+		{
+			// The CDO isn't an editor widget and thus won't use the editor style, call post edit change to mark difference from CDO
+			PostEditChange();
+		}
 #else
 		WidgetStyle = FUMGCoreStyle::Get().GetWidgetStyle<FComboBoxStyle>("ComboBox");
 		ItemStyle = FUMGCoreStyle::Get().GetWidgetStyle<FTableRowStyle>("TableView.Row");
diff --git a/Engine/Source/Runtime/UMG/Private/Components/ComboBoxString.cpp b/Engine/Source/Runtime/UMG/Private/Components/ComboBoxString.cpp
--- a/Engine/Source/Runtime/UMG/Private/Components/ComboBoxString.cpp
+++ b/Engine/Source/Runtime/UMG/Private/Components/ComboBoxString.cpp
@@ -63,6 +63,9 @@ UComboBoxString::UComboBoxString(const FObjectInitializer& ObjectInitializer)
 	{
 		WidgetStyle = *EditorComboBoxStyle;
 		ItemStyle = *EditorComboBoxRowStyle;
+
+		// The CDO isn't an editor widget and thus won't use the editor style, call post edit change to mark difference from CDO
+		PostEditChange();
 	}
 #endif // WITH_EDITOR
 
diff --git a/Engine/Source/Runtime/UMG/Private/Components/EditableText.cpp b/Engine/Source/Runtime/UMG/Private/Components/EditableText.cpp
--- a/Engine/Source/Runtime/UMG/Private/Components/EditableText.cpp
+++ b/Engine/Source/Runtime/UMG/Private/Components/EditableText.cpp
@@ -44,6 +44,9 @@ UEditableText::UEditableText(const FObjectInitializer& ObjectInitializer)
 	if (IsEditorWidget())
 	{
 		WidgetStyle = *EditorEditableTextStyle;
+
+		// The CDO isn't an editor widget and thus won't use the editor style, call post edit change to mark difference from CDO
+		PostEditChange();
 	}
 #endif // WITH_EDITOR
 
diff --git a/Engine/Source/Runtime/UMG/Private/Components/EditableTextBox.cpp b/Engine/Source/Runtime/UMG/Private/Components/EditableTextBox.cpp
--- a/Engine/Source/Runtime/UMG/Private/Components/EditableTextBox.cpp
+++ b/Engine/Source/Runtime/UMG/Private/Components/EditableTextBox.cpp
@@ -66,6 +66,9 @@ UEditableTextBox::UEditableTextBox(const FObjectInitializer& ObjectInitializer)
 	if (IsEditorWidget())
 	{
 		WidgetStyle = *EditorEditableTextBoxStyle;
+
+		// The CDO isn't an editor widget and thus won't use the editor style, call post edit change to mark difference from CDO
+		PostEditChange();
 	}
 #endif // WITH_EDITOR
 
diff --git a/Engine/Source/Runtime/UMG/Private/Components/ExpandableArea.cpp b/Engine/Source/Runtime/UMG/Private/Components/ExpandableArea.cpp
--- a/Engine/Source/Runtime/UMG/Private/Components/ExpandableArea.cpp
+++ b/Engine/Source/Runtime/UMG/Private/Components/ExpandableArea.cpp
@@ -69,6 +69,9 @@ UExpandableArea::UExpandableArea(const FObjectInitializer& ObjectInitializer)
 	{
 		Style = *EditorExpandableAreaStyle;
 		BorderBrush = *EditorExpandableAreaBorderBrush;
+
+		// The CDO isn't an editor widget and thus won't use the editor style, call post edit change to mark difference from CDO
+		PostEditChange();
 	}
 #endif // WITH_EDITOR
 
diff --git a/Engine/Source/Runtime/UMG/Private/Components/InputKeySelector.cpp b/Engine/Source/Runtime/UMG/Private/Components/InputKeySelector.cpp
--- a/Engine/Source/Runtime/UMG/Private/Components/InputKeySelector.cpp
+++ b/Engine/Source/Runtime/UMG/Private/Components/InputKeySelector.cpp
@@ -62,6 +62,9 @@ UInputKeySelector::UInputKeySelector( const FObjectInitializer& ObjectInitialize
 	{
 		WidgetStyle = *EditorInputKeySelectorButtonStyle;
 		TextStyle = *EditorInputKeySelectorTextStyle;
+
+		// The CDO isn't an editor widget and thus won't use the editor style, call post edit change to mark difference from CDO
+		PostEditChange();
 	}
 #endif // WITH_EDITOR
 
diff --git a/Engine/Source/Runtime/UMG/Private/Components/MultiLineEditableText.cpp b/Engine/Source/Runtime/UMG/Private/Components/MultiLineEditableText.cpp
--- a/Engine/Source/Runtime/UMG/Private/Components/MultiLineEditableText.cpp
+++ b/Engine/Source/Runtime/UMG/Private/Components/MultiLineEditableText.cpp
@@ -43,6 +43,9 @@ UMultiLineEditableText::UMultiLineEditableText(const FObjectInitializer& ObjectI
 	if (IsEditorWidget())
 	{
 		WidgetStyle = *EditorMultiLineEditableTextStyle;
+
+		// The CDO isn't an editor widget and thus won't use the editor style, call post edit change to mark difference from CDO
+		PostEditChange();
 	}
 #endif // WITH_EDITOR
 
diff --git a/Engine/Source/Runtime/UMG/Private/Components/MultiLineEditableTextBox.cpp b/Engine/Source/Runtime/UMG/Private/Components/MultiLineEditableTextBox.cpp
--- a/Engine/Source/Runtime/UMG/Private/Components/MultiLineEditableTextBox.cpp
+++ b/Engine/Source/Runtime/UMG/Private/Components/MultiLineEditableTextBox.cpp
@@ -67,6 +67,9 @@ UMultiLineEditableTextBox::UMultiLineEditableTextBox(const FObjectInitializer& O
 	{
 		WidgetStyle = *EditorMultiLineEditableTextBoxStyle;
 		TextStyle = *EditorMultiLineEditableTextBoxTextStyle;
+
+		// The CDO isn't an editor widget and thus won't use the editor style, call post edit change to mark difference from CDO
+		PostEditChange();
 	}
 #endif // WITH_EDITOR
 
diff --git a/Engine/Source/Runtime/UMG/Private/Components/ProgressBar.cpp b/Engine/Source/Runtime/UMG/Private/Components/ProgressBar.cpp
--- a/Engine/Source/Runtime/UMG/Private/Components/ProgressBar.cpp
+++ b/Engine/Source/Runtime/UMG/Private/Components/ProgressBar.cpp
@@ -40,6 +40,9 @@ UProgressBar::UProgressBar(const FObjectInitializer& ObjectInitializer)
 	if (IsEditorWidget())
 	{
 		WidgetStyle = *EditorProgressBarStyle;
+
+		// The CDO isn't an editor widget and thus won't use the editor style, call post edit change to mark difference from CDO
+		PostEditChange();
 	}
 #endif // WITH_EDITOR
 
diff --git a/Engine/Source/Runtime/UMG/Private/Components/ScrollBar.cpp b/Engine/Source/Runtime/UMG/Private/Components/ScrollBar.cpp
--- a/Engine/Source/Runtime/UMG/Private/Components/ScrollBar.cpp
+++ b/Engine/Source/Runtime/UMG/Private/Components/ScrollBar.cpp
@@ -48,6 +48,9 @@ UScrollBar::UScrollBar(const FObjectInitializer& ObjectInitializer)
 	if (IsEditorWidget())
 	{
 		WidgetStyle = *EditorScrollBarStyle;
+
+		// The CDO isn't an editor widget and thus won't use the editor style, call post edit change to mark difference from CDO
+		PostEditChange();
 	}
 #endif // WITH_EDITOR
 }
diff --git a/Engine/Source/Runtime/UMG/Private/Components/ScrollBox.cpp b/Engine/Source/Runtime/UMG/Private/Components/ScrollBox.cpp
--- a/Engine/Source/Runtime/UMG/Private/Components/ScrollBox.cpp
+++ b/Engine/Source/Runtime/UMG/Private/Components/ScrollBox.cpp
@@ -80,6 +80,9 @@ UScrollBox::UScrollBox(const FObjectInitializer& ObjectInitializer)
 	{
 		WidgetStyle = *EditorScrollBoxStyle;
 		WidgetBarStyle = *EditorScrollBoxBarStyle;
+
+		// The CDO isn't an editor widget and thus won't use the editor style, call post edit change to mark difference from CDO
+		PostEditChange();
 	}
 #endif // WITH_EDITOR
 
diff --git a/Engine/Source/Runtime/UMG/Private/Components/Slider.cpp b/Engine/Source/Runtime/UMG/Private/Components/Slider.cpp
--- a/Engine/Source/Runtime/UMG/Private/Components/Slider.cpp
+++ b/Engine/Source/Runtime/UMG/Private/Components/Slider.cpp
@@ -52,6 +52,9 @@ USlider::USlider(const FObjectInitializer& ObjectInitializer)
 	if (IsEditorWidget())
 	{
 		WidgetStyle = *EditorSliderStyle;
+
+		// The CDO isn't an editor widget and thus won't use the editor style, call post edit change to mark difference from CDO
+		PostEditChange();
 	}
 #endif // WITH_EDITOR
 
diff --git a/Engine/Source/Runtime/UMG/Private/Components/SpinBox.cpp b/Engine/Source/Runtime/UMG/Private/Components/SpinBox.cpp
--- a/Engine/Source/Runtime/UMG/Private/Components/SpinBox.cpp
+++ b/Engine/Source/Runtime/UMG/Private/Components/SpinBox.cpp
@@ -61,6 +61,9 @@ USpinBox::USpinBox(const FObjectInitializer& ObjectInitializer)
 	if (IsEditorWidget())
 	{
 		WidgetStyle = *EditorSpinBoxStyle;
+
+		// The CDO isn't an editor widget and thus won't use the editor style, call post edit change to mark difference from CDO
+		PostEditChange();
 	}
 #endif // WITH_EDITOR
 
diff --git a/Engine/Source/Runtime/UMG/Private/Components/Throbber.cpp b/Engine/Source/Runtime/UMG/Private/Components/Throbber.cpp
--- a/Engine/Source/Runtime/UMG/Private/Components/Throbber.cpp
+++ b/Engine/Source/Runtime/UMG/Private/Components/Throbber.cpp
@@ -47,6 +47,9 @@ UThrobber::UThrobber(const FObjectInitializer& ObjectInitializer)
 	if (IsEditorWidget())
 	{
 		Image = *EditorThrobberBrush;
+
+		// The CDO isn't an editor widget and thus won't use the editor style, call post edit change to mark difference from CDO
+		PostEditChange();
 	}
 #endif // WITH_EDITOR
 }
diff --git a/Engine/Source/Runtime/UMG/Private/Components/Widget.cpp b/Engine/Source/Runtime/UMG/Private/Components/Widget.cpp
--- a/Engine/Source/Runtime/UMG/Private/Components/Widget.cpp
+++ b/Engine/Source/Runtime/UMG/Private/Components/Widget.cpp
@@ -1089,6 +1089,10 @@ bool UWidget::IsEditorWidget() const
 		{
 			return WidgetBP->AllowEditorWidget();
 		}
+		else if (UUserWidget* UserWidget = Cast<UUserWidget>(WidgetBPObject))
+		{
+			return UserWidget->IsEditorUtility();
+		}
 	}
 
 	return false;
diff --git a/Engine/Source/Runtime/UMG/Public/Blueprint/UserWidget.h b/Engine/Source/Runtime/UMG/Public/Blueprint/UserWidget.h
--- a/Engine/Source/Runtime/UMG/Public/Blueprint/UserWidget.h
+++ b/Engine/Source/Runtime/UMG/Public/Blueprint/UserWidget.h
@@ -846,6 +846,9 @@ public:
 	UFUNCTION(BlueprintCallable, Category = Animation)
 	void BindToAnimationEvent(UWidgetAnimation* Animation, FWidgetAnimationDynamicEvent Delegate, EWidgetAnimationEvent AnimationEvent, FName UserTag = NAME_None);
 
+	/** Is this widget an editor utility widget. */
+	virtual bool IsEditorUtility() const { return false; }
+
 protected:
 
 	/**
