PR #8306: Add DrawSpline to Widget Blueprint library (Contributed by [REDACTED_PERSON])

diff --git a/Engine/Source/Runtime/UMG/Private/WidgetBlueprintLibrary.cpp b/Engine/Source/Runtime/UMG/Private/WidgetBlueprintLibrary.cpp
--- a/Engine/Source/Runtime/UMG/Private/WidgetBlueprintLibrary.cpp
+++ b/Engine/Source/Runtime/UMG/Private/WidgetBlueprintLibrary.cpp
@@ -166,6 +166,23 @@ void UWidgetBlueprintLibrary::DrawBox(FPaintContext& Context, FVector2D Position
 	}
 }
 
+void UWidgetBlueprintLibrary::DrawSpline(FPaintContext& Context, FVector2D Start, FVector2D StartDir, FVector2D End, FVector2D EndDir, FLinearColor Tint, float Thickness)
+{
+	Context.MaxLayer++;
+
+	FSlateDrawElement::MakeSpline(
+		Context.OutDrawElements,
+		Context.MaxLayer,
+		Context.AllottedGeometry.ToPaintGeometry(),
+		Start,
+		StartDir,
+		End,
+		EndDir,
+		Thickness,
+		ESlateDrawEffect::None,
+		Tint);
+}
+
 void UWidgetBlueprintLibrary::DrawLine(FPaintContext& Context, FVector2D PositionA, FVector2D PositionB, FLinearColor Tint, bool bAntiAlias, float Thickness)
 {
 	Context.MaxLayer++;
diff --git a/Engine/Source/Runtime/UMG/Public/Blueprint/WidgetBlueprintLibrary.h b/Engine/Source/Runtime/UMG/Public/Blueprint/WidgetBlueprintLibrary.h
--- a/Engine/Source/Runtime/UMG/Public/Blueprint/WidgetBlueprintLibrary.h
+++ b/Engine/Source/Runtime/UMG/Public/Blueprint/WidgetBlueprintLibrary.h
@@ -71,6 +71,19 @@ public:
 	UFUNCTION(BlueprintCallable, Category="Painting")
 	static void DrawBox(UPARAM(ref) FPaintContext& Context, FVector2D Position, FVector2D Size, USlateBrushAsset* Brush, FLinearColor Tint = FLinearColor::White);
 
+	/**
+	 * Draws a hermite spline.
+	 *
+	 * @param Start			Starting position of the spline in local space.
+	 * @param StartDir		The direction of the spline from the start point.
+	 * @param End			Ending position of the spline in local space.
+	 * @param EndDir		The direction of the spline to the end point.
+	 * @param Tint			Color to render the spline.
+	 * @param Thickness		How many pixels thick this spline should be.
+	 */
+	UFUNCTION(BlueprintCallable, meta = (AdvancedDisplay = "6"), Category = "Painting")
+	static void DrawSpline(UPARAM(ref) FPaintContext& Context, FVector2D Start, FVector2D StartDir, FVector2D End, FVector2D EndDir, FLinearColor Tint = FLinearColor::White, float Thickness = 1.0f);
+
 	/**
 	 * Draws a line.
 	 *
