PR #7612: Checking RenderOpacity in UWidget::IsVisible() (Contributed by [REDACTED_PERSON])

diff --git a/Engine/Source/Runtime/UMG/Private/Components/Widget.cpp b/Engine/Source/Runtime/UMG/Private/Components/Widget.cpp
--- a/Engine/Source/Runtime/UMG/Private/Components/Widget.cpp
+++ b/Engine/Source/Runtime/UMG/Private/Components/Widget.cpp
@@ -282,6 +282,17 @@ void UWidget::ResetCursor()
 	}
 }
 
+bool UWidget::IsRendered() const
+{
+	TSharedPtr<SWidget> SafeWidget = GetCachedWidget();
+	if ( SafeWidget.IsValid() )
+	{
+		return SafeWidget->GetVisibility().IsVisible() && SafeWidget->GetRenderOpacity() > 0.0f;
+	}
+
+	return false;
+}
+
 bool UWidget::IsVisible() const
 {
 	TSharedPtr<SWidget> SafeWidget = GetCachedWidget();
diff --git a/Engine/Source/Runtime/UMG/Public/Components/Widget.h b/Engine/Source/Runtime/UMG/Public/Components/Widget.h
--- a/Engine/Source/Runtime/UMG/Public/Components/Widget.h
+++ b/Engine/Source/Runtime/UMG/Public/Components/Widget.h
@@ -511,6 +511,10 @@ public:
 	/** Resets the cursor to use on the widget, removing any customization for it. */
 	UFUNCTION(BlueprintCallable, Category="Widget")
 	void ResetCursor();
+	
+	/** Returns true if the widget is Visible, HitTestInvisible or SelfHitTestInvisible and the Render Opacity is greater than 0. */
+	UFUNCTION(BlueprintCallable, Category="Widget")
+	bool IsRendered() const;
 
 	/** Returns true if the widget is Visible, HitTestInvisible or SelfHitTestInvisible. */
 	UFUNCTION(BlueprintCallable, Category="Widget")
