Fix default editor utility widget foreground colors hard to see.

diff --git a/Engine/Source/Runtime/UMG/Private/Components/ComboBoxKey.cpp b/Engine/Source/Runtime/UMG/Private/Components/ComboBoxKey.cpp
--- a/Engine/Source/Runtime/UMG/Private/Components/ComboBoxKey.cpp
+++ b/Engine/Source/Runtime/UMG/Private/Components/ComboBoxKey.cpp
@@ -35,7 +35,7 @@ UComboBoxKey::UComboBoxKey()
 	}
 
 	ContentPadding = FMargin(4.0, 2.0);
-	ForegroundColor = FLinearColor::Black;
+	ForegroundColor = ItemStyle.TextColor;
 
 	MaxListHeight = 450.0f;
 	bHasDownArrow = true;
diff --git a/Engine/Source/Runtime/UMG/Private/Components/ComboBoxString.cpp b/Engine/Source/Runtime/UMG/Private/Components/ComboBoxString.cpp
--- a/Engine/Source/Runtime/UMG/Private/Components/ComboBoxString.cpp
+++ b/Engine/Source/Runtime/UMG/Private/Components/ComboBoxString.cpp
@@ -72,7 +72,7 @@ UComboBoxString::UComboBoxString(const FObjectInitializer& ObjectInitializer)
 	WidgetStyle.UnlinkColors();
 	ItemStyle.UnlinkColors();
 
-	ForegroundColor = FLinearColor::Black;
+	ForegroundColor = ItemStyle.TextColor;
 	bIsFocusable = true;
 
 	ContentPadding = FMargin(4.0, 2.0);
diff --git a/Engine/Source/Runtime/UMG/Private/Components/MultiLineEditableTextBox.cpp b/Engine/Source/Runtime/UMG/Private/Components/MultiLineEditableTextBox.cpp
--- a/Engine/Source/Runtime/UMG/Private/Components/MultiLineEditableTextBox.cpp
+++ b/Engine/Source/Runtime/UMG/Private/Components/MultiLineEditableTextBox.cpp
@@ -82,11 +82,6 @@ UMultiLineEditableTextBox::UMultiLineEditableTextBox(const FObjectInitializer& O
 	{
 		static ConstructorHelpers::FObjectFinder<UFont> RobotoFontObj(*UWidget::GetDefaultFontName());
 		Font_DEPRECATED = FSlateFontInfo(RobotoFontObj.Object, 12, FName("Bold"));
-
-		WidgetStyle.SetFont(Font_DEPRECATED);
-		WidgetStyle.SetForegroundColor(ForegroundColor_DEPRECATED);
-		WidgetStyle.SetBackgroundColor(BackgroundColor_DEPRECATED);
-		WidgetStyle.SetReadOnlyForegroundColor(ReadOnlyForegroundColor_DEPRECATED);
 	}
 }
 
