Making Scrollbox "NavigationDestination" editable at runtime

diff --git a/Engine/Source/Runtime/Slate/Private/Widgets/Layout/SScrollBox.cpp b/Engine/Source/Runtime/Slate/Private/Widgets/Layout/SScrollBox.cpp
--- a/Engine/Source/Runtime/Slate/Private/Widgets/Layout/SScrollBox.cpp
+++ b/Engine/Source/Runtime/Slate/Private/Widgets/Layout/SScrollBox.cpp
@@ -470,6 +470,11 @@ EOrientation SScrollBox::GetOrientation()
 	return Orientation;
 }
 
+void SScrollBox::SetNavigationDestination(const EDescendantScrollDestination NewNavigationDestination)
+{
+	NavigationDestination = NewNavigationDestination;
+}
+
 void SScrollBox::SetConsumeMouseWheel(EConsumeMouseWheel NewConsumeMouseWheel)
 {
 	ConsumeMouseWheel = NewConsumeMouseWheel;
diff --git a/Engine/Source/Runtime/Slate/Public/Widgets/Layout/SScrollBox.h b/Engine/Source/Runtime/Slate/Public/Widgets/Layout/SScrollBox.h
--- a/Engine/Source/Runtime/Slate/Public/Widgets/Layout/SScrollBox.h
+++ b/Engine/Source/Runtime/Slate/Public/Widgets/Layout/SScrollBox.h
@@ -205,6 +205,8 @@ public:
 	/** Get the current orientation of the scrollbox. */
 	EOrientation GetOrientation();
 
+	void SetNavigationDestination(const EDescendantScrollDestination NewNavigationDestination);
+
 	void SetConsumeMouseWheel(EConsumeMouseWheel NewConsumeMouseWheel);
 
 	/** Sets the current orientation of the scrollbox and updates the layout */
diff --git a/Engine/Source/Runtime/UMG/Private/Components/ScrollBox.cpp b/Engine/Source/Runtime/UMG/Private/Components/ScrollBox.cpp
--- a/Engine/Source/Runtime/UMG/Private/Components/ScrollBox.cpp
+++ b/Engine/Source/Runtime/UMG/Private/Components/ScrollBox.cpp
@@ -292,6 +292,16 @@ void UScrollBox::PostLoad()
 
 #endif // if WITH_EDITORONLY_DATA
 
+void UScrollBox::SetNavigationDestination(const EDescendantScrollDestination NewNavigationDestination)
+{
+	NavigationDestination = NewNavigationDestination;
+
+	if (MyScrollBox.IsValid())
+	{
+		MyScrollBox->SetNavigationDestination(NewNavigationDestination);
+	}
+}
+
 void UScrollBox::SetConsumeMouseWheel(EConsumeMouseWheel NewConsumeMouseWheel)
 {
 	ConsumeMouseWheel = NewConsumeMouseWheel;
diff --git a/Engine/Source/Runtime/UMG/Public/Components/ScrollBox.h b/Engine/Source/Runtime/UMG/Public/Components/ScrollBox.h
--- a/Engine/Source/Runtime/UMG/Public/Components/ScrollBox.h
+++ b/Engine/Source/Runtime/UMG/Public/Components/ScrollBox.h
@@ -134,6 +134,9 @@ public:
 	UFUNCTION(BlueprintCallable, Category = "Scroll")
 	void SetScrollWhenFocusChanges(EScrollWhenFocusChanges NewScrollWhenFocusChanges);
 
+	UFUNCTION(BlueprintCallable, Category = "Scroll")
+	void SetNavigationDestination(const EDescendantScrollDestination NewNavigationDestination);
+
 	/** Instantly stops any inertial scrolling that is currently in progress */
 	UFUNCTION(BlueprintCallable, Category = "Scroll")
 	void EndInertialScrolling();
