Sequencer: Specualtive fix for Animation Tick Managers becoming null during latent actions

diff --git a/Engine/Source/Runtime/UMG/Private/Animation/UMGSequencePlayer.cpp b/Engine/Source/Runtime/UMG/Private/Animation/UMGSequencePlayer.cpp
--- a/Engine/Source/Runtime/UMG/Private/Animation/UMGSequencePlayer.cpp
+++ b/Engine/Source/Runtime/UMG/Private/Animation/UMGSequencePlayer.cpp
@@ -7,6 +7,7 @@
 #include "MovieSceneTimeHelpers.h"
 #include "Evaluation/MovieScenePlayback.h"
 #include "EntitySystem/MovieSceneEntitySystemLinker.h"
+#include "Animation/UMGSequenceTickManager.h"
 
 extern TAutoConsoleVariable<bool> CVarUserWidgetUseParallelAnimation;
 
@@ -41,6 +42,14 @@ UMovieSceneEntitySystemLinker* UUMGSequencePlayer::ConstructEntitySystemLinker()
 	UUserWidget* Widget = UserWidget.Get();
 	if (ensure(Widget) && !EnumHasAnyFlags(Animation->GetFlags(), EMovieSceneSequenceFlags::BlockingEvaluation))
 	{
+		if (!ensure(Widget->AnimationTickManager))
+		{
+			// @todo: There should be no possible way that the animation tick manager is null here, but there is a very low-rate
+			// crash caused by it being null that is very hard to track down, so patching with a band-aid for now.
+			Widget->AnimationTickManager = UUMGSequenceTickManager::Get(Widget);
+			Widget->AnimationTickManager->AddWidget(Widget);
+		}
+
 		return Widget->AnimationTickManager->GetLinker();
 	}
 
