Update comments on WidgetBlueprintLibrary to be more explict that Input Events will not be called in UI Only mode. Remove some old deprecated functions from 4.13

#rnx

diff --git a/Engine/Plugins/Experimental/EnhancedInput/Source/InputBlueprintNodes/Private/K2Node_EnhancedInputAction.cpp b/Engine/Plugins/Experimental/EnhancedInput/Source/InputBlueprintNodes/Private/K2Node_EnhancedInputAction.cpp
--- a/Engine/Plugins/Experimental/EnhancedInput/Source/InputBlueprintNodes/Private/K2Node_EnhancedInputAction.cpp
+++ b/Engine/Plugins/Experimental/EnhancedInput/Source/InputBlueprintNodes/Private/K2Node_EnhancedInputAction.cpp
@@ -188,7 +188,11 @@ FText UK2Node_EnhancedInputAction::GetTooltipText() const
 	{
 		// FText::Format() is slow, so we cache this to save on performance
 		FString ActionPath = InputAction ? InputAction->GetFullName() : TEXT("");
-		CachedTooltip.SetCachedText(FText::Format(LOCTEXT("EnhancedInputAction_Tooltip", "Event for when '{0}' triggers.\n\nNOTE: This is not guaranteed to fire every frame, only when the Action is triggered."), FText::FromString(ActionPath)), this);
+		CachedTooltip.SetCachedText(
+			FText::Format(
+				LOCTEXT("EnhancedInputAction_Tooltip", "Event for when '{0}' triggers.\n\nNOTE: This is not guaranteed to fire every frame, only when the Action is triggered and the current Input Mode includes 'Game'."),
+			FText::FromString(ActionPath)),
+			this);
 	}
 	return CachedTooltip;
 }
diff --git a/Engine/Source/Runtime/UMG/Private/WidgetBlueprintLibrary.cpp b/Engine/Source/Runtime/UMG/Private/WidgetBlueprintLibrary.cpp
--- a/Engine/Source/Runtime/UMG/Private/WidgetBlueprintLibrary.cpp
+++ b/Engine/Source/Runtime/UMG/Private/WidgetBlueprintLibrary.cpp
@@ -76,11 +76,6 @@ UDragDropOperation* UWidgetBlueprintLibrary::CreateDragDropOperation(TSubclassOf
 	return DragDropOperation;
 }
 
-void UWidgetBlueprintLibrary::SetInputMode_UIOnly(APlayerController* Target, UWidget* InWidgetToFocus, bool bLockMouseToViewport)
-{
-	SetInputMode_UIOnlyEx(Target, InWidgetToFocus, bLockMouseToViewport ? EMouseLockMode::LockOnCapture : EMouseLockMode::DoNotLock);
-}
-
 void UWidgetBlueprintLibrary::SetInputMode_UIOnlyEx(APlayerController* PlayerController, UWidget* InWidgetToFocus, EMouseLockMode InMouseLockMode)
 {
 	if (PlayerController != nullptr)
@@ -102,11 +97,6 @@ void UWidgetBlueprintLibrary::SetInputMode_UIOnlyEx(APlayerController* PlayerCon
 	#endif // WITH_EDITOR
 }
 
-void UWidgetBlueprintLibrary::SetInputMode_GameAndUI(APlayerController* Target, UWidget* InWidgetToFocus, bool bLockMouseToViewport, bool bHideCursorDuringCapture)
-{
-	SetInputMode_GameAndUIEx(Target, InWidgetToFocus, bLockMouseToViewport ? EMouseLockMode::LockOnCapture : EMouseLockMode::DoNotLock, bHideCursorDuringCapture);
-}
-
 void UWidgetBlueprintLibrary::SetInputMode_GameAndUIEx(APlayerController* PlayerController, UWidget* InWidgetToFocus, EMouseLockMode InMouseLockMode, bool bHideCursorDuringCapture)
 {
 	if (PlayerController != nullptr)
diff --git a/Engine/Source/Runtime/UMG/Public/Blueprint/WidgetBlueprintLibrary.h b/Engine/Source/Runtime/UMG/Public/Blueprint/WidgetBlueprintLibrary.h
--- a/Engine/Source/Runtime/UMG/Public/Blueprint/WidgetBlueprintLibrary.h
+++ b/Engine/Source/Runtime/UMG/Public/Blueprint/WidgetBlueprintLibrary.h
@@ -41,25 +41,27 @@ public:
 	UFUNCTION(BlueprintCallable, Category="Widget|Drag and Drop", meta=( BlueprintInternalUseOnly="true" ))
 	static UDragDropOperation* CreateDragDropOperation(TSubclassOf<UDragDropOperation> OperationClass);
 	
-	/** Setup an input mode that allows only the UI to respond to user input. */
-	UE_DEPRECATED(4.13, "Locking the mouse to the viewport is now controlled by an enum. Call SetInputMode_UIOnlyEx instead")
-	UFUNCTION(BlueprintCallable, BlueprintCosmetic, Category = "Input", meta = (DeprecatedFunction, DeprecationMessage = "Use the new version of Set Input Mode UI Only instead"), DisplayName = "Set Input Mode UI Only (Deprecated)")
-	static void SetInputMode_UIOnly(APlayerController* Target, UWidget* InWidgetToFocus = nullptr, bool bLockMouseToViewport = false);
-
-	/** Setup an input mode that allows only the UI to respond to user input. */
+	/**
+	 * Setup an input mode that allows only the UI to respond to user input.
+	 * 
+	 * Note: This means that any bound Input Events in the widget will not be called!
+	 */
 	UFUNCTION(BlueprintCallable, BlueprintCosmetic, Category="Input", meta=(DisplayName="Set Input Mode UI Only"))
 	static void SetInputMode_UIOnlyEx(APlayerController* PlayerController, UWidget* InWidgetToFocus = nullptr, EMouseLockMode InMouseLockMode = EMouseLockMode::DoNotLock);
 
-	/** Setup an input mode that allows only the UI to respond to user input, and if the UI doesn't handle it player input / player controller gets a chance. */
-	UE_DEPRECATED(4.13, "Locking the mouse to the viewport is now controlled by an enum. Call SetInputMode_GameAndUIEx instead")
-	UFUNCTION(BlueprintCallable, BlueprintCosmetic, Category = "Input", meta = (DeprecatedFunction, DeprecationMessage = "Use the new version of Set Input Mode Game And UI instead"), DisplayName = "Set Input Mode Game And UI (Deprecated)")
-	static void SetInputMode_GameAndUI(APlayerController* Target, UWidget* InWidgetToFocus = nullptr, bool bLockMouseToViewport = false, bool bHideCursorDuringCapture = true);
-
-	/** Setup an input mode that allows only the UI to respond to user input, and if the UI doesn't handle it player input / player controller gets a chance. */
+	/**
+	 * Setup an input mode that allows only the UI to respond to user input, and if the UI doesn't handle it player input / player controller gets a chance.
+	 * 
+	 * Note: This means that any bound Input events in the widget will be called.
+	 */
 	UFUNCTION(BlueprintCallable, BlueprintCosmetic, Category="Input", meta = (DisplayName = "Set Input Mode Game And UI"))
 	static void SetInputMode_GameAndUIEx(APlayerController* PlayerController, UWidget* InWidgetToFocus = nullptr, EMouseLockMode InMouseLockMode = EMouseLockMode::DoNotLock, bool bHideCursorDuringCapture = true);
 
-	/** Setup an input mode that allows only player input / player controller to respond to user input. */
+	/**
+	 * Setup an input mode that allows only player input / player controller to respond to user input.
+	 * 
+	 * Note: Any bound Input Events in this widget will be called.
+	 */
 	UFUNCTION(BlueprintCallable, BlueprintCosmetic, Category="Input")
 	static void SetInputMode_GameOnly(APlayerController* PlayerController);
 
