Add latent action to clear StoppedSequencePlayers even if the widget is not ticking and would thus not handle PostTickActionsAndAnimation.

#tests: PIE: Verified off-screen play, stop, and play again operates as expected, and animation plays when widget starts ticking again.

diff --git a/Engine/Source/Runtime/UMG/Private/UserWidget.cpp b/Engine/Source/Runtime/UMG/Private/UserWidget.cpp
--- a/Engine/Source/Runtime/UMG/Private/UserWidget.cpp
+++ b/Engine/Source/Runtime/UMG/Private/UserWidget.cpp
@@ -697,6 +697,11 @@ void UUserWidget::OnAnimationFinishedPlaying(UUMGSequencePlayer& Player)
 	if ( Player.GetPlaybackStatus() == EMovieScenePlayerStatus::Stopped )
 	{
 		StoppedSequencePlayers.Add(&Player);
+
+		if (AnimationTickManager)
+		{
+			AnimationTickManager->AddLatentAction(FMovieSceneSequenceLatentActionDelegate::CreateUObject(this, &UUserWidget::ClearStoppedSequencePlayers));
+		}
 	}
 
 	UpdateCanTick();
@@ -1458,14 +1463,7 @@ void UUserWidget::TickActionsAndAnimation(float InDeltaTime)
 
 void UUserWidget::PostTickActionsAndAnimation(float InDeltaTime)
 {
-	// The process of ticking the players above can stop them so we remove them after all players have ticked
-	for (UUMGSequencePlayer* StoppedPlayer : StoppedSequencePlayers)
-	{
-		ActiveSequencePlayers.RemoveSwap(StoppedPlayer);
-		StoppedPlayer->TearDown();
-	}
-
-	StoppedSequencePlayers.Empty();
+	ClearStoppedSequencePlayers();
 }
 
 void UUserWidget::FlushAnimations()
@@ -2046,6 +2044,18 @@ UUserWidget* UUserWidget::CreateInstanceInternal(UObject* Outer, TSubclassOf<UUs
 }
 
 
+void UUserWidget::ClearStoppedSequencePlayers()
+{
+	// after all players have ticked, remove and tear down stopped players
+	for (UUMGSequencePlayer* StoppedPlayer : StoppedSequencePlayers)
+	{
+		ActiveSequencePlayers.RemoveSwap(StoppedPlayer);
+		StoppedPlayer->TearDown();
+	}
+
+	StoppedSequencePlayers.Empty();
+}
+
 void UUserWidget::OnLatentActionsChanged(UObject* ObjectWhichChanged, ELatentActionChangeType ChangeType)
 {
 	if (UUserWidget* WidgetThatChanged = Cast<UUserWidget>(ObjectWhichChanged))
diff --git a/Engine/Source/Runtime/UMG/Public/Blueprint/UserWidget.h b/Engine/Source/Runtime/UMG/Public/Blueprint/UserWidget.h
--- a/Engine/Source/Runtime/UMG/Public/Blueprint/UserWidget.h
+++ b/Engine/Source/Runtime/UMG/Public/Blueprint/UserWidget.h
@@ -1142,6 +1142,8 @@ public:
 private:
 	static UUserWidget* CreateInstanceInternal(UObject* Outer, TSubclassOf<UUserWidget> UserWidgetClass, FName WidgetName, UWorld* World, ULocalPlayer* LocalPlayer);
 
+	void ClearStoppedSequencePlayers();
+
 public:
 	/** The color and opacity of this widget.  Tints all child widgets. */
 	UPROPERTY(EditAnywhere, BlueprintReadOnly, Category="Appearance")
