Expose overflow policy to RichTextBlock
- Moving the propertly to the Clipping->Advanced section ad it make more sense to put it there
- A few comments fixes
- Minor fixes to TextBlock usage of the property

diff --git a/Engine/Source/Runtime/SlateCore/Public/Styling/SlateTypes.h b/Engine/Source/Runtime/SlateCore/Public/Styling/SlateTypes.h
--- a/Engine/Source/Runtime/SlateCore/Public/Styling/SlateTypes.h
+++ b/Engine/Source/Runtime/SlateCore/Public/Styling/SlateTypes.h
@@ -289,7 +289,7 @@ struct TStructOpsTypeTraits<FCheckBoxStyle> : public TStructOpsTypeTraitsBase2<F
 #endif
 
 /**
- * The different methods that can be used if a word is too long to be broken by the default line-break iterator.
+ * Text transformation policy that can be applied to the text before displaying it.
  */
 UENUM(BlueprintType)
 enum class ETextTransformPolicy : uint8
diff --git a/Engine/Source/Runtime/UMG/Private/Components/RichTextBlock.cpp b/Engine/Source/Runtime/UMG/Private/Components/RichTextBlock.cpp
--- a/Engine/Source/Runtime/UMG/Private/Components/RichTextBlock.cpp
+++ b/Engine/Source/Runtime/UMG/Private/Components/RichTextBlock.cpp
@@ -48,6 +48,7 @@ URichTextBlock::URichTextBlock(const FObjectInitializer& ObjectInitializer)
 {
 	Visibility = ESlateVisibility::SelfHitTestInvisible;
 	TextTransformPolicy = ETextTransformPolicy::None;
+	TextOverflowPolicy = ETextOverflowPolicy::Clip;
 }
 
 void URichTextBlock::ReleaseSlateResources(bool bReleaseChildren)
@@ -80,11 +81,16 @@ void URichTextBlock::SynchronizeProperties()
 {
 	Super::SynchronizeProperties();
 
-	MyRichTextBlock->SetText(Text);
-	MyRichTextBlock->SetTransformPolicy(TextTransformPolicy);
-	MyRichTextBlock->SetMinDesiredWidth(MinDesiredWidth);
+	if (MyRichTextBlock.IsValid())
+	{
+		MyRichTextBlock->SetText(Text);
+		MyRichTextBlock->SetTransformPolicy(TextTransformPolicy);
+		MyRichTextBlock->SetMinDesiredWidth(MinDesiredWidth);
 
-	Super::SynchronizeTextLayoutProperties( *MyRichTextBlock );
+		MyRichTextBlock->SetOverflowPolicy(TextOverflowPolicy);
+
+		Super::SynchronizeTextLayoutProperties(*MyRichTextBlock);
+	}
 }
 
 void URichTextBlock::UpdateStyleData()
@@ -350,6 +356,16 @@ void URichTextBlock::SetTextTransformPolicy(ETextTransformPolicy InTransformPoli
 	}
 }
 
+void URichTextBlock::SetTextOverflowPolicy(ETextOverflowPolicy InOverflowPolicy)
+{
+	TextOverflowPolicy = InOverflowPolicy;
+
+	if (MyRichTextBlock.IsValid())
+	{
+		MyRichTextBlock->SetOverflowPolicy(TextOverflowPolicy);
+	}
+}
+
 void URichTextBlock::ApplyUpdatedDefaultTextStyle()
 {
 	if (MyRichTextBlock.IsValid())
diff --git a/Engine/Source/Runtime/UMG/Private/Components/TextBlock.cpp b/Engine/Source/Runtime/UMG/Private/Components/TextBlock.cpp
--- a/Engine/Source/Runtime/UMG/Private/Components/TextBlock.cpp
+++ b/Engine/Source/Runtime/UMG/Private/Components/TextBlock.cpp
@@ -25,6 +25,7 @@ UTextBlock::UTextBlock(const FObjectInitializer& ObjectInitializer)
 	ShadowColorAndOpacity = FLinearColor::Transparent;
 	bAutoWrapText_DEPRECATED = false;
 	TextTransformPolicy = ETextTransformPolicy::None;
+	TextOverflowPolicy = ETextOverflowPolicy::Clip;
 
 	if (!IsRunningDedicatedServer())
 	{
diff --git a/Engine/Source/Runtime/UMG/Public/Components/RichTextBlock.h b/Engine/Source/Runtime/UMG/Public/Components/RichTextBlock.h
--- a/Engine/Source/Runtime/UMG/Public/Components/RichTextBlock.h
+++ b/Engine/Source/Runtime/UMG/Public/Components/RichTextBlock.h
@@ -95,6 +95,13 @@ public:
 	UFUNCTION(BlueprintCallable, Category = "Appearance")
 	void SetTextTransformPolicy(ETextTransformPolicy InTransformPolicy);
 
+	/**
+	* Set the text overflow policy for this text block.
+	* @param InOverflowPolicy the new text overflow policy.
+	*/
+	UFUNCTION(BlueprintCallable, Category = "Appearance")
+	void SetTextOverflowPolicy(ETextOverflowPolicy InOverflowPolicy);
+
 	/** 
 	 * Wholesale override of the currently established default text style
 	 * @param InDefaultTextStyle The new text style to apply to all default (i.e. undecorated) text in the block
@@ -200,6 +207,10 @@ protected:
 	UPROPERTY(EditAnywhere, BlueprintReadOnly, Category=Appearance, meta=(DisplayName="Transform Policy"))
 	ETextTransformPolicy TextTransformPolicy;
 
+	/** Sets what happens to text that is clipped and doesn't fit within the clip rect for this widget */
+	UPROPERTY(EditAnywhere, BlueprintReadWrite, Category = Clipping, AdvancedDisplay, meta = (DisplayName = "Overflow Policy"))
+	ETextOverflowPolicy TextOverflowPolicy;
+
 	UPROPERTY(Transient)
 	FTextBlockStyle DefaultTextStyle;
 
diff --git a/Engine/Source/Runtime/UMG/Public/Components/TextBlock.h b/Engine/Source/Runtime/UMG/Public/Components/TextBlock.h
--- a/Engine/Source/Runtime/UMG/Public/Components/TextBlock.h
+++ b/Engine/Source/Runtime/UMG/Public/Components/TextBlock.h
@@ -97,6 +97,11 @@ public:
 	UFUNCTION(BlueprintCallable, Category = "Appearance")
 	void SetTextTransformPolicy(ETextTransformPolicy InTransformPolicy);
 
+	/**
+	 * Set the text overflow policy for this text block.
+	 *
+	 * @param InOverflowPolicy the new text overflow policy.
+	 */
 	UFUNCTION(BlueprintCallable, Category = "Appearance")
 	void SetTextOverflowPolicy(ETextOverflowPolicy InOverflowPolicy);
 
@@ -162,7 +167,7 @@ public:
 	ETextTransformPolicy TextTransformPolicy;
 
 	/** Sets what happens to text that is clipped and doesn't fit within the clip rect for this widget */
-	UPROPERTY(EditAnywhere, BlueprintReadWrite, Category = Appearance, AdvancedDisplay, meta = (DisplayName = "Overflow Policy"))
+	UPROPERTY(EditAnywhere, BlueprintReadWrite, Category = Clipping, AdvancedDisplay, meta = (DisplayName = "Overflow Policy"))
 	ETextOverflowPolicy TextOverflowPolicy;
 
 	/** 
