UMG list view will now track its actor references and remove them when they're about to be destroyed

diff --git a/Engine/Source/Runtime/UMG/Private/Components/ListView.cpp b/Engine/Source/Runtime/UMG/Private/Components/ListView.cpp
--- a/Engine/Source/Runtime/UMG/Private/Components/ListView.cpp
+++ b/Engine/Source/Runtime/UMG/Private/Components/ListView.cpp
@@ -247,6 +247,45 @@ void UListView::BP_ClearSelection()
 void UListView::OnItemsChanged(const TArray<UObject*>& AddedItems, const TArray<UObject*>& RemovedItems)
 {
 	// Allow subclasses to do special things when objects are added or removed from the list.
+
+	// Keep track of references to Actors and make sure to release them when Actors are about to be removed
+	for (UObject* AddedItem : AddedItems)
+	{		
+		if (AActor* AddedActor = Cast<AActor>(AddedItem))
+		{
+			AddedActor->OnEndPlay.AddDynamic(this, &UListView::OnListItemEndPlayed);
+		}
+		else if (AActor* AddedItemOuterActor = AddedItem->GetTypedOuter<AActor>())
+		{
+			// Unique so that we don't spam events for shared actor outers but this also means we can't
+			// unsubscribe when processing RemovedItems
+			AddedItemOuterActor->OnEndPlay.AddUniqueDynamic(this, &UListView::OnListItemOuterEndPlayed);
+		}
+	}
+	for (UObject* RemovedItem : RemovedItems)
+	{
+		if (AActor* RemovedActor = Cast<AActor>(RemovedItem))
+		{
+			RemovedActor->OnEndPlay.RemoveDynamic(this, &UListView::OnListItemEndPlayed);
+		}
+	}
+}
+
+void UListView::OnListItemEndPlayed(AActor* Item, EEndPlayReason::Type EndPlayReason)
+{
+	RemoveItem(Item);
+}
+
+void UListView::OnListItemOuterEndPlayed(AActor* ItemOuter, EEndPlayReason::Type EndPlayReason)
+{
+	for (int32 ItemIndex = ListItems.Num() - 1; ItemIndex >= 0; --ItemIndex)
+	{
+		UObject* Item = ListItems[ItemIndex];
+		if (Item->IsIn(ItemOuter))
+		{
+			RemoveItem(Item);
+		}
+	}
 }
 
 TSharedRef<STableViewBase> UListView::RebuildListWidget()
diff --git a/Engine/Source/Runtime/UMG/Public/Components/ListView.h b/Engine/Source/Runtime/UMG/Public/Components/ListView.h
--- a/Engine/Source/Runtime/UMG/Public/Components/ListView.h
+++ b/Engine/Source/Runtime/UMG/Public/Components/ListView.h
@@ -117,6 +117,12 @@ public:
 protected:
 	virtual void OnItemsChanged(const TArray<UObject*>& AddedItems, const TArray<UObject*>& RemovedItems);
 
+	UFUNCTION()
+	void OnListItemEndPlayed(AActor* Item, EEndPlayReason::Type EndPlayReason);
+
+	UFUNCTION()
+	void OnListItemOuterEndPlayed(AActor* ItemOuter, EEndPlayReason::Type EndPlayReason);
+
 	virtual TSharedRef<STableViewBase> RebuildListWidget() override;
 	virtual void HandleListEntryHovered(UUserWidget& EntryWidget) override;
 	virtual void HandleListEntryUnhovered(UUserWidget& EntryWidget) override;
