Direct Input Delegate Support inside of Widget Blueprints. This allows users to place the normal "Input Action" and "Input Axis" event nodes in a widget and they will behave as they would in an actor as long as there is an owning Player controller.

diff --git a/Engine/Source/Editor/UMGEditor/Private/WidgetBlueprint.cpp b/Engine/Source/Editor/UMGEditor/Private/WidgetBlueprint.cpp
--- a/Engine/Source/Editor/UMGEditor/Private/WidgetBlueprint.cpp
+++ b/Engine/Source/Editor/UMGEditor/Private/WidgetBlueprint.cpp
@@ -935,6 +935,11 @@ bool UWidgetBlueprint::AllowsDynamicBinding() const
 	return true;
 }
 
+bool UWidgetBlueprint::SupportsInputEvents() const
+{
+	return true;
+}
+
 void UWidgetBlueprint::GatherDependencies(TSet<TWeakObjectPtr<UBlueprint>>& InDependencies) const
 {
 	Super::GatherDependencies(InDependencies);
diff --git a/Engine/Source/Editor/UMGEditor/Public/WidgetBlueprint.h b/Engine/Source/Editor/UMGEditor/Public/WidgetBlueprint.h
--- a/Engine/Source/Editor/UMGEditor/Public/WidgetBlueprint.h
+++ b/Engine/Source/Editor/UMGEditor/Public/WidgetBlueprint.h
@@ -296,6 +296,8 @@ public:
 
 	virtual bool AllowsDynamicBinding() const override;
 
+	virtual bool SupportsInputEvents() const override;
+
 	virtual bool SupportedByDefaultBlueprintFactory() const override
 	{
 		return false;
diff --git a/Engine/Source/Runtime/UMG/Private/UserWidget.cpp b/Engine/Source/Runtime/UMG/Private/UserWidget.cpp
--- a/Engine/Source/Runtime/UMG/Private/UserWidget.cpp
+++ b/Engine/Source/Runtime/UMG/Private/UserWidget.cpp
@@ -31,6 +31,7 @@
 #include "UObject/Package.h"
 #include "Editor/WidgetCompilerLog.h"
 #include "GameFramework/InputSettings.h"
+#include "Engine/InputDelegateBinding.h"
 
 #define LOCTEXT_NAMESPACE "UMG"
 
@@ -1349,6 +1350,12 @@ void UUserWidget::BindToAnimationEvent(UWidgetAnimation* InAnimation, FWidgetAni
 
 void UUserWidget::NativeOnInitialized()
 {
+	// Bind any input delegates that may be on this widget to its owning player controller
+	if(APlayerController* PC = GetOwningPlayer())
+	{
+		UInputDelegateBinding::BindInputDelegates(GetClass(), PC->InputComponent, this);		
+	}
+	
 	if (UWidgetBlueprintGeneratedClass* Class = GetWidgetTreeOwningClass())
 	{
 		Class->ForEachExtension([this](UWidgetBlueprintGeneratedClassExtension* Extension)
@@ -1608,7 +1615,11 @@ void UUserWidget::InitializeInputComponent()
 {
 	if ( APlayerController* Controller = GetOwningPlayer() )
 	{
-		InputComponent = NewObject< UInputComponent >( this, UInputSettings::GetDefaultInputComponentClass(), NAME_None, RF_Transient );
+		// Use the existing PC's input class, or fallback to the project default. We should use the existing class
+		// instead of just the default one because if you have a plugin that has a PC with a different default input
+		// class then this would fail
+		UClass* InputClass = Controller->PlayerInput ? Controller->PlayerInput->GetClass() : UInputSettings::GetDefaultInputComponentClass();
+		InputComponent = NewObject< UInputComponent >( this, InputClass, NAME_None, RF_Transient );
 		InputComponent->bBlockInput = bStopAction;
 		InputComponent->Priority = Priority;
 		Controller->PushInputComponent( InputComponent );
