Expose IndexInList property on ListViews

SetIndexInList() is being called on all items in a menu, but was inaccessible to subclasses. This change exposes that property using established patterns in the base class

#tests added logging to every click in the sidebar, verified that the new function was returning the absolute position in the list.

diff --git a/Engine/Source/Runtime/Slate/Public/Widgets/Views/ITableRow.h b/Engine/Source/Runtime/Slate/Public/Widgets/Views/ITableRow.h
--- a/Engine/Source/Runtime/Slate/Public/Widgets/Views/ITableRow.h
+++ b/Engine/Source/Runtime/Slate/Public/Widgets/Views/ITableRow.h
@@ -32,6 +32,9 @@ public:
 	*/
 	virtual void SetIndexInList(int32 InIndexInList) = 0;
 
+	/** @return  The index of the item for which this widget was generated */
+	virtual int32 GetIndexInList() = 0;
+
 	/** @return true if the corresponding item is expanded; false otherwise*/
 	virtual bool IsItemExpanded() const = 0;
 
diff --git a/Engine/Source/Runtime/Slate/Public/Widgets/Views/STableRow.h b/Engine/Source/Runtime/Slate/Public/Widgets/Views/STableRow.h
--- a/Engine/Source/Runtime/Slate/Public/Widgets/Views/STableRow.h
+++ b/Engine/Source/Runtime/Slate/Public/Widgets/Views/STableRow.h
@@ -803,6 +803,11 @@ public:
 		IndexInList = InIndexInList;
 	}
 
+	virtual int32 GetIndexInList() override
+	{
+		return IndexInList;
+	}
+
 	virtual bool IsItemExpanded() const override
 	{
 		TSharedRef< ITypedTableView<ItemType> > OwnerTable = OwnerTablePtr.Pin().ToSharedRef();
diff --git a/Engine/Source/Runtime/UMG/Public/Components/ListViewBase.h b/Engine/Source/Runtime/UMG/Public/Components/ListViewBase.h
--- a/Engine/Source/Runtime/UMG/Public/Components/ListViewBase.h
+++ b/Engine/Source/Runtime/UMG/Public/Components/ListViewBase.h
@@ -132,6 +132,16 @@ public:
 		return nullptr;
 	}
 
+	int32 GetIndexInList(const ItemType& Item) const
+	{
+		if (SListView<ItemType>* MyListView = GetMyListView())
+		{
+			TSharedPtr<ITableRow> RowWidget = MyListView->WidgetFromItem(Item);
+			return RowWidget->GetIndexInList();
+		}
+		return INDEX_NONE;
+	}
+
 	int32 GetSelectedItems(TArray<ItemType>& OutSelectedItems) const
 	{
 		SListView<ItemType>* MyListView = GetMyListView();
diff --git a/Engine/Source/Runtime/UMG/Public/Slate/SObjectTableRow.h b/Engine/Source/Runtime/UMG/Public/Slate/SObjectTableRow.h
--- a/Engine/Source/Runtime/UMG/Public/Slate/SObjectTableRow.h
+++ b/Engine/Source/Runtime/UMG/Public/Slate/SObjectTableRow.h
@@ -164,6 +164,7 @@ public:
 
 	virtual TSharedRef<SWidget> AsWidget() override { return SharedThis(this); }
 	virtual void SetIndexInList(int32 InIndexInList) override { IndexInList = InIndexInList; }
+	virtual int32 GetIndexInList() override { return IndexInList; }
 	virtual TSharedPtr<SWidget> GetContent() override { return ChildSlot.GetChildAt(0); }
 	virtual int32 GetIndentLevel() const override { return OwnerTablePtr.Pin()->Private_GetNestingDepth(IndexInList); }
 	virtual int32 DoesItemHaveChildren() const override { return OwnerTablePtr.Pin()->Private_DoesItemHaveChildren(IndexInList); }
