Make WidgetNavigation Resolve and TryRename to allow for customization of navigation behavior ina plugin for UMG.

PR #8750: Making ResolveRules to virtual method. (Contributed by [REDACTED_PERSON])

diff --git a/Engine/Source/Runtime/UMG/Public/Blueprint/WidgetNavigation.h b/Engine/Source/Runtime/UMG/Public/Blueprint/WidgetNavigation.h
--- a/Engine/Source/Runtime/UMG/Public/Blueprint/WidgetNavigation.h
+++ b/Engine/Source/Runtime/UMG/Public/Blueprint/WidgetNavigation.h
@@ -86,13 +86,13 @@ public:
 	/**  */
 	EUINavigationRule GetNavigationRule(EUINavigation Nav);
 
-	/** Try to rename any explicit or custom bindings from an old to a new name. */
-	void TryToRenameBinding(FName OldName, FName NewName);
+	/** Try to rename any explicit or custom bindings from an old to a new name. This method can be overridden to customize resolving rules. */
+	virtual void TryToRenameBinding(FName OldName, FName NewName);
 
 #endif
 
-	/** Resolve widget names */
-	void ResolveRules(class UUserWidget* Outer, class UWidgetTree* WidgetTree);
+	/** Resolve widget names. This method can be overridden to customize resolving rules. */
+	virtual void ResolveRules(class UUserWidget* Outer, class UWidgetTree* WidgetTree);
 
 	/** Updates a slate metadata object to match this configured navigation ruleset. */
 	void UpdateMetaData(TSharedRef<FNavigationMetaData> MetaData);
