PR #8311: Added arrow key shortcuts to change order of Widgets in Panels (Contributed by [REDACTED_PERSON])

diff --git a/Engine/Source/Editor/UMGEditor/Private/Designer/SDesignerView.cpp b/Engine/Source/Editor/UMGEditor/Private/Designer/SDesignerView.cpp
--- a/Engine/Source/Editor/UMGEditor/Private/Designer/SDesignerView.cpp
+++ b/Engine/Source/Editor/UMGEditor/Private/Designer/SDesignerView.cpp
@@ -2012,6 +2012,9 @@ FReply SDesignerView::NudgeSelectedWidget(FVector2D Nudge)
 					if (TemplateSlot->NudgeByDesigner(Nudge, WidgetDesignerSettings->GridSnapEnabled ? TOptional<int32>(WidgetDesignerSettings->GridSnapSize) : TOptional<int32>()))
 					{
 						PreviewSlot->SynchronizeFromTemplate(TemplateSlot);
+						
+						UWidgetBlueprint* Blueprint = GetBlueprint();
+						FBlueprintEditorUtils::MarkBlueprintAsStructurallyModified(Blueprint);
 					}
 					// Nudge failed, cancel transaction.
 					else
diff --git a/Engine/Source/Runtime/UMG/Private/Components/GridSlot.cpp b/Engine/Source/Runtime/UMG/Private/Components/GridSlot.cpp
--- a/Engine/Source/Runtime/UMG/Private/Components/GridSlot.cpp
+++ b/Engine/Source/Runtime/UMG/Private/Components/GridSlot.cpp
@@ -134,3 +134,33 @@ void UGridSlot::SynchronizeProperties()
 
 	SetLayer(Layer);
 }
+
+#if WITH_EDITOR
+
+bool UGridSlot::NudgeByDesigner(const FVector2D& NudgeDirection, const TOptional<int32>& GridSnapSize)
+{
+	const FVector2D ClampedDirection = NudgeDirection.ClampAxes(-1.0f, 1.0f);
+	const int32 NewColumn = Column + ClampedDirection.X;
+	const int32 NewRow = Row + ClampedDirection.Y;
+
+	if (NewColumn < 0 || NewRow < 0 || (NewColumn == Column && NewRow == Row))
+	{
+		return false;
+	}
+	
+	Modify();
+
+	SetRow(NewRow);
+	SetColumn(NewColumn);
+
+	return true;
+}
+
+void UGridSlot::SynchronizeFromTemplate(const UPanelSlot* const TemplateSlot)
+{
+	const ThisClass* const TemplateGridSlot = CastChecked<ThisClass>(TemplateSlot);
+	SetRow(TemplateGridSlot->Row);
+	SetColumn(TemplateGridSlot->Column);
+}
+
+#endif
diff --git a/Engine/Source/Runtime/UMG/Private/Components/HorizontalBoxSlot.cpp b/Engine/Source/Runtime/UMG/Private/Components/HorizontalBoxSlot.cpp
--- a/Engine/Source/Runtime/UMG/Private/Components/HorizontalBoxSlot.cpp
+++ b/Engine/Source/Runtime/UMG/Private/Components/HorizontalBoxSlot.cpp
@@ -2,6 +2,7 @@
 
 #include "Components/HorizontalBoxSlot.h"
 #include "Components/Widget.h"
+#include "Components/HorizontalBox.h"
 
 /////////////////////////////////////////////////////
 // UHorizontalBoxSlot
@@ -78,3 +79,40 @@ void UHorizontalBoxSlot::SynchronizeProperties()
 	SetHorizontalAlignment(HorizontalAlignment);
 	SetVerticalAlignment(VerticalAlignment);
 }
+
+#if WITH_EDITOR
+
+bool UHorizontalBoxSlot::NudgeByDesigner(const FVector2D& NudgeDirection, const TOptional<int32>& GridSnapSize)
+{
+	if (NudgeDirection.X == 0.0f)
+	{
+		return false;
+	}
+	
+	const FVector2D ClampedDirection = NudgeDirection.ClampAxes(-1, 1);
+	UHorizontalBox* ParentHorizontalBox = CastChecked<UHorizontalBox>(Parent);
+
+	const int32 CurrentIndex = ParentHorizontalBox->GetChildIndex(Content);
+
+	if ((CurrentIndex == 0 && ClampedDirection.X < 0.0f) ||
+		(CurrentIndex + 1 >= ParentHorizontalBox->GetChildrenCount() && ClampedDirection.X > 0.0f))
+	{
+		return false;
+	}
+
+	ParentHorizontalBox->Modify();
+	ParentHorizontalBox->ShiftChild(CurrentIndex + ClampedDirection.X, Content);
+
+	return true;
+}
+
+void UHorizontalBoxSlot::SynchronizeFromTemplate(const UPanelSlot* const TemplateSlot)
+{
+	const ThisClass* const TemplateHorizontalBoxSlot = CastChecked<ThisClass>(TemplateSlot);
+	const int32 CurrentIndex = TemplateHorizontalBoxSlot->Parent->GetChildIndex(TemplateHorizontalBoxSlot->Content);
+
+	UHorizontalBox* ParentHorizontalBox = CastChecked<UHorizontalBox>(Parent);
+	ParentHorizontalBox->ShiftChild(CurrentIndex, Content);
+}
+
+#endif
diff --git a/Engine/Source/Runtime/UMG/Private/Components/UniformGridSlot.cpp b/Engine/Source/Runtime/UMG/Private/Components/UniformGridSlot.cpp
--- a/Engine/Source/Runtime/UMG/Private/Components/UniformGridSlot.cpp
+++ b/Engine/Source/Runtime/UMG/Private/Components/UniformGridSlot.cpp
@@ -75,3 +75,33 @@ void UUniformGridSlot::SynchronizeProperties()
 	SetHorizontalAlignment(HorizontalAlignment);
 	SetVerticalAlignment(VerticalAlignment);
 }
+
+#if WITH_EDITOR
+
+bool UUniformGridSlot::NudgeByDesigner(const FVector2D& NudgeDirection, const TOptional<int32>& GridSnapSize)
+{
+	const FVector2D ClampedDirection = NudgeDirection.ClampAxes(-1.0f, 1.0f);
+	const int32 NewColumn = Column + ClampedDirection.X;
+	const int32 NewRow = Row + ClampedDirection.Y;
+
+	if (NewColumn < 0 || NewRow < 0 || (NewColumn == Column && NewRow == Row))
+	{
+		return false;
+	}
+
+	Modify();
+
+	SetRow(NewRow);
+	SetColumn(NewColumn);
+
+	return true;
+}
+
+void UUniformGridSlot::SynchronizeFromTemplate(const UPanelSlot* const TemplateSlot)
+{
+	const ThisClass* const TemplateUniformGridSlot = CastChecked<ThisClass>(TemplateSlot);
+	SetRow(TemplateUniformGridSlot->Row);
+	SetColumn(TemplateUniformGridSlot->Column);
+}
+
+#endif
diff --git a/Engine/Source/Runtime/UMG/Private/Components/VerticalBoxSlot.cpp b/Engine/Source/Runtime/UMG/Private/Components/VerticalBoxSlot.cpp
--- a/Engine/Source/Runtime/UMG/Private/Components/VerticalBoxSlot.cpp
+++ b/Engine/Source/Runtime/UMG/Private/Components/VerticalBoxSlot.cpp
@@ -2,6 +2,7 @@
 
 #include "Components/VerticalBoxSlot.h"
 #include "Components/Widget.h"
+#include "Components/VerticalBox.h"
 
 /////////////////////////////////////////////////////
 // UVerticalBoxSlot
@@ -79,3 +80,40 @@ void UVerticalBoxSlot::SynchronizeProperties()
 	SetHorizontalAlignment(HorizontalAlignment);
 	SetVerticalAlignment(VerticalAlignment);
 }
+
+#if WITH_EDITOR
+
+bool UVerticalBoxSlot::NudgeByDesigner(const FVector2D& NudgeDirection, const TOptional<int32>& GridSnapSize)
+{
+	if (NudgeDirection.Y == 0)
+	{
+		return false;
+	}
+	
+	const FVector2D ClampedDirection = NudgeDirection.ClampAxes(-1, 1);
+	UVerticalBox* ParentVerticalBox = CastChecked<UVerticalBox>(Parent);
+
+	const int32 CurrentIndex = ParentVerticalBox->GetChildIndex(Content);
+
+	if ((CurrentIndex == 0 && ClampedDirection.Y < 0.0f) ||
+		(CurrentIndex + 1 >= ParentVerticalBox->GetChildrenCount() && ClampedDirection.Y > 0.0f))
+	{
+		return false;
+	}
+
+	ParentVerticalBox->Modify();
+	ParentVerticalBox->ShiftChild(CurrentIndex + ClampedDirection.Y, Content);
+
+	return true;
+}
+
+void UVerticalBoxSlot::SynchronizeFromTemplate(const UPanelSlot* const TemplateSlot)
+{
+	const ThisClass* const TemplateVerticalBoxSlot = CastChecked<ThisClass>(TemplateSlot);
+	const int32 CurrentIndex = TemplateVerticalBoxSlot->Parent->GetChildIndex(TemplateVerticalBoxSlot->Content);
+
+	UVerticalBox* ParentVerticalBox = CastChecked<UVerticalBox>(Parent);
+	ParentVerticalBox->ShiftChild(CurrentIndex, Content);
+}
+
+#endif
diff --git a/Engine/Source/Runtime/UMG/Public/Components/GridSlot.h b/Engine/Source/Runtime/UMG/Public/Components/GridSlot.h
--- a/Engine/Source/Runtime/UMG/Public/Components/GridSlot.h
+++ b/Engine/Source/Runtime/UMG/Public/Components/GridSlot.h
@@ -104,6 +104,11 @@ public:
 
 	virtual void ReleaseSlateResources(bool bReleaseChildren) override;
 
+#if WITH_EDITOR
+	virtual bool NudgeByDesigner(const FVector2D& NudgeDirection, const TOptional<int32>& GridSnapSize) override;
+	virtual void SynchronizeFromTemplate(const UPanelSlot* const TemplateSlot) override;
+#endif //WITH_EDITOR
+
 private:
 	//TODO UMG Slots should hold weak or shared refs to slots.
 
diff --git a/Engine/Source/Runtime/UMG/Public/Components/HorizontalBoxSlot.h b/Engine/Source/Runtime/UMG/Public/Components/HorizontalBoxSlot.h
--- a/Engine/Source/Runtime/UMG/Public/Components/HorizontalBoxSlot.h
+++ b/Engine/Source/Runtime/UMG/Public/Components/HorizontalBoxSlot.h
@@ -56,4 +56,8 @@ public:
 
 	virtual void ReleaseSlateResources(bool bReleaseChildren) override;
 
+#if WITH_EDITOR
+	virtual bool NudgeByDesigner(const FVector2D& NudgeDirection, const TOptional<int32>& GridSnapSize) override;
+	virtual void SynchronizeFromTemplate(const UPanelSlot* const TemplateSlot) override;
+#endif //WITH_EDITOR
 };
diff --git a/Engine/Source/Runtime/UMG/Public/Components/UniformGridSlot.h b/Engine/Source/Runtime/UMG/Public/Components/UniformGridSlot.h
--- a/Engine/Source/Runtime/UMG/Public/Components/UniformGridSlot.h
+++ b/Engine/Source/Runtime/UMG/Public/Components/UniformGridSlot.h
@@ -66,6 +66,11 @@ public:
 
 	virtual void ReleaseSlateResources(bool bReleaseChildren) override;
 
+#if WITH_EDITOR
+	virtual bool NudgeByDesigner(const FVector2D& NudgeDirection, const TOptional<int32>& GridSnapSize) override;
+	virtual void SynchronizeFromTemplate(const UPanelSlot* const TemplateSlot) override;
+#endif //WITH_EDITOR
+
 private:
 	//TODO UMG Slots should hold weak or shared refs to slots.
 
diff --git a/Engine/Source/Runtime/UMG/Public/Components/VerticalBoxSlot.h b/Engine/Source/Runtime/UMG/Public/Components/VerticalBoxSlot.h
--- a/Engine/Source/Runtime/UMG/Public/Components/VerticalBoxSlot.h
+++ b/Engine/Source/Runtime/UMG/Public/Components/VerticalBoxSlot.h
@@ -68,4 +68,9 @@ public:
 
 	/** Builds the underlying FSlot for the Slate layout panel. */
 	void BuildSlot(TSharedRef<SVerticalBox> InVerticalBox);
+
+#if WITH_EDITOR
+	virtual bool NudgeByDesigner(const FVector2D& NudgeDirection, const TOptional<int32>& GridSnapSize) override;
+	virtual void SynchronizeFromTemplate(const UPanelSlot* const TemplateSlot) override;
+#endif //WITH_EDITOR
 };
