Rebuild MyListView list after clearing list items in UListView.

Fixes issue where a previously-first element in a UListView had incorrect padding farther down the list

diff --git a/Engine/Source/Runtime/UMG/Private/Components/ListView.cpp b/Engine/Source/Runtime/UMG/Private/Components/ListView.cpp
--- a/Engine/Source/Runtime/UMG/Private/Components/ListView.cpp
+++ b/Engine/Source/Runtime/UMG/Private/Components/ListView.cpp
@@ -87,9 +87,8 @@ void UListView::AddItem(UObject* Item)
 {
 	ListItems.Add(Item);
 
-	TArray<UObject*> Added;
-	TArray<UObject*> Removed;
-	Added.Add(Item);
+	const TArray<UObject*> Added = { Item };
+	const TArray<UObject*> Removed;
 	OnItemsChanged(Added, Removed);
 
 	RequestRefresh();
@@ -99,9 +98,8 @@ void UListView::RemoveItem(UObject* Item)
 {
 	ListItems.Remove(Item);
 
-	TArray<UObject*> Added;
-	TArray<UObject*> Removed;
-	Removed.Add(Item);
+	const TArray<UObject*> Added;
+	const TArray<UObject*> Removed = { Item };
 	OnItemsChanged(Added, Removed);
 
 	RequestRefresh();
@@ -124,11 +122,16 @@ int32 UListView::GetIndexForItem(const UObject* Item) const
 
 void UListView::ClearListItems()
 {
-	TArray<UObject*> Added;
-	TArray<UObject*> Removed = MoveTemp(ListItems);
+	const TArray<UObject*> Added;
+	const TArray<UObject*> Removed = MoveTemp(ListItems);
 
 	ListItems.Reset();
 
+	if (MyListView)
+	{
+		MyListView->RebuildList();
+	}
+
 	OnItemsChanged(Added, Removed);
 
 	RequestRefresh();
