Undo changelist 18605401

diff --git a/Engine/Source/Runtime/UMG/Private/Components/ListView.cpp b/Engine/Source/Runtime/UMG/Private/Components/ListView.cpp
--- a/Engine/Source/Runtime/UMG/Private/Components/ListView.cpp
+++ b/Engine/Source/Runtime/UMG/Private/Components/ListView.cpp
@@ -87,8 +87,9 @@ void UListView::AddItem(UObject* Item)
 {
 	ListItems.Add(Item);
 
-	const TArray<UObject*> Added = { Item };
-	const TArray<UObject*> Removed;
+	TArray<UObject*> Added;
+	TArray<UObject*> Removed;
+	Added.Add(Item);
 	OnItemsChanged(Added, Removed);
 
 	RequestRefresh();
@@ -98,8 +99,9 @@ void UListView::RemoveItem(UObject* Item)
 {
 	ListItems.Remove(Item);
 
-	const TArray<UObject*> Added;
-	const TArray<UObject*> Removed = { Item };
+	TArray<UObject*> Added;
+	TArray<UObject*> Removed;
+	Removed.Add(Item);
 	OnItemsChanged(Added, Removed);
 
 	RequestRefresh();
@@ -122,16 +124,11 @@ int32 UListView::GetIndexForItem(const UObject* Item) const
 
 void UListView::ClearListItems()
 {
-	const TArray<UObject*> Added;
-	const TArray<UObject*> Removed = MoveTemp(ListItems);
+	TArray<UObject*> Added;
+	TArray<UObject*> Removed = MoveTemp(ListItems);
 
 	ListItems.Reset();
 
-	if (MyListView)
-	{
-		MyListView->RebuildList();
-	}
-
 	OnItemsChanged(Added, Removed);
 
 	RequestRefresh();
