Add check in GetIndexInList for a valid WidgetFromItem value

#tests reproduced the issue with an invite rejection and verified that the fix solves it.

diff --git a/Engine/Source/Runtime/UMG/Public/Components/ListViewBase.h b/Engine/Source/Runtime/UMG/Public/Components/ListViewBase.h
--- a/Engine/Source/Runtime/UMG/Public/Components/ListViewBase.h
+++ b/Engine/Source/Runtime/UMG/Public/Components/ListViewBase.h
@@ -136,8 +136,10 @@ public:
 	{
 		if (SListView<ItemType>* MyListView = GetMyListView())
 		{
-			TSharedPtr<ITableRow> RowWidget = MyListView->WidgetFromItem(Item);
-			return RowWidget->GetIndexInList();
+			if (TSharedPtr<ITableRow> RowWidget = MyListView->WidgetFromItem(Item))
+			{
+				return RowWidget->GetIndexInList();
+			}
 		}
 		return INDEX_NONE;
 	}
