Crash fix. Undoing ListView refresh

diff --git a/Engine/Source/Runtime/UMG/Private/Components/ListView.cpp b/Engine/Source/Runtime/UMG/Private/Components/ListView.cpp
--- a/Engine/Source/Runtime/UMG/Private/Components/ListView.cpp
+++ b/Engine/Source/Runtime/UMG/Private/Components/ListView.cpp
@@ -130,11 +130,6 @@ void UListView::ClearListItems()
 	OnItemsChanged(Added, Removed);
 
 	RequestRefresh();
-
-	if (MyListView)
-	{
-		MyListView->RebuildList();
-	}
 }
 
 void UListView::SetSelectionMode(TEnumAsByte<ESelectionMode::Type> InSelectionMode)
