RichText: add a BP node that creates/returns a dynamic material for the default text style of a rich text block

diff --git a/Engine/Source/Runtime/UMG/Private/Components/RichTextBlock.cpp b/Engine/Source/Runtime/UMG/Private/Components/RichTextBlock.cpp
--- a/Engine/Source/Runtime/UMG/Private/Components/RichTextBlock.cpp
+++ b/Engine/Source/Runtime/UMG/Private/Components/RichTextBlock.cpp
@@ -284,6 +284,30 @@ void URichTextBlock::ClearAllDefaultStyleOverrides()
 	}
 }
 
+UMaterialInstanceDynamic* URichTextBlock::GetDefaultDynamicMaterial()
+{
+	TObjectPtr<UObject>& MaterialObject = bOverrideDefaultStyle ? DefaultTextStyleOverride.Font.FontMaterial : DefaultTextStyle.Font.FontMaterial;
+
+	if (UMaterialInterface* Material = Cast<UMaterialInterface>(MaterialObject.Get()))
+	{
+		UMaterialInstanceDynamic* DynamicMaterial = Cast<UMaterialInstanceDynamic>(Material);
+
+		if (!DynamicMaterial)
+		{
+			BeginDefaultStyleOverride();
+
+			DynamicMaterial = UMaterialInstanceDynamic::Create(Material, this);
+			DefaultTextStyleOverride.Font.FontMaterial = DynamicMaterial;
+			
+			ApplyUpdatedDefaultTextStyle();
+		}
+
+		return DynamicMaterial;
+	}
+
+	return nullptr;
+}
+
 void URichTextBlock::SetDefaultColorAndOpacity(FSlateColor InColorAndOpacity)
 {
 	BeginDefaultStyleOverride();
diff --git a/Engine/Source/Runtime/UMG/Public/Components/RichTextBlock.h b/Engine/Source/Runtime/UMG/Public/Components/RichTextBlock.h
--- a/Engine/Source/Runtime/UMG/Public/Components/RichTextBlock.h
+++ b/Engine/Source/Runtime/UMG/Public/Components/RichTextBlock.h
@@ -113,6 +113,13 @@ public:
 	UFUNCTION()
 	void ClearAllDefaultStyleOverrides();
 
+	/**
+	 * Creates a dynamic material for the default font or returns it if it already
+	 * exists
+	 */
+	UFUNCTION(BlueprintCallable, Category = "Appearance")
+	UMaterialInstanceDynamic* GetDefaultDynamicMaterial();
+
 public:
 	URichTextBlock(const FObjectInitializer& ObjectInitializer);
 
