Rich Text Block: add missing include

diff --git a/Engine/Source/Runtime/UMG/Private/Components/RichTextBlock.cpp b/Engine/Source/Runtime/UMG/Private/Components/RichTextBlock.cpp
--- a/Engine/Source/Runtime/UMG/Private/Components/RichTextBlock.cpp
+++ b/Engine/Source/Runtime/UMG/Private/Components/RichTextBlock.cpp
@@ -13,6 +13,7 @@
 #include "Framework/Text/IRichTextMarkupWriter.h"
 #include "RenderingThread.h"
 #include "Editor/WidgetCompilerLog.h"
+#include "Materials/MaterialInstanceDynamic.h"
 
 #define LOCTEXT_NAMESPACE "UMG"
 
diff --git a/Engine/Source/Runtime/UMG/Public/Components/RichTextBlock.h b/Engine/Source/Runtime/UMG/Public/Components/RichTextBlock.h
--- a/Engine/Source/Runtime/UMG/Public/Components/RichTextBlock.h
+++ b/Engine/Source/Runtime/UMG/Public/Components/RichTextBlock.h
@@ -12,6 +12,7 @@
 #include "RichTextBlock.generated.h"
 
 class SRichTextBlock;
+class UMaterialInstanceDynamic;
 class URichTextBlockDecorator;
 
 /** Simple struct for rich text styles */
