Explicitly set the blender system on the FMargin property.

#tests Compiled DevelopmentEditor Win64 / Launched PIE instance and opened a Social profile which previously triggered an ensure in UMovieScenePropertyInstantiatorSystem::InitializeBlendPath. Verified no ensure.

diff --git a/Engine/Source/Runtime/UMG/Private/Animation/MovieSceneUMGComponentTypes.cpp b/Engine/Source/Runtime/UMG/Private/Animation/MovieSceneUMGComponentTypes.cpp
--- a/Engine/Source/Runtime/UMG/Private/Animation/MovieSceneUMGComponentTypes.cpp
+++ b/Engine/Source/Runtime/UMG/Private/Animation/MovieSceneUMGComponentTypes.cpp
@@ -86,6 +86,7 @@ FMovieSceneUMGComponentTypes::FMovieSceneUMGComponentTypes()
 	.AddComposite(BuiltInComponents->FloatResult[1], &FMargin::Top)
 	.AddComposite(BuiltInComponents->FloatResult[2], &FMargin::Right)
 	.AddComposite(BuiltInComponents->FloatResult[3], &FMargin::Bottom)
+	.SetBlenderSystem<UMovieScenePiecewiseFloatBlenderSystem>()
 	.Commit();
 
 	BuiltInComponents->PropertyRegistry.DefineCompositeProperty(WidgetTransform)
