Fix an issue where you could crash when binding input delegates to widgets because we were using the PlayerInput class instead of the InputComponent class.

#rnx

diff --git a/Engine/Source/Runtime/UMG/Private/UserWidget.cpp b/Engine/Source/Runtime/UMG/Private/UserWidget.cpp
--- a/Engine/Source/Runtime/UMG/Private/UserWidget.cpp
+++ b/Engine/Source/Runtime/UMG/Private/UserWidget.cpp
@@ -1618,7 +1618,7 @@ void UUserWidget::InitializeInputComponent()
 		// Use the existing PC's input class, or fallback to the project default. We should use the existing class
 		// instead of just the default one because if you have a plugin that has a PC with a different default input
 		// class then this would fail
-		UClass* InputClass = Controller->PlayerInput ? Controller->PlayerInput->GetClass() : UInputSettings::GetDefaultInputComponentClass();
+		UClass* InputClass = Controller->InputComponent ? Controller->InputComponent->GetClass() : UInputSettings::GetDefaultInputComponentClass();
 		InputComponent = NewObject< UInputComponent >( this, InputClass, NAME_None, RF_Transient );
 		InputComponent->bBlockInput = bStopAction;
 		InputComponent->Priority = Priority;
