ListView: create event when listview finishes scrolling

diff --git a/Engine/Source/Runtime/Slate/Private/Widgets/Views/STableViewBase.cpp b/Engine/Source/Runtime/Slate/Private/Widgets/Views/STableViewBase.cpp
--- a/Engine/Source/Runtime/Slate/Private/Widgets/Views/STableViewBase.cpp
+++ b/Engine/Source/Runtime/Slate/Private/Widgets/Views/STableViewBase.cpp
@@ -345,6 +345,10 @@ void STableViewBase::Tick( const FGeometry& AllottedGeometry, const double InCur
 				// We call this rather than just leave bItemsNeedRefresh as true to ensure that EnsureTickToRefresh is registered
 				RequestLayoutRefresh();
 			}
+			else if (CurrentScrollOffset == TargetScrollOffset)
+			{
+				NotifyFinishedScrolling();
+			}
 		}
 	}
 }
diff --git a/Engine/Source/Runtime/Slate/Public/Framework/SlateDelegates.h b/Engine/Source/Runtime/Slate/Public/Framework/SlateDelegates.h
--- a/Engine/Source/Runtime/Slate/Public/Framework/SlateDelegates.h
+++ b/Engine/Source/Runtime/Slate/Public/Framework/SlateDelegates.h
@@ -121,6 +121,9 @@ DECLARE_DELEGATE_OneParam( FOnLinearColorValueChanged, FLinearColor )
 /** Notification for FVector4 value change */
 DECLARE_DELEGATE_OneParam(FOnVector4ValueChanged, FVector4&)
 
+/** Invoked when a list view has reached its target offset.  */
+DECLARE_DELEGATE(FOnFinishedScrolling);
+
 DECLARE_DELEGATE(FOnTableViewBadState);
 
 template< typename ArgumentType >
diff --git a/Engine/Source/Runtime/Slate/Public/Widgets/Views/SListView.h b/Engine/Source/Runtime/Slate/Public/Widgets/Views/SListView.h
--- a/Engine/Source/Runtime/Slate/Public/Widgets/Views/SListView.h
+++ b/Engine/Source/Runtime/Slate/Public/Widgets/Views/SListView.h
@@ -123,6 +123,8 @@ public:
 		SLATE_EVENT( FOnTableViewScrolled, OnListViewScrolled )
 
 		SLATE_EVENT( FOnItemScrolledIntoView, OnItemScrolledIntoView )
+		
+		SLATE_EVENT( FOnFinishedScrolling, OnFinishedScrolling )
 
 		SLATE_ARGUMENT( const TArray<ItemType>* , ListItemsSource )
 
@@ -199,6 +201,7 @@ public:
 		this->OnEntryInitialized = InArgs._OnEntryInitialized;
 		this->OnRowReleased = InArgs._OnRowReleased;
 		this->OnItemScrolledIntoView = InArgs._OnItemScrolledIntoView;
+		this->OnFinishedScrolling = InArgs._OnFinishedScrolling;
 
 		this->ItemsSource = InArgs._ListItemsSource;
 		this->OnContextMenuOpening = InArgs._OnContextMenuOpening;
@@ -1887,6 +1890,11 @@ protected:
 		}
 	}
 
+	virtual void NotifyFinishedScrolling() override
+	{
+		OnFinishedScrolling.ExecuteIfBound();
+	}
+
 	virtual float ScrollBy( const FGeometry& MyGeometry, float ScrollByAmountInSlateUnits, EAllowOverscroll InAllowOverscroll ) override
 	{
 		auto DoubleFractional = [](double Value) -> double
@@ -2153,6 +2161,9 @@ protected:
 	/** Delegate to be invoked when an item has come into view after it was requested to come into view. */
 	FOnItemScrolledIntoView OnItemScrolledIntoView;
 
+	/** Delegate to be invoked when TargetScrollOffset is reached at the end of a ::Tick. */
+	FOnFinishedScrolling OnFinishedScrolling;
+
 	/** A set of selected data items */
 	TItemSet SelectedItems;
 
diff --git a/Engine/Source/Runtime/Slate/Public/Widgets/Views/STableViewBase.h b/Engine/Source/Runtime/Slate/Public/Widgets/Views/STableViewBase.h
--- a/Engine/Source/Runtime/Slate/Public/Widgets/Views/STableViewBase.h
+++ b/Engine/Source/Runtime/Slate/Public/Widgets/Views/STableViewBase.h
@@ -357,6 +357,11 @@ protected:
 	 */
 	virtual void NotifyItemScrolledIntoView() = 0;
 
+	/**
+	 * Called when CurrentScrollOffset == TargetScrollOffset at the end of a ::Tick
+	 */
+	virtual void NotifyFinishedScrolling() = 0;
+
 	/** Util Function so templates classes don't need to include SlateApplication */
 	void NavigateToWidget(const uint32 UserIndex, const TSharedPtr<SWidget>& NavigationDestination, ENavigationSource NavigationSource = ENavigationSource::FocusedWidget) const;
 
diff --git a/Engine/Source/Runtime/UMG/Public/Components/ListViewBase.h b/Engine/Source/Runtime/UMG/Public/Components/ListViewBase.h
--- a/Engine/Source/Runtime/UMG/Public/Components/ListViewBase.h
+++ b/Engine/Source/Runtime/UMG/Public/Components/ListViewBase.h
@@ -53,6 +53,9 @@ public:
 	DECLARE_MULTICAST_DELEGATE_TwoParams(FOnListViewScrolled, float, float);
 	virtual FOnListViewScrolled& OnListViewScrolled() const = 0;
 
+	DECLARE_MULTICAST_DELEGATE(FOnFinishedScrolling);
+	virtual FOnFinishedScrolling& OnFinishedScrolling() const = 0;
+
 	DECLARE_MULTICAST_DELEGATE_TwoParams(FOnItemScrolledIntoView, ItemType, UUserWidget&);
 	virtual FOnItemScrolledIntoView& OnItemScrolledIntoView() const = 0;
 
@@ -258,6 +261,7 @@ protected:
 			.OnRowReleased_UObject(Implementer, &UListViewBaseT::HandleRowReleased)
 			.OnItemScrolledIntoView_UObject(Implementer, &UListViewBaseT::HandleItemScrolledIntoView)
 			.OnListViewScrolled_UObject(Implementer, &UListViewBaseT::HandleListViewScrolled)
+			.OnFinishedScrolling_UObject(Implementer, &UListViewBaseT::HandleFinishedScrolling)
 			.OnMouseButtonClick_UObject(Implementer, &UListViewBaseT::HandleItemClicked)
 			.OnMouseButtonDoubleClick_UObject(Implementer, &UListViewBaseT::HandleItemDoubleClicked);
 	}
@@ -488,6 +492,11 @@ private:
 		}
 	}
 
+	void HandleFinishedScrolling()
+	{
+		OnFinishedScrolling().Broadcast();
+	}
+
 	void HandleItemScrolledIntoView(ItemType Item, const TSharedPtr<ITableRow>& InWidget)
 	{
 		UUserWidget* RowWidget = GetEntryWidgetFromItem(Item);
@@ -783,6 +792,7 @@ private:	\
 	mutable FOnItemIsHoveredChanged OnItemIsHoveredChangedEvent;	\
 	mutable FOnItemScrolledIntoView OnItemScrolledIntoViewEvent;	\
 	mutable FOnListViewScrolled OnListViewScrolledEvent;	\
+	mutable FOnFinishedScrolling OnFinishedScrollingEvent;	\
 	mutable FOnItemExpansionChanged OnItemExpansionChangedEvent;	\
 	mutable FOnGetEntryClassForItem OnGetEntryClassForItemDelegate;	\
 public:	\
@@ -793,5 +803,6 @@ public:	\
 	virtual FOnItemSelectionChanged& OnItemSelectionChanged() const override { return OnItemSelectionChangedEvent; }	\
 	virtual FOnItemScrolledIntoView& OnItemScrolledIntoView() const override { return OnItemScrolledIntoViewEvent; }	\
 	virtual FOnListViewScrolled& OnListViewScrolled() const override { return OnListViewScrolledEvent; }	\
+	virtual FOnFinishedScrolling& OnFinishedScrolling() const override { return OnFinishedScrollingEvent; }	\
 	virtual FOnItemExpansionChanged& OnItemExpansionChanged() const override { return OnItemExpansionChangedEvent; }	\
 	virtual FOnGetEntryClassForItem& OnGetEntryClassForItem() const override { return OnGetEntryClassForItemDelegate; }
\ No newline at end of file
