Fixed crash when calling Add Item on a List View in Pre Construct of an Editor Utility Widget.

UListView doesn't support duplicate or null items, and adding either causes a crash. We now report a warning and stop the erroneous input from crashing.

diff --git a/Engine/Source/Runtime/UMG/Private/Components/ListView.cpp b/Engine/Source/Runtime/UMG/Private/Components/ListView.cpp
--- a/Engine/Source/Runtime/UMG/Private/Components/ListView.cpp
+++ b/Engine/Source/Runtime/UMG/Private/Components/ListView.cpp
@@ -85,6 +85,18 @@ void UListView::OnRefreshDesignerItems()
 
 void UListView::AddItem(UObject* Item)
 {
+	if (Item == nullptr)
+	{
+		FFrame::KismetExecutionMessage(TEXT("Cannot add null item into ListView."), ELogVerbosity::Warning, "NullListViewItem");
+		return;
+	}
+
+	if (ListItems.Contains(Item))
+	{
+		FFrame::KismetExecutionMessage(TEXT("Cannot add duplicate item into ListView."), ELogVerbosity::Warning, "DuplicateListViewItem");
+		return;
+	}
+
 	ListItems.Add(Item);
 
 	const TArray<UObject*> Added = { Item };
@@ -248,7 +260,7 @@ void UListView::OnItemsChanged(const TArray<UObject*>& AddedItems, const TArray<
 
 	// Keep track of references to Actors and make sure to release them when Actors are about to be removed
 	for (UObject* AddedItem : AddedItems)
-	{		
+	{
 		if (AActor* AddedActor = Cast<AActor>(AddedItem))
 		{
 			AddedActor->OnEndPlay.AddDynamic(this, &UListView::OnListItemEndPlayed);
