Fix widget instances not updating in UMG designer on compile

diff --git a/Engine/Source/Runtime/UMG/Private/UserWidget.cpp b/Engine/Source/Runtime/UMG/Private/UserWidget.cpp
--- a/Engine/Source/Runtime/UMG/Private/UserWidget.cpp
+++ b/Engine/Source/Runtime/UMG/Private/UserWidget.cpp
@@ -202,7 +202,7 @@ void UUserWidget::DuplicateAndInitializeFromWidgetTree(UWidgetTree* InWidgetTree
 	if ( ensure(InWidgetTree) && !HasAnyFlags(RF_NeedPostLoad))
 	{
 		FObjectInstancingGraph ObjectInstancingGraph;
-		WidgetTree = NewObject<UWidgetTree>(this, InWidgetTree->GetClass(), TEXT("WidgetTree"), RF_Transactional, InWidgetTree, false, &ObjectInstancingGraph);
+		WidgetTree = NewObject<UWidgetTree>(this, InWidgetTree->GetClass(), NAME_None, RF_Transactional, InWidgetTree, false, &ObjectInstancingGraph);
 		WidgetTree->SetFlags(RF_Transient | RF_DuplicateTransient);
 
 		// After using the widget tree as a template, we need to loop over the instanced sub-objects and
