Added viewmodel update mode icons to UMGStyle.

diff --git a/Engine/Source/Runtime/UMG/Private/UMGStyle.cpp b/Engine/Source/Runtime/UMG/Private/UMGStyle.cpp
--- a/Engine/Source/Runtime/UMG/Private/UMGStyle.cpp
+++ b/Engine/Source/Runtime/UMG/Private/UMGStyle.cpp
@@ -160,6 +160,10 @@ TSharedRef< FSlateStyleSet > FUMGStyle::Create()
 	Style->Set("ClassThumbnail.DetailsView", new IMAGE_BRUSH_SVG(TEXT("DetailsView_64"), Icon64x64));
 	Style->Set("ClassThumbnail.SinglePropertyView", new IMAGE_BRUSH_SVG(TEXT("SinglePropertyView_64"), Icon64x64));
 
+	Style->Set("MVVM.OneWay", new IMAGE_BRUSH_SVG("ViewModel_OneWay", Icon16x16));
+	Style->Set("MVVM.TwoWay", new IMAGE_BRUSH_SVG("ViewModel_TwoWay", Icon16x16));
+	Style->Set("MVVM.OneTime", new IMAGE_BRUSH_SVG("ViewModel_OneTime", Icon16x16));
+
 	return Style;
 }
 
