Slate: Clear the uncached draw elements and the batches elements once they are rendered. We do this to remove hard references to the Glyph Sequence from the previous frame. Cached draw elements are still cached. An old Glyph Sequence may still be a reference in a cached element but it is the responsibility of the SWidget to do the proper invalidation. Deprecated the FSlateRenderer::GetDrawBuffer to remind the user to call FSlateRenderer::ReleaseDrawBuffer.

diff --git a/Engine/Source/Developer/StandaloneRenderer/Private/OpenGL/SlateOpenGLRenderer.cpp b/Engine/Source/Developer/StandaloneRenderer/Private/OpenGL/SlateOpenGLRenderer.cpp
--- a/Engine/Source/Developer/StandaloneRenderer/Private/OpenGL/SlateOpenGLRenderer.cpp
+++ b/Engine/Source/Developer/StandaloneRenderer/Private/OpenGL/SlateOpenGLRenderer.cpp
@@ -79,13 +79,23 @@ FSlateOpenGLRenderer::~FSlateOpenGLRenderer()
 }
 
 /** Returns a draw buffer that can be used by Slate windows to draw window elements */
-FSlateDrawBuffer& FSlateOpenGLRenderer::GetDrawBuffer()
+FSlateDrawBuffer& FSlateOpenGLRenderer::AcquireDrawBuffer()
 {
+	ensureMsgf(!DrawBuffer.IsLocked(), TEXT("The DrawBuffer is already locked. Make sure to ReleaseDrawBuffer the DrawBuffer"));
+	DrawBuffer.Lock();
+
 	// Clear out the buffer each time its accessed
 	DrawBuffer.ClearBuffer();
+
 	return DrawBuffer;
 }
 
+void FSlateOpenGLRenderer::ReleaseDrawBuffer(FSlateDrawBuffer& InWindowDrawBuffer)
+{
+	ensureMsgf(&DrawBuffer == &InWindowDrawBuffer, TEXT("It release a DrawBuffer that is not a member of the SlateNullRenderer"));
+	InWindowDrawBuffer.Unlock();
+}
+
 bool FSlateOpenGLRenderer::Initialize()
 {
 	SharedContext.Initialize( NULL, NULL );
diff --git a/Engine/Source/Developer/StandaloneRenderer/Private/OpenGL/SlateOpenGLRenderer.h b/Engine/Source/Developer/StandaloneRenderer/Private/OpenGL/SlateOpenGLRenderer.h
--- a/Engine/Source/Developer/StandaloneRenderer/Private/OpenGL/SlateOpenGLRenderer.h
+++ b/Engine/Source/Developer/StandaloneRenderer/Private/OpenGL/SlateOpenGLRenderer.h
@@ -146,7 +146,8 @@ public:
 	~FSlateOpenGLRenderer();
 
 	/** FSlateRenderer interface */
-	virtual FSlateDrawBuffer& GetDrawBuffer() override;
+	virtual FSlateDrawBuffer& AcquireDrawBuffer() override;
+	virtual void ReleaseDrawBuffer( FSlateDrawBuffer& InWindowDrawBuffer ) override;
 	virtual bool Initialize() override;
 	virtual void Destroy() override {}
 	virtual void DrawWindows( FSlateDrawBuffer& InWindowDrawBuffer ) override;
diff --git a/Engine/Source/Developer/StandaloneRenderer/Private/Windows/D3D/SlateD3DRenderer.cpp b/Engine/Source/Developer/StandaloneRenderer/Private/Windows/D3D/SlateD3DRenderer.cpp
--- a/Engine/Source/Developer/StandaloneRenderer/Private/Windows/D3D/SlateD3DRenderer.cpp
+++ b/Engine/Source/Developer/StandaloneRenderer/Private/Windows/D3D/SlateD3DRenderer.cpp
@@ -208,12 +208,23 @@ bool FSlateD3DRenderer::CreateDevice()
 	return bResult;
 }
 
-FSlateDrawBuffer& FSlateD3DRenderer::GetDrawBuffer()
+FSlateDrawBuffer& FSlateD3DRenderer::AcquireDrawBuffer()
 {
+	ensureMsgf(!DrawBuffer.IsLocked(), TEXT("The DrawBuffer is already locked. Make sure to call ReleaseDrawBuffer to release the DrawBuffer"));
+	DrawBuffer.Lock();
+
+	// Clear out the buffer each time its accessed
 	DrawBuffer.ClearBuffer();
+
 	return DrawBuffer;
 }
 
+void FSlateD3DRenderer::ReleaseDrawBuffer(FSlateDrawBuffer& InWindowDrawBuffer)
+{
+	ensureMsgf(&DrawBuffer == &InWindowDrawBuffer, TEXT("It release a DrawBuffer that is not a member of the SlateD3DRenderer"));
+	InWindowDrawBuffer.Unlock();
+}
+
 void FSlateD3DRenderer::LoadStyleResources( const ISlateStyle& Style )
 {
 	if ( TextureManager.IsValid() )
@@ -355,7 +366,18 @@ void FSlateD3DRenderer::Private_CreateViewport( TSharedRef<SWindow> InWindow, co
 
 	CreateBackBufferResources(Viewport.D3DSwapChain, Viewport.BackBufferTexture, Viewport.RenderTargetView);
 
-	Viewport.ProjectionMatrix = CreateProjectionMatrixD3D(Width, Height);
+	if (Width > 0 && Height > 0)
+	{
+		Viewport.ProjectionMatrix = CreateProjectionMatrixD3D(Width, Height);
+	}
+	else
+	{
+		Viewport.ProjectionMatrix.SetIdentity();
+	}
+
+#if !(UE_BUILD_SHIPPING || UE_BUILD_TEST)
+	ensure(!Viewport.ProjectionMatrix.ContainsNaN());
+#endif
 
 	WindowToViewportMap.Add(&InWindow.Get(), Viewport);
 }
@@ -421,7 +443,19 @@ void FSlateD3DRenderer::Private_ResizeViewport( const TSharedRef<SWindow> InWind
 		Viewport->ViewportInfo.Width = Width;
 		Viewport->ViewportInfo.Height = Height;
 		Viewport->bFullscreen = bFullscreen;
-		Viewport->ProjectionMatrix = CreateProjectionMatrixD3D( Width, Height );
+
+		if (Width > 0 && Height > 0)
+		{
+			Viewport->ProjectionMatrix = CreateProjectionMatrixD3D( Width, Height );
+		}
+		else
+		{
+			Viewport->ProjectionMatrix.SetIdentity();
+		}
+
+#if !(UE_BUILD_SHIPPING || UE_BUILD_TEST)
+		ensure(!Viewport->ProjectionMatrix.ContainsNaN());
+#endif
 
 		DXGI_SWAP_CHAIN_DESC Desc;
 		HRESULT Hr = Viewport->D3DSwapChain->GetDesc( &Desc );
diff --git a/Engine/Source/Developer/StandaloneRenderer/Private/Windows/D3D/SlateD3DRenderer.h b/Engine/Source/Developer/StandaloneRenderer/Private/Windows/D3D/SlateD3DRenderer.h
--- a/Engine/Source/Developer/StandaloneRenderer/Private/Windows/D3D/SlateD3DRenderer.h
+++ b/Engine/Source/Developer/StandaloneRenderer/Private/Windows/D3D/SlateD3DRenderer.h
@@ -52,7 +52,8 @@ public:
 	/** FSlateRenderer Interface */
 	virtual bool Initialize() override;
 	virtual void Destroy() override;
-	virtual FSlateDrawBuffer& GetDrawBuffer() override; 
+	virtual FSlateDrawBuffer& AcquireDrawBuffer() override;
+	virtual void ReleaseDrawBuffer( FSlateDrawBuffer& InWindowDrawBuffer ) override;
 	virtual void DrawWindows( FSlateDrawBuffer& InWindowDrawBuffer ) override;
 	virtual void OnWindowDestroyed( const TSharedRef<SWindow>& InWindow ) override;
 	virtual void CreateViewport( const TSharedRef<SWindow> InWindow ) override;
diff --git a/Engine/Source/Editor/UMGEditor/Private/WidgetBlueprintEditorUtils.cpp b/Engine/Source/Editor/UMGEditor/Private/WidgetBlueprintEditorUtils.cpp
--- a/Engine/Source/Editor/UMGEditor/Private/WidgetBlueprintEditorUtils.cpp
+++ b/Engine/Source/Editor/UMGEditor/Private/WidgetBlueprintEditorUtils.cpp
@@ -2029,18 +2029,20 @@ int32 FWidgetBlueprintEditorUtils::UpdateHittestGrid(FHittestGrid& HitTestGrid,
 	bool bUseGammaSpace = false;
 	TSharedPtr<ISlate3DRenderer, ESPMode::ThreadSafe> Renderer = FModuleManager::Get().LoadModuleChecked<ISlateRHIRendererModule>("SlateRHIRenderer")
 		.CreateSlate3DRenderer(bUseGammaSpace);
-	FSlateDrawBuffer& DrawBuffer = Renderer->GetDrawBuffer();
-	FSlateWindowElementList& WindowElementList = DrawBuffer.AddWindowElementList(Window);
 
-	int32 MaxLayerId = Window->Paint(
+	int32 MaxLayerId = 0;
+	{
+		ISlate3DRenderer::FScopedAcquireDrawBuffer ScopedDrawBuffer{ *Renderer };
+		FSlateWindowElementList& WindowElementList = ScopedDrawBuffer.GetDrawBuffer().AddWindowElementList(Window);
+
+		MaxLayerId = Window->Paint(
 			PaintArgs,
 			WindowGeometry, WindowClipRect,
 			WindowElementList,
 			0,
 			FWidgetStyle(),
 			Window->IsEnabled());
-
-	DrawBuffer.Unlock();
+	}
 
 	FSlateApplication::Get().InvalidateAllWidgets(false);
 
diff --git a/Engine/Source/Runtime/MoviePlayer/Private/DefaultGameMoviePlayer.cpp b/Engine/Source/Runtime/MoviePlayer/Private/DefaultGameMoviePlayer.cpp
--- a/Engine/Source/Runtime/MoviePlayer/Private/DefaultGameMoviePlayer.cpp
+++ b/Engine/Source/Runtime/MoviePlayer/Private/DefaultGameMoviePlayer.cpp
@@ -1004,29 +1004,30 @@ void FMoviePlayerWidgetRenderer::DrawWindow(float DeltaTime)
 	HittestGrid->SetHittestArea(VirtualRenderWindow->GetPositionInScreen(), VirtualRenderWindow->GetViewportSize());
 	HittestGrid->Clear();
 
-	// Get the free buffer & add our virtual window
-	FSlateDrawBuffer& DrawBuffer = SlateRenderer->GetDrawBuffer();
-	FSlateWindowElementList& WindowElementList = DrawBuffer.AddWindowElementList(VirtualRenderWindow);
-
-	WindowElementList.SetRenderTargetWindow(MainWindow);
-
-	int32 MaxLayerId = 0;
 	{
-		FPaintArgs PaintArgs(nullptr, *HittestGrid, FVector2D::ZeroVector, FSlateApplication::Get().GetCurrentTime(), FSlateApplication::Get().GetDeltaTime());
-
-		// Paint the window
-		MaxLayerId = VirtualRenderWindow->Paint(
-			PaintArgs,
-			WindowGeometry, ClipRect,
-			WindowElementList,
-			0,
-			FWidgetStyle(),
-			VirtualRenderWindow->IsEnabled());
-	}
+		// Get the free buffer & add our virtual window
+		FSlateRenderer::FScopedAcquireDrawBuffer ScopedDrawBuffer{ *SlateRenderer };
+		FSlateWindowElementList& WindowElementList = ScopedDrawBuffer.GetDrawBuffer().AddWindowElementList(VirtualRenderWindow);
 
-	SlateRenderer->DrawWindows(DrawBuffer);
+		WindowElementList.SetRenderTargetWindow(MainWindow);
 
-	DrawBuffer.ViewOffset = FVector2D::ZeroVector;
+		int32 MaxLayerId = 0;
+		{
+			FPaintArgs PaintArgs(nullptr, *HittestGrid, FVector2D::ZeroVector, FSlateApplication::Get().GetCurrentTime(), FSlateApplication::Get().GetDeltaTime());
+
+			// Paint the window
+			MaxLayerId = VirtualRenderWindow->Paint(
+				PaintArgs,
+				WindowGeometry, ClipRect,
+				WindowElementList,
+				0,
+				FWidgetStyle(),
+				VirtualRenderWindow->IsEnabled());
+		}
+
+		SlateRenderer->DrawWindows(ScopedDrawBuffer.GetDrawBuffer());
+		ScopedDrawBuffer.GetDrawBuffer().ViewOffset = FVector2D::ZeroVector;
+	}
 }
 
 float FDefaultGameMoviePlayer::GetViewportDPIScale() const
diff --git a/Engine/Source/Runtime/PreLoadScreen/Private/PreLoadSlateThreading.cpp b/Engine/Source/Runtime/PreLoadScreen/Private/PreLoadSlateThreading.cpp
--- a/Engine/Source/Runtime/PreLoadScreen/Private/PreLoadSlateThreading.cpp
+++ b/Engine/Source/Runtime/PreLoadScreen/Private/PreLoadSlateThreading.cpp
@@ -80,29 +80,30 @@ void FPreLoadSlateWidgetRenderer::DrawWindow(float DeltaTime)
 		HittestGrid->SetHittestArea(VirtualRenderWindow->GetPositionInScreen(), VirtualRenderWindow->GetViewportSize());
 		HittestGrid->Clear();
 
-		// Get the free buffer & add our virtual window
-		FSlateDrawBuffer& DrawBuffer = SlateRenderer->GetDrawBuffer();
-		FSlateWindowElementList& WindowElementList = DrawBuffer.AddWindowElementList(VirtualRenderWindow);
-
-		WindowElementList.SetRenderTargetWindow(MainWindow);
-
-		int32 MaxLayerId = 0;
 		{
-			FPaintArgs PaintArgs(nullptr, *HittestGrid, FVector2D::ZeroVector, FSlateApplication::Get().GetCurrentTime(), FSlateApplication::Get().GetDeltaTime());
-
-			// Paint the window
-			MaxLayerId = VirtualRenderWindow->Paint(
-				PaintArgs,
-				WindowGeometry, ClipRect,
-				WindowElementList,
-				0,
-				FWidgetStyle(),
-				VirtualRenderWindow->IsEnabled());
-		}
+			// Get the free buffer & add our virtual window
+			FSlateRenderer::FScopedAcquireDrawBuffer ScopedDrawBuffer{ *SlateRenderer };
+			FSlateWindowElementList& WindowElementList = ScopedDrawBuffer.GetDrawBuffer().AddWindowElementList(VirtualRenderWindow);
 
-		SlateRenderer->DrawWindows(DrawBuffer);
+			WindowElementList.SetRenderTargetWindow(MainWindow);
 
-		DrawBuffer.ViewOffset = FVector2D::ZeroVector;
+			int32 MaxLayerId = 0;
+			{
+				FPaintArgs PaintArgs(nullptr, *HittestGrid, FVector2D::ZeroVector, FSlateApplication::Get().GetCurrentTime(), FSlateApplication::Get().GetDeltaTime());
+
+				// Paint the window
+				MaxLayerId = VirtualRenderWindow->Paint(
+					PaintArgs,
+					WindowGeometry, ClipRect,
+					WindowElementList,
+					0,
+					FWidgetStyle(),
+					VirtualRenderWindow->IsEnabled());
+			}
+
+			SlateRenderer->DrawWindows(ScopedDrawBuffer.GetDrawBuffer());
+			ScopedDrawBuffer.GetDrawBuffer().ViewOffset = FVector2D::ZeroVector;
+		}
 	}
 }
 
diff --git a/Engine/Source/Runtime/Slate/Private/Framework/Application/SlateApplication.cpp b/Engine/Source/Runtime/Slate/Private/Framework/Application/SlateApplication.cpp
--- a/Engine/Source/Runtime/Slate/Private/Framework/Application/SlateApplication.cpp
+++ b/Engine/Source/Runtime/Slate/Private/Framework/Application/SlateApplication.cpp
@@ -1291,67 +1291,70 @@ void FSlateApplication::PrivateDrawWindows( TSharedPtr<SWindow> DrawOnlyThisWind
 	// Prepass the window
 	DrawPrepass( DrawOnlyThisWindow );
 
-	FDrawWindowArgs DrawWindowArgs( Renderer->GetDrawBuffer(), WidgetsToVisualizeUnderCursor);
-
 	{
-		SCOPE_CYCLE_COUNTER( STAT_SlateDrawWindowTime );
-
-		TSharedPtr<SWindow> ActiveModalWindow = GetActiveModalWindow(); 
+		FSlateRenderer::FScopedAcquireDrawBuffer ScopedDrawBuffer{ *Renderer };
+		FDrawWindowArgs DrawWindowArgs( ScopedDrawBuffer.GetDrawBuffer(), WidgetsToVisualizeUnderCursor);
+		ensureMsgf(DrawWindowArgs.OutDrawBuffer.IsLocked(), TEXT("The buffer should be lock by GetDrawBuffer."));
 
-		if (ActiveModalWindow.IsValid())
 		{
-			DrawWindowAndChildren( ActiveModalWindow.ToSharedRef(), DrawWindowArgs );
+			SCOPE_CYCLE_COUNTER( STAT_SlateDrawWindowTime );
 
-			for( TArray< TSharedRef<SWindow> >::TConstIterator CurrentWindowIt( SlateWindows ); CurrentWindowIt; ++CurrentWindowIt )
+			TSharedPtr<SWindow> ActiveModalWindow = GetActiveModalWindow(); 
+
+			if (ActiveModalWindow.IsValid())
 			{
-				const TSharedRef<SWindow>& CurrentWindow = *CurrentWindowIt;
-				if ( CurrentWindow->GetType() == EWindowType::ToolTip )
+				DrawWindowAndChildren( ActiveModalWindow.ToSharedRef(), DrawWindowArgs );
+
+				for( TArray< TSharedRef<SWindow> >::TConstIterator CurrentWindowIt( SlateWindows ); CurrentWindowIt; ++CurrentWindowIt )
 				{
-					DrawWindowAndChildren(CurrentWindow, DrawWindowArgs);
+					const TSharedRef<SWindow>& CurrentWindow = *CurrentWindowIt;
+					if ( CurrentWindow->GetType() == EWindowType::ToolTip )
+					{
+						DrawWindowAndChildren(CurrentWindow, DrawWindowArgs);
+					}
 				}
-			}
 
-			TArray< TSharedRef<SWindow> > NotificationWindows;
-			FSlateNotificationManager::Get().GetWindows(NotificationWindows);
-			for( auto CurrentWindowIt( NotificationWindows.CreateIterator() ); CurrentWindowIt; ++CurrentWindowIt )
+				TArray< TSharedRef<SWindow> > NotificationWindows;
+				FSlateNotificationManager::Get().GetWindows(NotificationWindows);
+				for( auto CurrentWindowIt( NotificationWindows.CreateIterator() ); CurrentWindowIt; ++CurrentWindowIt )
+				{
+					DrawWindowAndChildren(*CurrentWindowIt, DrawWindowArgs);
+				}	
+			}
+			else if( DrawOnlyThisWindow.IsValid() )
 			{
-				DrawWindowAndChildren(*CurrentWindowIt, DrawWindowArgs);
-			}	
-		}
-		else if( DrawOnlyThisWindow.IsValid() )
-		{
-			DrawWindowAndChildren( DrawOnlyThisWindow.ToSharedRef(), DrawWindowArgs );
-		}
-		else
-		{
-			// Draw all windows
-			// Use of an old-style iterator is intentional here, as SlateWindows 
-			// array may be mutated by user logic in draw calls. The iterator 
-			// prevents us from reading off the end and only keeps an index 
-			// internally:
-			for( TArray< TSharedRef<SWindow> >::TConstIterator CurrentWindowIt( SlateWindows ); CurrentWindowIt; ++CurrentWindowIt )
+				DrawWindowAndChildren( DrawOnlyThisWindow.ToSharedRef(), DrawWindowArgs );
+			}
+			else
 			{
-				TSharedRef<SWindow> CurrentWindow = *CurrentWindowIt;
-				// Only draw visible windows or in off-screen rendering mode
-				if (bRenderOffScreen || CurrentWindow->IsVisible() )
+				// Draw all windows
+				// Use of an old-style iterator is intentional here, as SlateWindows 
+				// array may be mutated by user logic in draw calls. The iterator 
+				// prevents us from reading off the end and only keeps an index 
+				// internally:
+				for( TArray< TSharedRef<SWindow> >::TConstIterator CurrentWindowIt( SlateWindows ); CurrentWindowIt; ++CurrentWindowIt )
 				{
-					DrawWindowAndChildren( CurrentWindow, DrawWindowArgs );
+					TSharedRef<SWindow> CurrentWindow = *CurrentWindowIt;
+					// Only draw visible windows or in off-screen rendering mode
+					if (bRenderOffScreen || CurrentWindow->IsVisible() )
+					{
+						DrawWindowAndChildren( CurrentWindow, DrawWindowArgs );
+					}
 				}
 			}
 		}
-	}
-
-	// This is potentially dangerous on the movie playback thread that slate sometimes runs on
-	if(!IsInSlateThread())
-	{
-		// Some windows may have been destroyed/removed.
-		// Do not attempt to draw any windows that have been removed.
-		TArray<SWindow*> AllWindows = GatherAllDescendants(SlateWindows);
-		DrawWindowArgs.OutDrawBuffer.RemoveUnusedWindowElement(AllWindows);
-	}
 
+		// This is potentially dangerous on the movie playback thread that slate sometimes runs on
+		if(!IsInSlateThread())
+		{
+			// Some windows may have been destroyed/removed.
+			// Do not attempt to draw any windows that have been removed.
+			TArray<SWindow*> AllWindows = GatherAllDescendants(SlateWindows);
+			DrawWindowArgs.OutDrawBuffer.RemoveUnusedWindowElement(AllWindows);
+		}
 
-	Renderer->DrawWindows( DrawWindowArgs.OutDrawBuffer );
+		Renderer->DrawWindows( DrawWindowArgs.OutDrawBuffer );
+	}
 }
 
 void FSlateApplication::PollGameDeviceState()
diff --git a/Engine/Source/Runtime/SlateCore/Private/Rendering/DrawElements.cpp b/Engine/Source/Runtime/SlateCore/Private/Rendering/DrawElements.cpp
--- a/Engine/Source/Runtime/SlateCore/Private/Rendering/DrawElements.cpp
+++ b/Engine/Source/Runtime/SlateCore/Private/Rendering/DrawElements.cpp
@@ -972,8 +972,6 @@ void FSlateWindowElementList::ResetElementList()
 {
 	QUICK_SCOPE_CYCLE_COUNTER(Slate_ResetElementList);
 
-	check(IsThreadSafeForSlateRendering());
-
 	DeferredPaintList.Reset();
 
 	BatchData.ResetData();
diff --git a/Engine/Source/Runtime/SlateCore/Private/Rendering/SlateDrawBuffer.cpp b/Engine/Source/Runtime/SlateCore/Private/Rendering/SlateDrawBuffer.cpp
--- a/Engine/Source/Runtime/SlateCore/Private/Rendering/SlateDrawBuffer.cpp
+++ b/Engine/Source/Runtime/SlateCore/Private/Rendering/SlateDrawBuffer.cpp
@@ -16,6 +16,8 @@ FSlateDrawBuffer::~FSlateDrawBuffer()
 
 FSlateWindowElementList& FSlateDrawBuffer::AddWindowElementList(TSharedRef<SWindow> ForWindow)
 {
+	ensureMsgf(IsLocked(), TEXT("The SlateDrawBuffer should be lock before modifying it."));
+
 	for ( int32 WindowIndex = 0; WindowIndex < WindowElementListsPool.Num(); ++WindowIndex )
 	{
 		TSharedRef<FSlateWindowElementList> ExistingElementList = WindowElementListsPool[WindowIndex];
@@ -25,6 +27,7 @@ FSlateWindowElementList& FSlateDrawBuffer::AddWindowElementList(TSharedRef<SWind
 			WindowElementLists.Add(ExistingElementList);
 			WindowElementListsPool.RemoveAtSwap(WindowIndex);
 
+			ensureMsgf(ExistingElementList->GetBatchData().GetNumFinalBatches() == 0, TEXT("The Buffer should have been clear when it was unlocked."));
 			ExistingElementList->ResetElementList();
 
 			return *ExistingElementList;
@@ -40,12 +43,15 @@ FSlateWindowElementList& FSlateDrawBuffer::AddWindowElementList(TSharedRef<SWind
 
 void FSlateDrawBuffer::RemoveUnusedWindowElement(const TArray<SWindow*>& AllWindows)
 {
+	ensureMsgf(IsLocked(), TEXT("The SlateDrawBuffer should be lock before modifying it."));
+
 	// Remove any window elements that are no longer valid.
 	for (int32 WindowIndex = 0; WindowIndex < WindowElementLists.Num(); ++WindowIndex)
 	{
 		SWindow* CandidateWindow = WindowElementLists[WindowIndex]->GetPaintWindow();
 		if (!CandidateWindow || !AllWindows.Contains(CandidateWindow))
 		{
+			WindowElementLists[WindowIndex]->ResetElementList();
 			WindowElementLists.RemoveAtSwap(WindowIndex);
 			--WindowIndex;
 		}
@@ -54,18 +60,31 @@ void FSlateDrawBuffer::RemoveUnusedWindowElement(const TArray<SWindow*>& AllWind
 
 bool FSlateDrawBuffer::Lock()
 {
-	return FPlatformAtomics::InterlockedCompareExchange(&Locked, 1, 0) == 0;
+	bool ExpectedValue = false;
+	bool bIsLock = bLocked.compare_exchange_strong(ExpectedValue, true);
+	if (bIsLock)
+	{
+		bIsLockedBySlateThread = IsInSlateThread();
+	}
+	return bIsLock;
 }
 
 void FSlateDrawBuffer::Unlock()
 {
-	FPlatformAtomics::InterlockedExchange(&Locked, 0);
+	// Rendering doesn't need the batch data anymore
+	for (TSharedRef<FSlateWindowElementList>& ExistingElementList : WindowElementLists)
+	{
+		ExistingElementList->ResetElementList();
+	}
+
+	bLocked = false;
 }
 
 void FSlateDrawBuffer::AddReferencedObjects(FReferenceCollector& Collector)
 {
-	// Locked buffers are the only ones that are currently referencing objects.  If unlocked, the element list is not in-use and contains to-be-cleared data
-	if(Locked != 0)
+	// Locked buffers are the only ones that are currently referencing objects.
+	//If unlocked, the element list is not in-use and contains "to-be-cleared" data.
+	if(bLocked && !bIsLockedBySlateThread)
 	{
 		for (TSharedRef<FSlateWindowElementList>& ElementList : WindowElementLists)
 		{
@@ -76,6 +95,8 @@ void FSlateDrawBuffer::AddReferencedObjects(FReferenceCollector& Collector)
 
 void FSlateDrawBuffer::ClearBuffer()
 {
+	ensureMsgf(IsLocked(), TEXT("The SlateDrawBuffer should be lock before modifying it."));
+
 	// Remove any window elements that are no longer valid.
 	for (int32 WindowIndex = 0; WindowIndex < WindowElementListsPool.Num(); ++WindowIndex)
 	{
@@ -101,6 +122,8 @@ void FSlateDrawBuffer::ClearBuffer()
 
 void FSlateDrawBuffer::UpdateResourceVersion(uint32 NewResourceVersion)
 {
+	ensureMsgf(IsLocked(), TEXT("The SlateDrawBuffer should be lock before modifying it."));
+
 	if (IsInGameThread() && NewResourceVersion != ResourceVersion)
 	{
 		WindowElementListsPool.Empty();
diff --git a/Engine/Source/Runtime/SlateCore/Public/Rendering/SlateDrawBuffer.h b/Engine/Source/Runtime/SlateCore/Public/Rendering/SlateDrawBuffer.h
--- a/Engine/Source/Runtime/SlateCore/Public/Rendering/SlateDrawBuffer.h
+++ b/Engine/Source/Runtime/SlateCore/Public/Rendering/SlateDrawBuffer.h
@@ -5,6 +5,8 @@
 #include "CoreMinimal.h"
 #include "UObject/GCObject.h"
 
+#include <atomic>
+
 class FSlateWindowElementList;
 class SWindow;
 
@@ -17,7 +19,8 @@ public:
 
 	/** Default constructor. */
 	explicit FSlateDrawBuffer()
-		: Locked(0)
+		: bLocked(false)
+		, bIsLockedBySlateThread(false)
 		, ResourceVersion(0)
 	{ }
 
@@ -51,7 +54,7 @@ public:
 	/** 
 	 * Locks the draw buffer.  Indicates that the viewport is in use.
 	 *
-	 * @return true if the viewport could be locked.  False otherwise.
+	 * @return true if the buffer could be locked.  False otherwise.
 	 * @see Unlock
 	 */
 	bool Lock();
@@ -63,6 +66,12 @@ public:
 	 */
 	void Unlock();
 
+	/** @return true if the buffer is locked. */
+	bool IsLocked() const
+	{
+		return bLocked;
+	}
+
 	/** FGCObject Interface */
 	virtual void AddReferencedObjects(FReferenceCollector& Collector) override;
 	virtual FString GetReferencerName() const
@@ -78,8 +87,8 @@ protected:
 	// that we restore if they're requested again.
 	TArray< TSharedRef<FSlateWindowElementList> > WindowElementListsPool;
 
-	// 1 if this buffer is locked, 0 otherwise.
-	volatile int32 Locked;
+	std::atomic<bool> bLocked;
+	bool bIsLockedBySlateThread;
 
 	// Last recorded version from the render. The WindowElementListsPool is emptied when this changes.
 	uint32 ResourceVersion;
diff --git a/Engine/Source/Runtime/SlateCore/Public/Rendering/SlateRenderer.h b/Engine/Source/Runtime/SlateCore/Public/Rendering/SlateRenderer.h
--- a/Engine/Source/Runtime/SlateCore/Public/Rendering/SlateRenderer.h
+++ b/Engine/Source/Runtime/SlateCore/Public/Rendering/SlateRenderer.h
@@ -176,9 +176,44 @@ public:
 	virtual ~FSlateRenderer();
 
 public:
+	/** Acquire the draw buffer and release it at the end of the scope. */
+	struct FScopedAcquireDrawBuffer
+	{
+		FScopedAcquireDrawBuffer(FSlateRenderer& InSlateRenderer)
+			: SlateRenderer(InSlateRenderer)
+			, DrawBuffer(InSlateRenderer.AcquireDrawBuffer())
+		{
+		}
+		~FScopedAcquireDrawBuffer()
+		{
+			SlateRenderer.ReleaseDrawBuffer(DrawBuffer);
+		}
+		FScopedAcquireDrawBuffer(const FScopedAcquireDrawBuffer&) = delete;
+		FScopedAcquireDrawBuffer& operator=(const FScopedAcquireDrawBuffer&) = delete;
+
+		FSlateDrawBuffer& GetDrawBuffer()
+		{
+			return DrawBuffer;
+		}
+
+	private:
+		FSlateRenderer& SlateRenderer;
+		FSlateDrawBuffer& DrawBuffer;
+	};
 
+public:
 	/** Returns a draw buffer that can be used by Slate windows to draw window elements */
-	virtual FSlateDrawBuffer& GetDrawBuffer() = 0;
+	UE_DEPRECATED(5.1, "Use FSlateRenderer::AcquireDrawBuffer instead and release the draw buffer.")
+	virtual FSlateDrawBuffer& GetDrawBuffer()
+	{
+		return AcquireDrawBuffer();
+	}
+
+	/** Returns a draw buffer that can be used by Slate windows to draw window elements */
+	virtual FSlateDrawBuffer& AcquireDrawBuffer() = 0;
+
+	/** Return the previously acquired buffer. */
+	virtual void ReleaseDrawBuffer( FSlateDrawBuffer& InWindowDrawBuffer ) = 0;
 
 	virtual bool Initialize() = 0;
 
diff --git a/Engine/Source/Runtime/SlateNullRenderer/Private/SlateNullRenderer.cpp b/Engine/Source/Runtime/SlateNullRenderer/Private/SlateNullRenderer.cpp
--- a/Engine/Source/Runtime/SlateNullRenderer/Private/SlateNullRenderer.cpp
+++ b/Engine/Source/Runtime/SlateNullRenderer/Private/SlateNullRenderer.cpp
@@ -1,13 +1,14 @@
 // Copyright Epic Games, Inc. All Rights Reserved.
 
 #include "SlateNullRenderer.h"
+#include "Rendering/DrawElements.h"
 #include "Rendering/SlateDrawBuffer.h"
 
-static TUniquePtr<FSlateDrawBuffer> StaticDrawBuffer;
 
 FSlateNullRenderer::FSlateNullRenderer(const TSharedRef<FSlateFontServices>& InSlateFontServices, const TSharedRef<FSlateShaderResourceManager>& InResourceManager)
 	: FSlateRenderer(InSlateFontServices)
 	, ResourceManager(InResourceManager)
+	, DrawBuffer(MakeUnique<FSlateDrawBuffer>())
 {
 }
 
@@ -18,18 +19,24 @@ bool FSlateNullRenderer::Initialize()
 
 void FSlateNullRenderer::Destroy()
 {
-	StaticDrawBuffer = nullptr;
+	DrawBuffer.Reset();
 }
 
-FSlateDrawBuffer& FSlateNullRenderer::GetDrawBuffer()
+FSlateDrawBuffer& FSlateNullRenderer::AcquireDrawBuffer()
 {
-	if (!StaticDrawBuffer.IsValid())
-	{
-		StaticDrawBuffer = MakeUnique<FSlateDrawBuffer>();
-	}
+	ensureMsgf(!DrawBuffer->IsLocked(), TEXT("The DrawBuffer is already locked. Make sure to call ReleaseDrawBuffer to release the DrawBuffer"));
+	DrawBuffer->Lock();
+
+	// Clear out the buffer each time its accessed
+	DrawBuffer->ClearBuffer();
 
-	StaticDrawBuffer->ClearBuffer();
-	return *StaticDrawBuffer;
+	return *DrawBuffer;
+}
+
+void FSlateNullRenderer::ReleaseDrawBuffer(FSlateDrawBuffer& InWindowDrawBuffer)
+{
+	ensureMsgf(DrawBuffer.Get() == &InWindowDrawBuffer, TEXT("It release a DrawBuffer that is not a member of the SlateNullRenderer"));
+	InWindowDrawBuffer.Unlock();
 }
 
 void FSlateNullRenderer::CreateViewport( const TSharedRef<SWindow> Window )
diff --git a/Engine/Source/Runtime/SlateNullRenderer/Private/SlateNullRenderer.h b/Engine/Source/Runtime/SlateNullRenderer/Private/SlateNullRenderer.h
--- a/Engine/Source/Runtime/SlateNullRenderer/Private/SlateNullRenderer.h
+++ b/Engine/Source/Runtime/SlateNullRenderer/Private/SlateNullRenderer.h
@@ -20,7 +20,8 @@ public:
 	/** FSlateRenderer interface */
 	virtual bool Initialize() override;
 	virtual void Destroy() override;
-	virtual FSlateDrawBuffer& GetDrawBuffer() override;
+	virtual FSlateDrawBuffer& AcquireDrawBuffer() override;
+	virtual void ReleaseDrawBuffer(FSlateDrawBuffer& InWindowDrawBuffer) override;
 	virtual void OnWindowDestroyed( const TSharedRef<SWindow>& InWindow ) override;
 	virtual void RequestResize( const TSharedPtr<SWindow>& Window, uint32 NewWidth, uint32 NewHeight ) override;
 	virtual void CreateViewport( const TSharedRef<SWindow> Window ) override;
@@ -44,5 +45,6 @@ public:
 
 private:
 	TSharedPtr<FSlateShaderResourceManager> ResourceManager;
+	TUniquePtr<FSlateDrawBuffer> DrawBuffer;
 	FCriticalSection ResourceCriticalSection;
 };
diff --git a/Engine/Source/Runtime/SlateRHIRenderer/Private/Slate3DRenderer.cpp b/Engine/Source/Runtime/SlateRHIRenderer/Private/Slate3DRenderer.cpp
--- a/Engine/Source/Runtime/SlateRHIRenderer/Private/Slate3DRenderer.cpp
+++ b/Engine/Source/Runtime/SlateRHIRenderer/Private/Slate3DRenderer.cpp
@@ -56,7 +56,7 @@ void FSlate3DRenderer::SetApplyColorDeficiencyCorrection(bool bApplyColorCorrect
 	RenderTargetPolicy->SetApplyColorDeficiencyCorrection(bApplyColorCorrection);
 }
 
-FSlateDrawBuffer& FSlate3DRenderer::GetDrawBuffer()
+FSlateDrawBuffer& FSlate3DRenderer::AcquireDrawBuffer()
 {
 	FreeBufferIndex = (FreeBufferIndex + 1) % NUM_DRAW_BUFFERS;
 	FSlateDrawBuffer* Buffer = &DrawBuffers[FreeBufferIndex];
@@ -76,6 +76,30 @@ FSlateDrawBuffer& FSlate3DRenderer::GetDrawBuffer()
 	return *Buffer;
 }
 
+void FSlate3DRenderer::ReleaseDrawBuffer(FSlateDrawBuffer& InWindowDrawBuffer)
+{
+#if DO_CHECK
+	bool bFound = false;
+	for (int32 Index = 0; Index < NUM_DRAW_BUFFERS; ++Index)
+	{
+		if (&DrawBuffers[Index] == &InWindowDrawBuffer)
+		{
+			bFound = true;
+			break;
+		}
+	}
+	ensureMsgf(bFound, TEXT("It release a DrawBuffer that is not a member of the Slate3DRenderer"));
+#endif
+
+	FSlateDrawBuffer* DrawBuffer = &InWindowDrawBuffer;
+	ENQUEUE_RENDER_COMMAND(SlateReleaseDrawBufferCommand)(
+		[DrawBuffer](FRHICommandListImmediate& RHICmdList)
+		{
+			FSlateReleaseDrawBufferCommand::ReleaseDrawBuffer(RHICmdList, DrawBuffer);
+		}
+	);
+}
+
 void FSlate3DRenderer::DrawWindow_GameThread(FSlateDrawBuffer& DrawBuffer)
 {
 	check( IsInGameThread() );
diff --git a/Engine/Source/Runtime/SlateRHIRenderer/Private/Slate3DRenderer.h b/Engine/Source/Runtime/SlateRHIRenderer/Private/Slate3DRenderer.h
--- a/Engine/Source/Runtime/SlateRHIRenderer/Private/Slate3DRenderer.h
+++ b/Engine/Source/Runtime/SlateRHIRenderer/Private/Slate3DRenderer.h
@@ -23,7 +23,8 @@ public:
 	virtual void SetUseGammaCorrection(bool bUseGammaCorrection) override;
 	virtual void SetApplyColorDeficiencyCorrection(bool bApplyColorCorrection) override;
 
-	virtual FSlateDrawBuffer& GetDrawBuffer() override;
+	virtual FSlateDrawBuffer& AcquireDrawBuffer() override;
+	virtual void ReleaseDrawBuffer(FSlateDrawBuffer& DrawBuffer) override;
 	virtual void DrawWindow_GameThread(FSlateDrawBuffer& DrawBuffer) override;
 	virtual void DrawWindowToTarget_RenderThread(FRHICommandListImmediate& RHICmdList, const struct FRenderThreadUpdateContext& Context) override;
 
diff --git a/Engine/Source/Runtime/SlateRHIRenderer/Private/SlateRHIRenderer.cpp b/Engine/Source/Runtime/SlateRHIRenderer/Private/SlateRHIRenderer.cpp
--- a/Engine/Source/Runtime/SlateRHIRenderer/Private/SlateRHIRenderer.cpp
+++ b/Engine/Source/Runtime/SlateRHIRenderer/Private/SlateRHIRenderer.cpp
@@ -268,7 +268,7 @@ void FSlateRHIRenderer::Destroy()
 }
 
 /** Returns a draw buffer that can be used by Slate windows to draw window elements */
-FSlateDrawBuffer& FSlateRHIRenderer::GetDrawBuffer()
+FSlateDrawBuffer& FSlateRHIRenderer::AcquireDrawBuffer()
 {
 	FreeBufferIndex = (FreeBufferIndex + 1) % NumDrawBuffers;
 
@@ -304,6 +304,30 @@ FSlateDrawBuffer& FSlateRHIRenderer::GetDrawBuffer()
 	return *Buffer;
 }
 
+void FSlateRHIRenderer::ReleaseDrawBuffer(FSlateDrawBuffer& InWindowDrawBuffer)
+{
+#if DO_CHECK
+	bool bFound = false;
+	for (int32 Index = 0; Index < NumDrawBuffers; ++Index)
+	{
+		if (&DrawBuffers[Index] == &InWindowDrawBuffer)
+		{
+			bFound = true;
+			break;
+		}
+	}
+	ensureMsgf(bFound, TEXT("It release a DrawBuffer that is not a member of the SlateRHIRenderer"));
+#endif
+
+	FSlateDrawBuffer* DrawBuffer = &InWindowDrawBuffer;
+	ENQUEUE_RENDER_COMMAND(SlateReleaseDrawBufferCommand)(
+		[DrawBuffer](FRHICommandListImmediate& RHICmdList)
+		{
+			FSlateReleaseDrawBufferCommand::ReleaseDrawBuffer(RHICmdList, DrawBuffer);
+		}
+	);
+}
+
 void FSlateRHIRenderer::CreateViewport(const TSharedRef<SWindow> Window)
 {
 	FlushRenderingCommands();
@@ -1858,7 +1882,6 @@ FSlateEndDrawingWindowsCommand::FSlateEndDrawingWindowsCommand(FSlateRHIRenderin
 
 void FSlateEndDrawingWindowsCommand::Execute(FRHICommandListBase& CmdList)
 {
-	DrawBuffer->Unlock();
 	Policy.EndDrawingWindows();
 }
 
@@ -1875,6 +1898,28 @@ void FSlateEndDrawingWindowsCommand::EndDrawingWindows(FRHICommandListImmediate&
 	}
 }
 
+FSlateReleaseDrawBufferCommand::FSlateReleaseDrawBufferCommand(FSlateDrawBuffer* InDrawBuffer)
+	: DrawBuffer(InDrawBuffer)
+{}
+
+void FSlateReleaseDrawBufferCommand::Execute(FRHICommandListBase& CmdList)
+{
+	DrawBuffer->Unlock();
+}
+
+void FSlateReleaseDrawBufferCommand::ReleaseDrawBuffer(FRHICommandListImmediate& RHICmdList, FSlateDrawBuffer* DrawBuffer)
+{
+	if (!RHICmdList.Bypass())
+	{
+		ALLOC_COMMAND_CL(RHICmdList, FSlateReleaseDrawBufferCommand)(DrawBuffer);
+	}
+	else
+	{
+		FSlateReleaseDrawBufferCommand Cmd(DrawBuffer);
+		Cmd.Execute(RHICmdList);
+	}
+}
+
 
 struct FClearCachedRenderingDataCommand2 final : public FRHICommand < FClearCachedRenderingDataCommand2 >
 {
diff --git a/Engine/Source/Runtime/SlateRHIRenderer/Private/SlateRHIRenderer.h b/Engine/Source/Runtime/SlateRHIRenderer/Private/SlateRHIRenderer.h
--- a/Engine/Source/Runtime/SlateRHIRenderer/Private/SlateRHIRenderer.h
+++ b/Engine/Source/Runtime/SlateRHIRenderer/Private/SlateRHIRenderer.h
@@ -177,7 +177,8 @@ public:
 	/** FSlateRenderer interface */
 	virtual bool Initialize() override;
 	virtual void Destroy() override;
-	virtual FSlateDrawBuffer& GetDrawBuffer() override;
+	virtual FSlateDrawBuffer& AcquireDrawBuffer() override;
+	virtual void ReleaseDrawBuffer(FSlateDrawBuffer& InWindowDrawBuffer) override;
 	virtual void OnWindowDestroyed( const TSharedRef<SWindow>& InWindow ) override;
 	virtual void OnWindowFinishReshaped(const TSharedPtr<SWindow>& InWindow) override;
 	virtual void RequestResize( const TSharedPtr<SWindow>& Window, uint32 NewWidth, uint32 NewHeight ) override;
@@ -214,7 +215,7 @@ public:
 	virtual void AddWidgetRendererUpdate(const struct FRenderThreadUpdateContext& Context, bool bDeferredRenderTargetUpdate) override;
 
 	/** Draws windows from a FSlateDrawBuffer on the render thread */
-	void DrawWindow_RenderThread(FRHICommandListImmediate& RHICmdList, FViewportInfo& ViewportInfo, FSlateWindowElementList& WindowElementList, const struct FSlateDrawWindowCommandParams& DrawParams);
+	void DrawWindow_RenderThread(FRHICommandListImmediate& RHICmdList, FViewportInfo& ViewportInfo, FSlateWindowElementList& WindowElementList, const struct FSlateDrawWindowCommandParams& DrawCommandParams);
 
 	/**
 	 * Reloads texture resources from disk                   
@@ -325,3 +326,18 @@ struct FSlateEndDrawingWindowsCommand final : public FRHICommand < FSlateEndDraw
 
 	static void EndDrawingWindows(FRHICommandListImmediate& RHICmdList, FSlateDrawBuffer* DrawBuffer, FSlateRHIRenderingPolicy& Policy);
 };
+
+struct FSlateReleaseDrawBufferCommandString
+{
+	static const TCHAR* TStr() { return TEXT("FSlateReleaseDrawBufferCommand"); }
+};
+struct FSlateReleaseDrawBufferCommand final : public FRHICommand < FSlateReleaseDrawBufferCommand, FSlateReleaseDrawBufferCommandString >
+{
+	FSlateDrawBuffer* DrawBuffer;
+
+	FSlateReleaseDrawBufferCommand(FSlateDrawBuffer* InDrawBuffer);
+
+	void Execute(FRHICommandListBase& CmdList);
+
+	static void ReleaseDrawBuffer(FRHICommandListImmediate& RHICmdList, FSlateDrawBuffer* DrawBuffer);
+};
diff --git a/Engine/Source/Runtime/SlateRHIRenderer/Public/Interfaces/ISlate3DRenderer.h b/Engine/Source/Runtime/SlateRHIRenderer/Public/Interfaces/ISlate3DRenderer.h
--- a/Engine/Source/Runtime/SlateRHIRenderer/Public/Interfaces/ISlate3DRenderer.h
+++ b/Engine/Source/Runtime/SlateRHIRenderer/Public/Interfaces/ISlate3DRenderer.h
@@ -11,14 +11,50 @@ class FTextureRenderTarget2DResource;
 class ISlate3DRenderer : public TSharedFromThis<ISlate3DRenderer, ESPMode::ThreadSafe>
 {
 public:
-	virtual ~ISlate3DRenderer() {}
+	virtual ~ISlate3DRenderer() = default;
+
+public:
+	/** Acquire the draw buffer and release it at the end of the scope. */
+	struct FScopedAcquireDrawBuffer
+	{
+		FScopedAcquireDrawBuffer(ISlate3DRenderer& InSlateRenderer)
+			: SlateRenderer(InSlateRenderer)
+			, DrawBuffer(InSlateRenderer.AcquireDrawBuffer())
+		{
+		}
+		~FScopedAcquireDrawBuffer()
+		{
+			SlateRenderer.ReleaseDrawBuffer(DrawBuffer);
+		}
+		FScopedAcquireDrawBuffer(const FScopedAcquireDrawBuffer&) = delete;
+		FScopedAcquireDrawBuffer& operator=(const FScopedAcquireDrawBuffer&) = delete;
+
+		FSlateDrawBuffer& GetDrawBuffer()
+		{
+			return DrawBuffer;
+		}
+
+	private:
+		ISlate3DRenderer& SlateRenderer;
+		FSlateDrawBuffer& DrawBuffer;
+	};
 
 	/** set if this renderer should render in gamma space by default. */
 	virtual void SetUseGammaCorrection(bool bUseGammaCorrection) = 0;
 	virtual void SetApplyColorDeficiencyCorrection(bool bApplyColorCorrection) = 0;
 
 	/** @return The free buffer for drawing */
-	virtual FSlateDrawBuffer& GetDrawBuffer() = 0;
+	UE_DEPRECATED(5.1, "Use ISlate3DRenderer::AcquireDrawBuffer instead and release the draw buffer.")
+	virtual FSlateDrawBuffer& GetDrawBuffer()
+	{
+		return AcquireDrawBuffer();
+	}
+
+	/** @return The free buffer for drawing */
+	virtual FSlateDrawBuffer& AcquireDrawBuffer() = 0;
+
+	/** Return the previously acquired buffer. */
+	virtual void ReleaseDrawBuffer(FSlateDrawBuffer& InWindowDrawBuffer) = 0;
 
 	/** 
 	 * Batches the draw elements in the buffer to prepare it for rendering.
diff --git a/Engine/Source/Runtime/UMG/Private/Slate/WidgetRenderer.cpp b/Engine/Source/Runtime/UMG/Private/Slate/WidgetRenderer.cpp
--- a/Engine/Source/Runtime/UMG/Private/Slate/WidgetRenderer.cpp
+++ b/Engine/Source/Runtime/UMG/Private/Slate/WidgetRenderer.cpp
@@ -218,58 +218,56 @@ void FWidgetRenderer::DrawWindow(
 	FSlateRenderer* MainSlateRenderer = FSlateApplication::Get().GetRenderer();
 	FScopeLock ScopeLock(MainSlateRenderer->GetResourceCriticalSection());
 
-	if ( LIKELY(FApp::CanEverRender()) )
+	if (LIKELY(FApp::CanEverRender()))
 	{
-	    if ( bPrepassNeeded )
-	    {
-		    // Ticking can cause geometry changes.  Recompute
-		    Window->SlatePrepass(WindowGeometry.Scale);
-	    }
-	
+		if (bPrepassNeeded)
+		{
+			// Ticking can cause geometry changes.  Recompute
+			Window->SlatePrepass(WindowGeometry.Scale);
+		}
+
 		PaintArgs.GetHittestGrid().SetHittestArea(WindowClipRect.GetTopLeft(), WindowClipRect.GetSize());
 
-		if ( bClearHitTestGrid )
+		if (bClearHitTestGrid)
 		{
 			// Prepare the test grid 
 			PaintArgs.GetHittestGrid().Clear();
 		}
-    
-	    // Get the free buffer & add our virtual window
-	    FSlateDrawBuffer& DrawBuffer = Renderer->GetDrawBuffer();
-	    FSlateWindowElementList& WindowElementList = DrawBuffer.AddWindowElementList(Window);
-
-	    int32 MaxLayerId = 0;
-	    {
-		    // Paint the window
-		    MaxLayerId = Window->Paint(
-			    PaintArgs,
-			    WindowGeometry, WindowClipRect,
-			    WindowElementList,
-			    0,
-			    FWidgetStyle(),
-			    Window->IsEnabled());
-	    }
-
-		//MaxLayerId = WindowElementList.PaintDeferred(MaxLayerId);
-		DeferredPaints = WindowElementList.GetDeferredPaintList();
-
-		Renderer->DrawWindow_GameThread(DrawBuffer);
-
-		DrawBuffer.ViewOffset = ViewOffset;
-
-		FRenderThreadUpdateContext RenderThreadUpdateContext =
+
 		{
-			&DrawBuffer,
-			static_cast<float>(FApp::GetCurrentTime() - GStartTime),
-			static_cast<float>(FApp::GetDeltaTime()),
-			static_cast<float>(FPlatformTime::Seconds() - GStartTime),
-			static_cast<float>(FApp::GetDeltaTime()),
-			RenderTarget,
-			Renderer.Get(),
-			bClearTarget
-		};
-
-		FSlateApplication::Get().GetRenderer()->AddWidgetRendererUpdate(RenderThreadUpdateContext, bDeferRenderTargetUpdate);
+			// Get the free buffer & add our virtual window
+			ISlate3DRenderer::FScopedAcquireDrawBuffer ScopedDrawBuffer{ *Renderer };
+			FSlateWindowElementList& WindowElementList = ScopedDrawBuffer.GetDrawBuffer().AddWindowElementList(Window);
+
+			// Paint the window
+			int32 MaxLayerId = Window->Paint(
+				PaintArgs,
+				WindowGeometry, WindowClipRect,
+				WindowElementList,
+				0,
+				FWidgetStyle(),
+				Window->IsEnabled());
+
+			//MaxLayerId = WindowElementList.PaintDeferred(MaxLayerId);
+			DeferredPaints = WindowElementList.GetDeferredPaintList();
+
+			Renderer->DrawWindow_GameThread(ScopedDrawBuffer.GetDrawBuffer());
+
+			ScopedDrawBuffer.GetDrawBuffer().ViewOffset = ViewOffset;
+
+			FRenderThreadUpdateContext RenderThreadUpdateContext =
+			{
+				&(ScopedDrawBuffer.GetDrawBuffer()),
+				static_cast<float>(FApp::GetCurrentTime() - GStartTime),
+				static_cast<float>(FApp::GetDeltaTime()),
+				static_cast<float>(FPlatformTime::Seconds() - GStartTime),
+				RenderTarget,
+				Renderer.Get(),
+				bClearTarget
+			};
+
+			MainSlateRenderer->AddWidgetRendererUpdate(RenderThreadUpdateContext, bDeferRenderTargetUpdate);
+		}
 	}
 #endif // !UE_SERVER
 }
@@ -298,53 +296,53 @@ bool FWidgetRenderer::DrawInvalidationRoot(TSharedRef<SVirtualWindow>& VirtualWi
 	{
 		// Need to set a new window element list so make a copy
 		FSlateInvalidationContext ContextCopy = Context;
-		
-		// Get the free buffer & add our virtual window
-		FSlateDrawBuffer& DrawBuffer = Renderer->GetDrawBuffer();
-		FSlateWindowElementList& WindowElementList = DrawBuffer.AddWindowElementList(VirtualWindow);
-
-		ContextCopy.WindowElementList = &WindowElementList;
-		FSlateInvalidationResult Result = Root.PaintInvalidationRoot(ContextCopy);
 
-		const int32 MaxLayerId = Result.MaxLayerIdPainted;
-
-		if(Result.bRepaintedWidgets)
 		{
-			//MaxLayerId = WindowElementList.PaintDeferred(MaxLayerId);
-			DeferredPaints = WindowElementList.GetDeferredPaintList();
+			// Get the free buffer & add our virtual window
+			ISlate3DRenderer::FScopedAcquireDrawBuffer ScopedDrawBuffer{ *Renderer };
+			FSlateWindowElementList& WindowElementList = ScopedDrawBuffer.GetDrawBuffer().AddWindowElementList(VirtualWindow);
 
-			Renderer->DrawWindow_GameThread(DrawBuffer);
+			ContextCopy.WindowElementList = &WindowElementList;
+			FSlateInvalidationResult Result = Root.PaintInvalidationRoot(ContextCopy);
 
-			DrawBuffer.ViewOffset = ViewOffset;
+			const int32 MaxLayerId = Result.MaxLayerIdPainted;
 
-			FRenderThreadUpdateContext RenderThreadUpdateContext =
+			if (Result.bRepaintedWidgets)
 			{
-				&DrawBuffer,
-				static_cast<float>(FApp::GetCurrentTime() - GStartTime),
-				static_cast<float>(FApp::GetDeltaTime()),
-				static_cast<float>(FPlatformTime::Seconds() - GStartTime),
-				static_cast<float>(FApp::GetDeltaTime()),
-				static_cast<FRenderTarget*>(RenderTarget->GameThread_GetRenderTargetResource()),
-				Renderer.Get(),
-				bClearTarget
-			};
-
-			bRepaintedWidgets = Result.bRepaintedWidgets;
-			FSlateApplication::Get().GetRenderer()->AddWidgetRendererUpdate(RenderThreadUpdateContext, bDeferRenderTargetUpdate);
-
-
-			// Any deferred painted elements of the retainer should be drawn directly by the main renderer, not rendered into the render target,
-			// as most of those sorts of things will break the rendering rect, things like tooltips, and popup menus.
-			for (auto& DeferredPaint : DeferredPaints)
+				//MaxLayerId = WindowElementList.PaintDeferred(MaxLayerId);
+				DeferredPaints = WindowElementList.GetDeferredPaintList();
+
+				Renderer->DrawWindow_GameThread(ScopedDrawBuffer.GetDrawBuffer());
+
+				ScopedDrawBuffer.GetDrawBuffer().ViewOffset = ViewOffset;
+
+				FRenderThreadUpdateContext RenderThreadUpdateContext =
+				{
+					&(ScopedDrawBuffer.GetDrawBuffer()),
+					static_cast<float>(FApp::GetCurrentTime() - GStartTime),
+					static_cast<float>(FApp::GetDeltaTime()),
+					static_cast<float>(FPlatformTime::Seconds() - GStartTime),
+					static_cast<FRenderTarget*>(RenderTarget->GameThread_GetRenderTargetResource()),
+					Renderer.Get(),
+					bClearTarget
+				};
+
+				bRepaintedWidgets = Result.bRepaintedWidgets;
+				FSlateApplication::Get().GetRenderer()->AddWidgetRendererUpdate(RenderThreadUpdateContext, bDeferRenderTargetUpdate);
+
+
+				// Any deferred painted elements of the retainer should be drawn directly by the main renderer, not rendered into the render target,
+				// as most of those sorts of things will break the rendering rect, things like tooltips, and popup menus.
+				for (auto& DeferredPaint : DeferredPaints)
+				{
+					Context.WindowElementList->QueueDeferredPainting(DeferredPaint->Copy(*Context.PaintArgs));
+				}
+			}
+			else
 			{
-				Context.WindowElementList->QueueDeferredPainting(DeferredPaint->Copy(*Context.PaintArgs));
+				WindowElementList.ResetElementList();
 			}
 		}
-		else
-		{
-			WindowElementList.ResetElementList();
-			DrawBuffer.Unlock();
-		}
 	}
 #endif // !UE_SERVER
 
