Fix bad merged introduced with CL 19092359
#rnx

diff --git a/Engine/Source/Runtime/UMG/Private/Slate/WidgetRenderer.cpp b/Engine/Source/Runtime/UMG/Private/Slate/WidgetRenderer.cpp
--- a/Engine/Source/Runtime/UMG/Private/Slate/WidgetRenderer.cpp
+++ b/Engine/Source/Runtime/UMG/Private/Slate/WidgetRenderer.cpp
@@ -261,6 +261,7 @@ void FWidgetRenderer::DrawWindow(
 				static_cast<float>(FApp::GetCurrentTime() - GStartTime),
 				static_cast<float>(FApp::GetDeltaTime()),
 				static_cast<float>(FPlatformTime::Seconds() - GStartTime),
+				static_cast<float>(FApp::GetDeltaTime()),
 				RenderTarget,
 				Renderer.Get(),
 				bClearTarget
@@ -322,6 +323,7 @@ bool FWidgetRenderer::DrawInvalidationRoot(TSharedRef<SVirtualWindow>& VirtualWi
 					static_cast<float>(FApp::GetCurrentTime() - GStartTime),
 					static_cast<float>(FApp::GetDeltaTime()),
 					static_cast<float>(FPlatformTime::Seconds() - GStartTime),
+					static_cast<float>(FApp::GetDeltaTime()),
 					static_cast<FRenderTarget*>(RenderTarget->GameThread_GetRenderTargetResource()),
 					Renderer.Get(),
 					bClearTarget
