UMG: Remove deprecation code bellow 4.23 in UMG code.

diff --git a/Engine/Source/Runtime/UMG/Private/Blueprint/IUserObjectListEntry.cpp b/Engine/Source/Runtime/UMG/Private/Blueprint/IUserObjectListEntry.cpp
--- a/Engine/Source/Runtime/UMG/Private/Blueprint/IUserObjectListEntry.cpp
+++ b/Engine/Source/Runtime/UMG/Private/Blueprint/IUserObjectListEntry.cpp
@@ -11,11 +11,6 @@ UUserObjectListEntry::UUserObjectListEntry(const FObjectInitializer& Initializer
 
 void IUserObjectListEntry::NativeOnListItemObjectSet(UObject* ListItemObject)
 {
-PRAGMA_DISABLE_DEPRECATION_WARNINGS
-	// Making sure the deprecated path still functions, remove this when the old method is deleted
-	SetListItemObjectInternal(ListItemObject);	
-PRAGMA_ENABLE_DEPRECATION_WARNINGS
-
 	Execute_OnListItemObjectSet(Cast<UObject>(this), ListItemObject);
 }
 
diff --git a/Engine/Source/Runtime/UMG/Private/Components/Border.cpp b/Engine/Source/Runtime/UMG/Private/Components/Border.cpp
--- a/Engine/Source/Runtime/UMG/Private/Components/Border.cpp
+++ b/Engine/Source/Runtime/UMG/Private/Components/Border.cpp
@@ -299,14 +299,6 @@ void UBorder::PostLoad()
 {
 	Super::PostLoad();
 
-#if WITH_EDITORONLY_DATA
-	if ( GetLinkerUEVersion() < VER_UE4_DEPRECATE_UMG_STYLE_ASSETS && Brush_DEPRECATED != nullptr )
-	{
-		Background = Brush_DEPRECATED->Brush;
-		Brush_DEPRECATED = nullptr;
-	}
-#endif
-
 	if ( GetChildrenCount() > 0 )
 	{
 		//TODO UMG Pre-Release Upgrade, now have slots of their own.  Convert existing slot to new slot.
diff --git a/Engine/Source/Runtime/UMG/Private/Components/Button.cpp b/Engine/Source/Runtime/UMG/Private/Components/Button.cpp
--- a/Engine/Source/Runtime/UMG/Private/Components/Button.cpp
+++ b/Engine/Source/Runtime/UMG/Private/Components/Button.cpp
@@ -208,17 +208,6 @@ void UButton::PostLoad()
 			}
 		}
 	}
-
-	if( GetLinkerUEVersion() < VER_UE4_DEPRECATE_UMG_STYLE_ASSETS && Style_DEPRECATED != nullptr )
-	{
-		const FButtonStyle* StylePtr = Style_DEPRECATED->GetStyle<FButtonStyle>();
-		if(StylePtr != nullptr)
-		{
-			WidgetStyle = *StylePtr;
-		}
-
-		Style_DEPRECATED = nullptr;
-	}
 }
 
 FReply UButton::SlateHandleClicked()
diff --git a/Engine/Source/Runtime/UMG/Private/Components/CheckBox.cpp b/Engine/Source/Runtime/UMG/Private/Components/CheckBox.cpp
--- a/Engine/Source/Runtime/UMG/Private/Components/CheckBox.cpp
+++ b/Engine/Source/Runtime/UMG/Private/Components/CheckBox.cpp
@@ -52,9 +52,6 @@ UCheckBox::UCheckBox(const FObjectInitializer& ObjectInitializer)
 	CheckedState = ECheckBoxState::Unchecked;
 
 	HorizontalAlignment = HAlign_Fill;
-	Padding_DEPRECATED = FMargin(0, 0, 0, 0);
-
-	BorderBackgroundColor_DEPRECATED = FLinearColor::White;
 
 	ClickMethod = EButtonClickMethod::DownAndUp;
 	TouchMethod = EButtonTouchMethod::DownAndUp;
@@ -203,91 +200,6 @@ void UCheckBox::SlateOnCheckStateChangedCallback(ECheckBoxState NewState)
 	OnCheckStateChanged.Broadcast(bWantsToBeChecked);
 }
 
-void UCheckBox::PostLoad()
-{
-	Super::PostLoad();
-
-	if ( GetLinkerUEVersion() < VER_UE4_DEPRECATE_UMG_STYLE_ASSETS )
-	{
-		if ( Style_DEPRECATED != nullptr )
-		{
-			const FCheckBoxStyle* StylePtr = Style_DEPRECATED->GetStyle<FCheckBoxStyle>();
-			if ( StylePtr != nullptr )
-			{
-				WidgetStyle = *StylePtr;
-			}
-
-			Style_DEPRECATED = nullptr;
-		}
-
-		if ( UncheckedImage_DEPRECATED != nullptr )
-		{
-			WidgetStyle.UncheckedImage = UncheckedImage_DEPRECATED->Brush;
-			UncheckedImage_DEPRECATED = nullptr;
-		}
-
-		if ( UncheckedHoveredImage_DEPRECATED != nullptr )
-		{
-			WidgetStyle.UncheckedHoveredImage = UncheckedHoveredImage_DEPRECATED->Brush;
-			UncheckedHoveredImage_DEPRECATED = nullptr;
-		}
-
-		if ( UncheckedPressedImage_DEPRECATED != nullptr )
-		{
-			WidgetStyle.UncheckedPressedImage = UncheckedPressedImage_DEPRECATED->Brush;
-			UncheckedPressedImage_DEPRECATED = nullptr;
-		}
-
-		if ( CheckedImage_DEPRECATED != nullptr )
-		{
-			WidgetStyle.CheckedImage = CheckedImage_DEPRECATED->Brush;
-			CheckedImage_DEPRECATED = nullptr;
-		}
-
-		if ( CheckedHoveredImage_DEPRECATED != nullptr )
-		{
-			WidgetStyle.CheckedHoveredImage = CheckedHoveredImage_DEPRECATED->Brush;
-			CheckedHoveredImage_DEPRECATED = nullptr;
-		}
-
-		if ( CheckedPressedImage_DEPRECATED != nullptr )
-		{
-			WidgetStyle.CheckedPressedImage = CheckedPressedImage_DEPRECATED->Brush;
-			CheckedPressedImage_DEPRECATED = nullptr;
-		}
-
-		if ( UndeterminedImage_DEPRECATED != nullptr )
-		{
-			WidgetStyle.UndeterminedImage = UndeterminedImage_DEPRECATED->Brush;
-			UndeterminedImage_DEPRECATED = nullptr;
-		}
-
-		if ( UndeterminedHoveredImage_DEPRECATED != nullptr )
-		{
-			WidgetStyle.UndeterminedHoveredImage = UndeterminedHoveredImage_DEPRECATED->Brush;
-			UndeterminedHoveredImage_DEPRECATED = nullptr;
-		}
-
-		if ( UndeterminedPressedImage_DEPRECATED != nullptr )
-		{
-			WidgetStyle.UndeterminedPressedImage = UndeterminedPressedImage_DEPRECATED->Brush;
-			UndeterminedPressedImage_DEPRECATED = nullptr;
-		}
-	}
-
-	if (GetLinkerUEVersion() < VER_UE4_DEPRECATE_UMG_STYLE_OVERRIDES)
-	{
-		WidgetStyle.Padding = Padding_DEPRECATED;
-		Padding_DEPRECATED = FMargin(0);
-
-		if (BorderBackgroundColor_DEPRECATED != FLinearColor::White)
-		{
-			WidgetStyle.BorderBackgroundColor = BorderBackgroundColor_DEPRECATED;
-			BorderBackgroundColor_DEPRECATED = FLinearColor::White;
-		}
-	}
-}
-
 #if WITH_ACCESSIBILITY
 TSharedPtr<SWidget> UCheckBox::GetAccessibleWidget() const
 {
diff --git a/Engine/Source/Runtime/UMG/Private/Components/CircularThrobber.cpp b/Engine/Source/Runtime/UMG/Private/Components/CircularThrobber.cpp
--- a/Engine/Source/Runtime/UMG/Private/Components/CircularThrobber.cpp
+++ b/Engine/Source/Runtime/UMG/Private/Components/CircularThrobber.cpp
@@ -119,20 +119,6 @@ void UCircularThrobber::SetRadius(float InRadius)
 	}
 }
 
-void UCircularThrobber::PostLoad()
-{
-	Super::PostLoad();
-
-	if ( GetLinkerUEVersion() < VER_UE4_DEPRECATE_UMG_STYLE_ASSETS )
-	{
-		if ( PieceImage_DEPRECATED != nullptr )
-		{
-			Image = PieceImage_DEPRECATED->Brush;
-			PieceImage_DEPRECATED = nullptr;
-		}
-	}
-}
-
 #if WITH_EDITOR
 
 const FText UCircularThrobber::GetPaletteCategory()
diff --git a/Engine/Source/Runtime/UMG/Private/Components/EditableText.cpp b/Engine/Source/Runtime/UMG/Private/Components/EditableText.cpp
--- a/Engine/Source/Runtime/UMG/Private/Components/EditableText.cpp
+++ b/Engine/Source/Runtime/UMG/Private/Components/EditableText.cpp
@@ -50,14 +50,6 @@ UEditableText::UEditableText(const FObjectInitializer& ObjectInitializer)
 	}
 #endif // WITH_EDITOR
 
-	ColorAndOpacity_DEPRECATED = FLinearColor::Black;
-
-	if (!IsRunningDedicatedServer())
-	{
-		static ConstructorHelpers::FObjectFinder<UFont> RobotoFontObj(*UWidget::GetDefaultFontName());
-		Font_DEPRECATED = FSlateFontInfo(RobotoFontObj.Object, 12, FName("Bold"));
-	}
-
 	IsReadOnly = false;
 	IsPassword = false;
 	MinimumDesiredWidth = 0.0f;
@@ -213,58 +205,6 @@ void UEditableText::HandleOnTextCommitted(const FText& InText, ETextCommit::Type
 	OnTextCommitted.Broadcast(InText, CommitMethod);
 }
 
-void UEditableText::PostLoad()
-{
-	Super::PostLoad();
-
-	if ( GetLinkerUEVersion() < VER_UE4_DEPRECATE_UMG_STYLE_ASSETS )
-	{
-		if ( Style_DEPRECATED != nullptr )
-		{
-			const FEditableTextStyle* StylePtr = Style_DEPRECATED->GetStyle<FEditableTextStyle>();
-			if ( StylePtr != nullptr )
-			{
-				WidgetStyle = *StylePtr;
-			}
-
-			Style_DEPRECATED = nullptr;
-		}
-
-		if ( BackgroundImageSelected_DEPRECATED != nullptr )
-		{
-			WidgetStyle.BackgroundImageSelected = BackgroundImageSelected_DEPRECATED->Brush;
-			BackgroundImageSelected_DEPRECATED = nullptr;
-		}
-
-		if ( BackgroundImageComposing_DEPRECATED != nullptr )
-		{
-			WidgetStyle.BackgroundImageComposing = BackgroundImageComposing_DEPRECATED->Brush;
-			BackgroundImageComposing_DEPRECATED = nullptr;
-		}
-
-		if ( CaretImage_DEPRECATED != nullptr )
-		{
-			WidgetStyle.CaretImage = CaretImage_DEPRECATED->Brush;
-			CaretImage_DEPRECATED = nullptr;
-		}
-	}
-
-	if (GetLinkerUEVersion() < VER_UE4_DEPRECATE_UMG_STYLE_OVERRIDES)
-	{
-		if (Font_DEPRECATED.HasValidFont())
-		{
-			WidgetStyle.Font = Font_DEPRECATED;
-			Font_DEPRECATED = FSlateFontInfo();
-		}
-
-		if (ColorAndOpacity_DEPRECATED != FLinearColor::Black)
-		{
-			WidgetStyle.ColorAndOpacity = ColorAndOpacity_DEPRECATED;
-			ColorAndOpacity_DEPRECATED = FLinearColor::Black;
-		}
-	}
-}
-
 #if WITH_ACCESSIBILITY
 TSharedPtr<SWidget> UEditableText::GetAccessibleWidget() const
 {
diff --git a/Engine/Source/Runtime/UMG/Private/Components/EditableTextBox.cpp b/Engine/Source/Runtime/UMG/Private/Components/EditableTextBox.cpp
--- a/Engine/Source/Runtime/UMG/Private/Components/EditableTextBox.cpp
+++ b/Engine/Source/Runtime/UMG/Private/Components/EditableTextBox.cpp
@@ -21,20 +21,9 @@ static FEditableTextBoxStyle* EditorEditableTextBoxStyle = nullptr;
 UEditableTextBox::UEditableTextBox(const FObjectInitializer& ObjectInitializer)
 	: Super(ObjectInitializer)
 {
-	ForegroundColor_DEPRECATED = FLinearColor::Black;
-	BackgroundColor_DEPRECATED = FLinearColor::White;
-	ReadOnlyForegroundColor_DEPRECATED = FLinearColor::Black;
-
-	if (!IsRunningDedicatedServer())
-	{
-		static ConstructorHelpers::FObjectFinder<UFont> RobotoFontObj(*UWidget::GetDefaultFontName());
-		Font_DEPRECATED = FSlateFontInfo(RobotoFontObj.Object, 12, FName("Bold"));
-	}
-
 	IsReadOnly = false;
 	IsPassword = false;
 	MinimumDesiredWidth = 0.0f;
-	Padding_DEPRECATED = FMargin(0, 0, 0, 0);
 	IsCaretMovedWhenGainFocus = true;
 	SelectAllTextWhenFocused = false;
 	RevertTextOnEscape = false;
@@ -245,55 +234,6 @@ void UEditableTextBox::HandleOnTextCommitted(const FText& InText, ETextCommit::T
 	OnTextCommitted.Broadcast(InText, CommitMethod);
 }
 
-void UEditableTextBox::PostLoad()
-{
-	Super::PostLoad();
-
-	if ( GetLinkerUEVersion() < VER_UE4_DEPRECATE_UMG_STYLE_ASSETS )
-	{
-		if ( Style_DEPRECATED != nullptr )
-		{
-			const FEditableTextBoxStyle* StylePtr = Style_DEPRECATED->GetStyle<FEditableTextBoxStyle>();
-			if ( StylePtr != nullptr )
-			{
-				WidgetStyle = *StylePtr;
-			}
-
-			Style_DEPRECATED = nullptr;
-		}
-	}
-
-	if (GetLinkerUEVersion() < VER_UE4_DEPRECATE_UMG_STYLE_OVERRIDES)
-	{
-		if (Font_DEPRECATED.HasValidFont())
-		{
-			WidgetStyle.Font = Font_DEPRECATED;
-			Font_DEPRECATED = FSlateFontInfo();
-		}
-
-		WidgetStyle.Padding = Padding_DEPRECATED;
-		Padding_DEPRECATED = FMargin(0);
-
-		if (ForegroundColor_DEPRECATED != FLinearColor::Black)
-		{
-			WidgetStyle.ForegroundColor = ForegroundColor_DEPRECATED;
-			ForegroundColor_DEPRECATED = FLinearColor::Black;
-		}
-
-		if (BackgroundColor_DEPRECATED != FLinearColor::White)
-		{
-			WidgetStyle.BackgroundColor = BackgroundColor_DEPRECATED;
-			BackgroundColor_DEPRECATED = FLinearColor::White;
-		}
-
-		if (ReadOnlyForegroundColor_DEPRECATED != FLinearColor::Black)
-		{
-			WidgetStyle.ReadOnlyForegroundColor = ReadOnlyForegroundColor_DEPRECATED;
-			ReadOnlyForegroundColor_DEPRECATED = FLinearColor::Black;
-		}
-	}
-}
-
 #if WITH_ACCESSIBILITY
 TSharedPtr<SWidget> UEditableTextBox::GetAccessibleWidget() const
 {
diff --git a/Engine/Source/Runtime/UMG/Private/Components/Image.cpp b/Engine/Source/Runtime/UMG/Private/Components/Image.cpp
--- a/Engine/Source/Runtime/UMG/Private/Components/Image.cpp
+++ b/Engine/Source/Runtime/UMG/Private/Components/Image.cpp
@@ -22,19 +22,6 @@ UImage::UImage(const FObjectInitializer& ObjectInitializer)
 {
 }
 
-#if WITH_EDITORONLY_DATA
-void UImage::PostLoad()
-{
-	Super::PostLoad();
-
-	if ( GetLinkerUEVersion() < VER_UE4_DEPRECATE_UMG_STYLE_ASSETS && Image_DEPRECATED != nullptr )
-	{
-		Brush = Image_DEPRECATED->Brush;
-		Image_DEPRECATED = nullptr;
-	}
-}
-#endif
-
 void UImage::ReleaseSlateResources(bool bReleaseChildren)
 {
 	Super::ReleaseSlateResources(bReleaseChildren);
diff --git a/Engine/Source/Runtime/UMG/Private/Components/InputKeySelector.cpp b/Engine/Source/Runtime/UMG/Private/Components/InputKeySelector.cpp
--- a/Engine/Source/Runtime/UMG/Private/Components/InputKeySelector.cpp
+++ b/Engine/Source/Runtime/UMG/Private/Components/InputKeySelector.cpp
@@ -90,17 +90,6 @@ void UInputKeySelector::Serialize(FArchive& Ar)
 	Ar.UsingCustomVersion(FFrameworkObjectVersion::GUID);
 }
 
-void UInputKeySelector::PostLoad()
-{
-	Super::PostLoad();
-
-	if (GetLinkerCustomVersion(FFrameworkObjectVersion::GUID) < FFrameworkObjectVersion::InputKeySelectorTextStyle)
-	{
-		TextStyle.Font = Font_DEPRECATED;
-		TextStyle.ColorAndOpacity = ColorAndOpacity_DEPRECATED;
-	}
-}
-
 void UInputKeySelector::SetSelectedKey( const FInputChord& InSelectedKey )
 {
 	if ( MyInputKeySelector.IsValid() )
diff --git a/Engine/Source/Runtime/UMG/Private/Components/MultiLineEditableText.cpp b/Engine/Source/Runtime/UMG/Private/Components/MultiLineEditableText.cpp
--- a/Engine/Source/Runtime/UMG/Private/Components/MultiLineEditableText.cpp
+++ b/Engine/Source/Runtime/UMG/Private/Components/MultiLineEditableText.cpp
@@ -62,9 +62,7 @@ UMultiLineEditableText::UMultiLineEditableText(const FObjectInitializer& ObjectI
 	if (!IsRunningDedicatedServer())
 	{
 		static ConstructorHelpers::FObjectFinder<UFont> RobotoFontObj(*UWidget::GetDefaultFontName());
-		Font_DEPRECATED = FSlateFontInfo(RobotoFontObj.Object, 12, FName("Bold"));
-
-		WidgetStyle.SetFont(Font_DEPRECATED);
+		WidgetStyle.SetFont(FSlateFontInfo(RobotoFontObj.Object, 12, FName("Bold")));
 	}
 }
 
@@ -206,20 +204,6 @@ void UMultiLineEditableText::HandleOnTextCommitted(const FText& InText, ETextCom
 	OnTextCommitted.Broadcast(InText, CommitMethod);
 }
 
-void UMultiLineEditableText::PostLoad()
-{
-	Super::PostLoad();
-
-	if (GetLinkerUEVersion() < VER_UE4_DEPRECATE_UMG_STYLE_OVERRIDES)
-	{
-		if (Font_DEPRECATED.HasValidFont())
-		{
-			WidgetStyle.Font = Font_DEPRECATED;
-			Font_DEPRECATED = FSlateFontInfo();
-		}
-	}
-}
-
 #if WITH_EDITOR
 
 const FText UMultiLineEditableText::GetPaletteCategory()
diff --git a/Engine/Source/Runtime/UMG/Private/Components/MultiLineEditableTextBox.cpp b/Engine/Source/Runtime/UMG/Private/Components/MultiLineEditableTextBox.cpp
--- a/Engine/Source/Runtime/UMG/Private/Components/MultiLineEditableTextBox.cpp
+++ b/Engine/Source/Runtime/UMG/Private/Components/MultiLineEditableTextBox.cpp
@@ -23,10 +23,6 @@ static FTextBlockStyle* EditorMultiLineEditableTextBoxTextStyle = nullptr;
 UMultiLineEditableTextBox::UMultiLineEditableTextBox(const FObjectInitializer& ObjectInitializer)
 	: Super(ObjectInitializer)
 {
-	ForegroundColor_DEPRECATED = FLinearColor::Black;
-	BackgroundColor_DEPRECATED = FLinearColor::White;
-	ReadOnlyForegroundColor_DEPRECATED = FLinearColor::Black;
-
 	if (DefaultMultiLineEditableTextBoxStyle == nullptr)
 	{
 		DefaultMultiLineEditableTextBoxStyle = new FEditableTextBoxStyle(FUMGCoreStyle::Get().GetWidgetStyle<FEditableTextBoxStyle>("NormalEditableTextBox"));
@@ -77,12 +73,6 @@ UMultiLineEditableTextBox::UMultiLineEditableTextBox(const FObjectInitializer& O
 	AllowContextMenu = true;
 	VirtualKeyboardDismissAction = EVirtualKeyboardDismissAction::TextChangeOnDismiss;
 	AutoWrapText = true;
-
-	if (!IsRunningDedicatedServer())
-	{
-		static ConstructorHelpers::FObjectFinder<UFont> RobotoFontObj(*UWidget::GetDefaultFontName());
-		Font_DEPRECATED = FSlateFontInfo(RobotoFontObj.Object, 12, FName("Bold"));
-	}
 }
 
 void UMultiLineEditableTextBox::ReleaseSlateResources(bool bReleaseChildren)
@@ -232,52 +222,6 @@ void UMultiLineEditableTextBox::HandleOnTextCommitted(const FText& InText, EText
 	OnTextCommitted.Broadcast(InText, CommitMethod);
 }
 
-void UMultiLineEditableTextBox::PostLoad()
-{
-	Super::PostLoad();
-
-	if ( GetLinkerUEVersion() < VER_UE4_DEPRECATE_UMG_STYLE_ASSETS )
-	{
-		if ( Style_DEPRECATED != nullptr )
-		{
-			const FEditableTextBoxStyle* StylePtr = Style_DEPRECATED->GetStyle<FEditableTextBoxStyle>();
-			if ( StylePtr != nullptr )
-			{
-				WidgetStyle = *StylePtr;
-			}
-
-			Style_DEPRECATED = nullptr;
-		}
-	}
-
-	if (GetLinkerUEVersion() < VER_UE4_DEPRECATE_UMG_STYLE_OVERRIDES)
-	{
-		if (Font_DEPRECATED.HasValidFont())
-		{
-			WidgetStyle.Font = Font_DEPRECATED;
-			Font_DEPRECATED = FSlateFontInfo();
-		}
-
-		if (ForegroundColor_DEPRECATED != FLinearColor::Black)
-		{
-			WidgetStyle.ForegroundColor = ForegroundColor_DEPRECATED;
-			ForegroundColor_DEPRECATED = FLinearColor::Black;
-		}
-
-		if (BackgroundColor_DEPRECATED != FLinearColor::White)
-		{
-			WidgetStyle.BackgroundColor = BackgroundColor_DEPRECATED;
-			BackgroundColor_DEPRECATED = FLinearColor::White;
-		}
-
-		if (ReadOnlyForegroundColor_DEPRECATED != FLinearColor::Black)
-		{
-			WidgetStyle.ReadOnlyForegroundColor = ReadOnlyForegroundColor_DEPRECATED;
-			ReadOnlyForegroundColor_DEPRECATED = FLinearColor::Black;
-		}
-	}
-}
-
 #if WITH_EDITOR
 
 const FText UMultiLineEditableTextBox::GetPaletteCategory()
diff --git a/Engine/Source/Runtime/UMG/Private/Components/ProgressBar.cpp b/Engine/Source/Runtime/UMG/Private/Components/ProgressBar.cpp
--- a/Engine/Source/Runtime/UMG/Private/Components/ProgressBar.cpp
+++ b/Engine/Source/Runtime/UMG/Private/Components/ProgressBar.cpp
@@ -113,43 +113,6 @@ void UProgressBar::SetPercent(float InPercent)
 	}
 }
 
-void UProgressBar::PostLoad()
-{
-	Super::PostLoad();
-
-	if (GetLinkerUEVersion() < VER_UE4_DEPRECATE_UMG_STYLE_ASSETS)
-	{
-		if (Style_DEPRECATED != nullptr)
-		{
-			const FProgressBarStyle* StylePtr = Style_DEPRECATED->GetStyle<FProgressBarStyle>();
-			if (StylePtr != nullptr)
-			{
-				WidgetStyle = *StylePtr;
-			}
-
-			Style_DEPRECATED = nullptr;
-		}
-
-		if (BackgroundImage_DEPRECATED != nullptr)
-		{
-			WidgetStyle.BackgroundImage = BackgroundImage_DEPRECATED->Brush;
-			BackgroundImage_DEPRECATED = nullptr;
-		}
-
-		if (FillImage_DEPRECATED != nullptr)
-		{
-			WidgetStyle.FillImage = FillImage_DEPRECATED->Brush;
-			FillImage_DEPRECATED = nullptr;
-		}
-
-		if (MarqueeImage_DEPRECATED != nullptr)
-		{
-			WidgetStyle.MarqueeImage = MarqueeImage_DEPRECATED->Brush;
-			MarqueeImage_DEPRECATED = nullptr;
-		}
-	}
-}
-
 #if WITH_EDITOR
 
 const FText UProgressBar::GetPaletteCategory()
diff --git a/Engine/Source/Runtime/UMG/Private/Components/ScrollBar.cpp b/Engine/Source/Runtime/UMG/Private/Components/ScrollBar.cpp
--- a/Engine/Source/Runtime/UMG/Private/Components/ScrollBar.cpp
+++ b/Engine/Source/Runtime/UMG/Private/Components/ScrollBar.cpp
@@ -117,25 +117,6 @@ void UScrollBar::Serialize(FArchive& Ar)
 	}
 }
 
-void UScrollBar::PostLoad()
-{
-	Super::PostLoad();
-
-	if ( GetLinkerUEVersion() < VER_UE4_DEPRECATE_UMG_STYLE_ASSETS )
-	{
-		if ( Style_DEPRECATED != nullptr )
-		{
-			const FScrollBarStyle* StylePtr = Style_DEPRECATED->GetStyle<FScrollBarStyle>();
-			if ( StylePtr != nullptr )
-			{
-				WidgetStyle = *StylePtr;
-			}
-
-			Style_DEPRECATED = nullptr;
-		}
-	}
-}
-
 #endif // if WITH_EDITORONLY_DATA
 
 #if WITH_EDITOR
diff --git a/Engine/Source/Runtime/UMG/Private/Components/ScrollBox.cpp b/Engine/Source/Runtime/UMG/Private/Components/ScrollBox.cpp
--- a/Engine/Source/Runtime/UMG/Private/Components/ScrollBox.cpp
+++ b/Engine/Source/Runtime/UMG/Private/Components/ScrollBox.cpp
@@ -263,36 +263,6 @@ void UScrollBox::Serialize(FArchive& Ar)
 	}
 }
 
-void UScrollBox::PostLoad()
-{
-	Super::PostLoad();
-
-	if ( GetLinkerUEVersion() < VER_UE4_DEPRECATE_UMG_STYLE_ASSETS )
-	{
-		if ( Style_DEPRECATED != nullptr )
-		{
-			const FScrollBoxStyle* StylePtr = Style_DEPRECATED->GetStyle<FScrollBoxStyle>();
-			if ( StylePtr != nullptr )
-			{
-				WidgetStyle = *StylePtr;
-			}
-
-			Style_DEPRECATED = nullptr;
-		}
-
-		if ( BarStyle_DEPRECATED != nullptr )
-		{
-			const FScrollBarStyle* StylePtr = BarStyle_DEPRECATED->GetStyle<FScrollBarStyle>();
-			if ( StylePtr != nullptr )
-			{
-				WidgetBarStyle = *StylePtr;
-			}
-
-			BarStyle_DEPRECATED = nullptr;
-		}
-	}
-}
-
 #endif // if WITH_EDITORONLY_DATA
 
 void UScrollBox::SetNavigationDestination(const EDescendantScrollDestination NewNavigationDestination)
diff --git a/Engine/Source/Runtime/UMG/Private/Components/SpinBox.cpp b/Engine/Source/Runtime/UMG/Private/Components/SpinBox.cpp
--- a/Engine/Source/Runtime/UMG/Private/Components/SpinBox.cpp
+++ b/Engine/Source/Runtime/UMG/Private/Components/SpinBox.cpp
@@ -400,26 +400,6 @@ void USpinBox::HandleOnEndSliderMovement(float InValue)
 	}
 }
 
-void USpinBox::PostLoad()
-{
-	Super::PostLoad();
-
-	if ( GetLinkerUEVersion() < VER_UE4_DEPRECATE_UMG_STYLE_ASSETS )
-	{
-		if ( Style_DEPRECATED != nullptr )
-		{
-			const FSpinBoxStyle* StylePtr = Style_DEPRECATED->GetStyle<FSpinBoxStyle>();
-			if ( StylePtr != nullptr )
-			{
-				WidgetStyle = *StylePtr;
-			}
-
-			Style_DEPRECATED = nullptr;
-		}
-	}
-}
-
-
 #if WITH_EDITOR
 
 const FText USpinBox::GetPaletteCategory()
diff --git a/Engine/Source/Runtime/UMG/Private/Components/TextBlock.cpp b/Engine/Source/Runtime/UMG/Private/Components/TextBlock.cpp
--- a/Engine/Source/Runtime/UMG/Private/Components/TextBlock.cpp
+++ b/Engine/Source/Runtime/UMG/Private/Components/TextBlock.cpp
@@ -23,7 +23,6 @@ UTextBlock::UTextBlock(const FObjectInitializer& ObjectInitializer)
 	ShadowOffset = FVector2D(1.0f, 1.0f);
 	ColorAndOpacity = FLinearColor::White;
 	ShadowColorAndOpacity = FLinearColor::Transparent;
-	bAutoWrapText_DEPRECATED = false;
 	TextTransformPolicy = ETextTransformPolicy::None;
 	TextOverflowPolicy = ETextOverflowPolicy::Clip;
 
@@ -39,17 +38,6 @@ UTextBlock::UTextBlock(const FObjectInitializer& ObjectInitializer)
 #endif
 }
 
-void UTextBlock::PostLoad()
-{
-	Super::PostLoad();
-
-	if (bAutoWrapText_DEPRECATED)
-	{
-		AutoWrapText = bAutoWrapText_DEPRECATED;
-		bAutoWrapText_DEPRECATED = false;
-	}
-}
-
 void UTextBlock::ReleaseSlateResources(bool bReleaseChildren)
 {
 	Super::ReleaseSlateResources(bReleaseChildren);
diff --git a/Engine/Source/Runtime/UMG/Private/Components/Throbber.cpp b/Engine/Source/Runtime/UMG/Private/Components/Throbber.cpp
--- a/Engine/Source/Runtime/UMG/Private/Components/Throbber.cpp
+++ b/Engine/Source/Runtime/UMG/Private/Components/Throbber.cpp
@@ -124,20 +124,6 @@ void UThrobber:: SetAnimateOpacity(bool bInAnimateOpacity)
 	}
 }
 
-void UThrobber::PostLoad()
-{
-	Super::PostLoad();
-
-	if ( GetLinkerUEVersion() < VER_UE4_DEPRECATE_UMG_STYLE_ASSETS )
-	{
-		if ( PieceImage_DEPRECATED != nullptr )
-		{
-			Image = PieceImage_DEPRECATED->Brush;
-			PieceImage_DEPRECATED = nullptr;
-		}
-	}
-}
-
 #if WITH_EDITOR
 
 const FText UThrobber::GetPaletteCategory()
diff --git a/Engine/Source/Runtime/UMG/Private/Components/Widget.cpp b/Engine/Source/Runtime/UMG/Private/Components/Widget.cpp
--- a/Engine/Source/Runtime/UMG/Private/Components/Widget.cpp
+++ b/Engine/Source/Runtime/UMG/Private/Components/Widget.cpp
@@ -686,12 +686,6 @@ void UWidget::SetNavigationRuleInternal(EUINavigation Direction, EUINavigationRu
 	}
 }
 
-void UWidget::SetNavigationRule(EUINavigation Direction, EUINavigationRule Rule, FName WidgetToFocus)
-{
-	SetNavigationRuleInternal(Direction, Rule, WidgetToFocus);
-	BuildNavigation();
-}
-
 void UWidget::SetNavigationRuleBase(EUINavigation Direction, EUINavigationRule Rule)
 {
 	if (Rule == EUINavigationRule::Explicit || Rule == EUINavigationRule::Custom || Rule == EUINavigationRule::CustomBoundary)
diff --git a/Engine/Source/Runtime/UMG/Private/Components/WrapBox.cpp b/Engine/Source/Runtime/UMG/Private/Components/WrapBox.cpp
--- a/Engine/Source/Runtime/UMG/Private/Components/WrapBox.cpp
+++ b/Engine/Source/Runtime/UMG/Private/Components/WrapBox.cpp
@@ -19,25 +19,6 @@ UWrapBox::UWrapBox(const FObjectInitializer& ObjectInitializer)
 	Orientation = EOrientation::Orient_Horizontal;
 }
 
-void UWrapBox::PostLoad()
-{
-	Super::PostLoad();
-
-#if WITH_EDITORONLY_DATA
-	if (WrapWidth_DEPRECATED != 0.0f)
-	{
-		WrapSize = WrapWidth_DEPRECATED;
-		WrapWidth_DEPRECATED = 0.0f;
-	}
-
-	if (bExplicitWrapWidth_DEPRECATED)
-	{
-		bExplicitWrapSize = bExplicitWrapWidth_DEPRECATED;
-		bExplicitWrapWidth_DEPRECATED = false;
-	}
-#endif
-}
-
 void UWrapBox::ReleaseSlateResources(bool bReleaseChildren)
 {
 	Super::ReleaseSlateResources(bReleaseChildren);
@@ -72,11 +53,6 @@ void UWrapBox::OnSlotRemoved(UPanelSlot* InSlot)
 	}
 }
 
-UWrapBoxSlot* UWrapBox::AddChildWrapBox(UWidget* Content)
-{
-	return AddChildToWrapBox(Content);
-}
-
 UWrapBoxSlot* UWrapBox::AddChildToWrapBox(UWidget* Content)
 {
 	return Cast<UWrapBoxSlot>(Super::AddChild(Content));
diff --git a/Engine/Source/Runtime/UMG/Private/UserWidget.cpp b/Engine/Source/Runtime/UMG/Private/UserWidget.cpp
--- a/Engine/Source/Runtime/UMG/Private/UserWidget.cpp
+++ b/Engine/Source/Runtime/UMG/Private/UserWidget.cpp
@@ -78,7 +78,6 @@ UUserWidget::UUserWidget(const FObjectInitializer& ObjectInitializer)
 	ViewportAnchors = FAnchors(0, 0, 1, 1);
 	Visibility = ESlateVisibility::SelfHitTestInvisible;
 
-	bSupportsKeyboardFocus_DEPRECATED = true;
 	bIsFocusable = false;
 	ColorAndOpacity = FLinearColor::White;
 	ForegroundColor = FSlateColor::UseForeground();
@@ -460,11 +459,6 @@ void UUserWidget::DisableAnimations()
 	}
 }
 
-void UUserWidget::Invalidate()
-{
-	Invalidate(EInvalidateWidgetReason::LayoutAndVolatility);
-}
-
 void UUserWidget::Invalidate(EInvalidateWidgetReason InvalidateReason)
 {
 	if (TSharedPtr<SWidget> CachedWidget = GetCachedWidget())
@@ -1918,21 +1912,6 @@ void UUserWidget::PostLoad()
 #endif
 }
 
-void UUserWidget::Serialize(FArchive& Ar)
-{
-	Super::Serialize(Ar);
-
-	Ar.UsingCustomVersion(FEditorObjectVersion::GUID);
-
-	if ( Ar.IsLoading() )
-	{
-		if ( Ar.UEVer() < VER_UE4_USERWIDGET_DEFAULT_FOCUSABLE_FALSE )
-		{
-			bIsFocusable = bSupportsKeyboardFocus_DEPRECATED;
-		}
-	}
-}
-
 /////////////////////////////////////////////////////
 
 UUserWidget* UUserWidget::CreateWidgetInstance(UWidget& OwningWidget, TSubclassOf<UUserWidget> UserWidgetClass, FName WidgetName)
diff --git a/Engine/Source/Runtime/UMG/Public/Animation/UMGSequencePlayer.h b/Engine/Source/Runtime/UMG/Public/Animation/UMGSequencePlayer.h
--- a/Engine/Source/Runtime/UMG/Public/Animation/UMGSequencePlayer.h
+++ b/Engine/Source/Runtime/UMG/Public/Animation/UMGSequencePlayer.h
@@ -40,10 +40,6 @@ public:
 	/** Reverses a running animation */
 	void Reverse();
 
-	/** Gets the current time position in the player (in seconds). */
-	UE_DEPRECATED(4.20, "Please use GetCurrentTime instead.")
-	double GetTimeCursorPosition() const { return GetCurrentTime().AsSeconds(); }
-
 	void SetCurrentTime(float InTime) { TimeCursorPosition = AnimationResolution.AsFrameTime(InTime); }
 	FQualifiedFrameTime GetCurrentTime() const { return FQualifiedFrameTime(TimeCursorPosition, AnimationResolution); }
 
diff --git a/Engine/Source/Runtime/UMG/Public/Animation/WidgetAnimation.h b/Engine/Source/Runtime/UMG/Public/Animation/WidgetAnimation.h
--- a/Engine/Source/Runtime/UMG/Public/Animation/WidgetAnimation.h
+++ b/Engine/Source/Runtime/UMG/Public/Animation/WidgetAnimation.h
@@ -66,16 +66,6 @@ public:
 	UFUNCTION(BlueprintCallable, Category="Animation")
 	UMG_API float GetEndTime() const;
 
-#if WITH_EDITORONLY_DATA
-	/** Fires when the widget animation starts playing. */
-	UPROPERTY()
-	FOnWidgetAnimationPlaybackStatusChanged OnAnimationStarted_DEPRECATED;
-
-	/** Fires when the widget animation is finished. */
-	UPROPERTY()
-	FOnWidgetAnimationPlaybackStatusChanged OnAnimationFinished_DEPRECATED;
-#endif
-
 	// These animation binding functions were added so that we could cleanly upgrade assets
 	// from before animation sharing, they don't actually modify the animation, they just pipe
 	// through to the UUserWidget.  If we didn't put the functions here, it would be much more
diff --git a/Engine/Source/Runtime/UMG/Public/Blueprint/IUserObjectListEntry.h b/Engine/Source/Runtime/UMG/Public/Blueprint/IUserObjectListEntry.h
--- a/Engine/Source/Runtime/UMG/Public/Blueprint/IUserObjectListEntry.h
+++ b/Engine/Source/Runtime/UMG/Public/Blueprint/IUserObjectListEntry.h
@@ -36,12 +36,6 @@ protected:
 	UFUNCTION(BlueprintImplementableEvent, Category = ObjectListEntry)
 	void OnListItemObjectSet(UObject* ListItemObject);
 
-	UE_DEPRECATED(4.23, "Implementers of IUserObjectListEntry no longer need to manually implement the GetListItemObject method.")
-	virtual UObject* GetListItemObject_Implementation() const { return nullptr; }
-
-	UE_DEPRECATED(4.23, "Renamed to NativeOnListItemObjectSet, as caching the item object is no longer a hard requirement.")
-	virtual void SetListItemObjectInternal(UObject* InObject) {}
-
 private:
 	UObject* GetListItemObjectInternal() const;
 	
diff --git a/Engine/Source/Runtime/UMG/Public/Blueprint/UserWidget.h b/Engine/Source/Runtime/UMG/Public/Blueprint/UserWidget.h
--- a/Engine/Source/Runtime/UMG/Public/Blueprint/UserWidget.h
+++ b/Engine/Source/Runtime/UMG/Public/Blueprint/UserWidget.h
@@ -216,7 +216,6 @@ public:
 	virtual void PostDuplicate(bool bDuplicateForPIE) override;
 	virtual void BeginDestroy() override;
 	virtual void PostLoad() override;
-	virtual void Serialize(FArchive& Ar) override;
 	//~ End UObject Interface
 
 	void DuplicateAndInitializeFromWidgetTree(UWidgetTree* InWidgetTree);
@@ -920,21 +919,6 @@ public:
 	UFUNCTION(BlueprintCallable, BlueprintCosmetic, Category = "User Interface|Animation")
 	UUMGSequencePlayer* PlayAnimation(UWidgetAnimation* InAnimation, float StartAtTime = 0.0f, int32 NumLoopsToPlay = 1, EUMGSequencePlayMode::Type PlayMode = EUMGSequencePlayMode::Forward, float PlaybackSpeed = 1.0f, bool bRestoreState = false);
 
-	/**
-	 * Plays an animation in this widget a specified number of times
-	 * 
-	 * @param InAnimation The animation to play
-	 * @param StartAtTime The time in the animation from which to start playing, relative to the start position. For looped animations, this will only affect the first playback of the animation.
-	 * @param NumLoopsToPlay The number of times to loop this animation (0 to loop indefinitely)
-	 * @param PlaybackSpeed The speed at which the animation should play
-	 * @param PlayMode Specifies the playback mode
-	 */
-	UE_DEPRECATED(4.22, "Short lived attempt to clarify what the default PlayAnimation function does, but going to just keep the default one to make things simple by default.")
-	UUMGSequencePlayer* PlayAnimationAtTime(UWidgetAnimation* InAnimation, float StartAtTime = 0.0f, int32 NumLoopsToPlay = 1, EUMGSequencePlayMode::Type PlayMode = EUMGSequencePlayMode::Forward, float PlaybackSpeed = 1.0f)
-	{
-		return PlayAnimation(InAnimation, StartAtTime, NumLoopsToPlay, PlayMode, PlaybackSpeed);
-	}
-
 	/**
 	 * Plays an animation in this widget a specified number of times stopping at a specified time
 	 * 
@@ -1220,9 +1204,6 @@ public:
 	UPROPERTY(EditAnywhere, BlueprintReadOnly, Category = "Input")
 	int32 Priority;
 
-	UPROPERTY()
-	uint8 bSupportsKeyboardFocus_DEPRECATED:1;
-
 	/** Setting this flag to true, allows this widget to accept focus when clicked, or when navigated to. */
 	UPROPERTY(EditAnywhere, BlueprintReadWrite, Category="Interaction")
 	uint8 bIsFocusable : 1;
@@ -1269,12 +1250,8 @@ protected:
 	virtual void NativePreConstruct();
 	virtual void NativeConstruct();
 	virtual void NativeDestruct();
-
 	virtual void NativeTick(const FGeometry& MyGeometry, float InDeltaTime);
 
-	UE_DEPRECATED(4.20, "Please override the other version of NativePaint that accepts all the parameters, not just the paint context.")
-	virtual void NativePaint(FPaintContext& InContext) const { }
-
 	/**
 	 * Native implemented paint function for the Widget
 	 * Returns the maximum LayerID painted on
@@ -1321,8 +1298,6 @@ protected:
 	virtual FReply NativeOnTouchForceChanged(const FGeometry& MyGeometry, const FPointerEvent& TouchEvent);
 	virtual FCursorReply NativeOnCursorQuery( const FGeometry& InGeometry, const FPointerEvent& InCursorEvent );
 	virtual FNavigationReply NativeOnNavigation(const FGeometry& InGeometry, const FNavigationEvent& InNavigationEvent);
-	UE_DEPRECATED(4.20, "Please use NativeOnMouseCaptureLost(const FCaptureLostEvent& CaptureLostEvent)")
-	void NativeOnMouseCaptureLost() {}
 	virtual void NativeOnMouseCaptureLost(const FCaptureLostEvent& CaptureLostEvent);
 
 protected:
@@ -1342,9 +1317,6 @@ protected:
 
 	void DisableAnimations();
 
-	UE_DEPRECATED(4.21, "You now need to provide the reason you're invalidating.")
-	void Invalidate();
-
 	void Invalidate(EInvalidateWidgetReason InvalidateReason);
 	
 	/**
diff --git a/Engine/Source/Runtime/UMG/Public/Components/Border.h b/Engine/Source/Runtime/UMG/Public/Components/Border.h
--- a/Engine/Source/Runtime/UMG/Public/Components/Border.h
+++ b/Engine/Source/Runtime/UMG/Public/Components/Border.h
@@ -197,12 +197,5 @@ protected:
 	/** Translates the bound brush data and assigns it to the cached brush used by this widget. */
 	const FSlateBrush* ConvertImage(TAttribute<FSlateBrush> InImageAsset) const;
 
-#if WITH_EDITORONLY_DATA
-protected:
-	/** Image to use for the border */
-	UPROPERTY()
-	TObjectPtr<USlateBrushAsset> Brush_DEPRECATED;
-#endif
-
 	PROPERTY_BINDING_IMPLEMENTATION(FLinearColor, ContentColorAndOpacity)
 };
diff --git a/Engine/Source/Runtime/UMG/Public/Components/Button.h b/Engine/Source/Runtime/UMG/Public/Components/Button.h
--- a/Engine/Source/Runtime/UMG/Public/Components/Button.h
+++ b/Engine/Source/Runtime/UMG/Public/Components/Button.h
@@ -33,10 +33,6 @@ class UMG_API UButton : public UContentWidget
 	GENERATED_UCLASS_BODY()
 
 public:
-	/** The template style asset, used to seed the mutable instance of the style. */
-	UPROPERTY()
-	TObjectPtr<USlateWidgetStyleAsset> Style_DEPRECATED;
-
 	/** The button style used at runtime */
 	UPROPERTY(EditAnywhere, BlueprintReadWrite, Category="Appearance", meta=( DisplayName="Style" ))
 	FButtonStyle WidgetStyle;
diff --git a/Engine/Source/Runtime/UMG/Public/Components/CheckBox.h b/Engine/Source/Runtime/UMG/Public/Components/CheckBox.h
--- a/Engine/Source/Runtime/UMG/Public/Components/CheckBox.h
+++ b/Engine/Source/Runtime/UMG/Public/Components/CheckBox.h
@@ -45,58 +45,10 @@ public:
 	UPROPERTY(EditAnywhere, BlueprintReadWrite, Category="Style", meta=( DisplayName="Style" ))
 	FCheckBoxStyle WidgetStyle;
 
-	/** Style of the check box */
-	UPROPERTY()
-	TObjectPtr<USlateWidgetStyleAsset> Style_DEPRECATED;
-
-	/** Image to use when the checkbox is unchecked */
-	UPROPERTY()
-	TObjectPtr<USlateBrushAsset> UncheckedImage_DEPRECATED;
-	
-	/** Image to use when the checkbox is unchecked and hovered */
-	UPROPERTY()
-	TObjectPtr<USlateBrushAsset> UncheckedHoveredImage_DEPRECATED;
-	
-	/** Image to use when the checkbox is unchecked and pressed */
-	UPROPERTY()
-	TObjectPtr<USlateBrushAsset> UncheckedPressedImage_DEPRECATED;
-	
-	/** Image to use when the checkbox is checked */
-	UPROPERTY()
-	TObjectPtr<USlateBrushAsset> CheckedImage_DEPRECATED;
-	
-	/** Image to use when the checkbox is checked and hovered */
-	UPROPERTY()
-	TObjectPtr<USlateBrushAsset> CheckedHoveredImage_DEPRECATED;
-	
-	/** Image to use when the checkbox is checked and pressed */
-	UPROPERTY()
-	TObjectPtr<USlateBrushAsset> CheckedPressedImage_DEPRECATED;
-	
-	/** Image to use when the checkbox is in an ambiguous state and hovered */
-	UPROPERTY()
-	TObjectPtr<USlateBrushAsset> UndeterminedImage_DEPRECATED;
-	
-	/** Image to use when the checkbox is checked and hovered */
-	UPROPERTY()
-	TObjectPtr<USlateBrushAsset> UndeterminedHoveredImage_DEPRECATED;
-	
-	/** Image to use when the checkbox is in an ambiguous state and pressed */
-	UPROPERTY()
-	TObjectPtr<USlateBrushAsset> UndeterminedPressedImage_DEPRECATED;
-
 	/** How the content of the toggle button should align within the given space */
 	UPROPERTY(EditAnywhere, BlueprintReadOnly, Category=Appearance)
 	TEnumAsByte<EHorizontalAlignment> HorizontalAlignment;
 
-	/** Spacing between the check box image and its content */
-	UPROPERTY()
-	FMargin Padding_DEPRECATED;
-
-	/** The color of the background border */
-	UPROPERTY()
-	FSlateColor BorderBackgroundColor_DEPRECATED;
-
 	/** The type of mouse action required by the user to trigger the buttons 'Click' */
 	UPROPERTY(EditAnywhere, BlueprintReadOnly, Category = "Interaction", AdvancedDisplay)
 	TEnumAsByte<EButtonClickMethod::Type> ClickMethod;
@@ -160,10 +112,6 @@ public:
 	virtual void ReleaseSlateResources(bool bReleaseChildren) override;
 	//~ End UVisual Interface
 
-	//~ Begin UObject Interface
-	virtual void PostLoad() override;
-	//~ End UObject Interface
-
 #if WITH_EDITOR
 	virtual const FText GetPaletteCategory() override;
 #endif
diff --git a/Engine/Source/Runtime/UMG/Public/Components/CircularThrobber.h b/Engine/Source/Runtime/UMG/Public/Components/CircularThrobber.h
--- a/Engine/Source/Runtime/UMG/Public/Components/CircularThrobber.h
+++ b/Engine/Source/Runtime/UMG/Public/Components/CircularThrobber.h
@@ -37,10 +37,6 @@ public:
 	UPROPERTY(EditAnywhere, BlueprintReadOnly, Category=Appearance, meta=(EditCondition="bEnableRadius"))
 	float Radius;
 
-	/** Image to use for each segment of the throbber */
-	UPROPERTY()
-	TObjectPtr<USlateBrushAsset> PieceImage_DEPRECATED;
-
 	UPROPERTY(EditAnywhere, BlueprintReadOnly, Category=Appearance)
 	FSlateBrush Image;
 
@@ -66,10 +62,6 @@ public:
 	virtual void ReleaseSlateResources(bool bReleaseChildren) override;
 	//~ End UVisual Interface
 
-	//~ Begin UObject Interface
-	virtual void PostLoad() override;
-	//~ End UObject Interface
-
 #if WITH_EDITOR
 	virtual const FText GetPaletteCategory() override;
 #endif
diff --git a/Engine/Source/Runtime/UMG/Public/Components/EditableText.h b/Engine/Source/Runtime/UMG/Public/Components/EditableText.h
--- a/Engine/Source/Runtime/UMG/Public/Components/EditableText.h
+++ b/Engine/Source/Runtime/UMG/Public/Components/EditableText.h
@@ -53,30 +53,6 @@ public:
 	UPROPERTY(EditAnywhere, BlueprintReadWrite, Category=Appearance, meta=(DisplayName="Style", ShowOnlyInnerProperties))
 	FEditableTextStyle WidgetStyle;
 
-	/** Text style */
-	UPROPERTY()
-	TObjectPtr<USlateWidgetStyleAsset> Style_DEPRECATED;
-
-	/** Background image for the selected text (overrides Style) */
-	UPROPERTY()
-	TObjectPtr<USlateBrushAsset> BackgroundImageSelected_DEPRECATED;
-
-	/** Background image for the composing text (overrides Style) */
-	UPROPERTY()
-	TObjectPtr<USlateBrushAsset> BackgroundImageComposing_DEPRECATED;
-
-	/** Image brush used for the caret (overrides Style) */
-	UPROPERTY()
-	TObjectPtr<USlateBrushAsset> CaretImage_DEPRECATED;
-
-	/** Font color and opacity (overrides Style) */
-	UPROPERTY()
-	FSlateFontInfo Font_DEPRECATED;
-
-	/** Text color and opacity (overrides Style) */
-	UPROPERTY()
-	FSlateColor ColorAndOpacity_DEPRECATED;
-
 	/** Sets whether this text box can actually be modified interactively by the user */
 	UPROPERTY(EditAnywhere, Category=Appearance)
 	bool IsReadOnly;
@@ -194,10 +170,6 @@ public:
 	virtual void ReleaseSlateResources(bool bReleaseChildren) override;
 	//~ End UVisual Interface
 
-	//~ Begin UObject Interface
-	virtual void PostLoad() override;
-	//~ End UObject Interface
-
 #if WITH_EDITOR
 	virtual const FText GetPaletteCategory() override;
 #endif
diff --git a/Engine/Source/Runtime/UMG/Public/Components/EditableTextBox.h b/Engine/Source/Runtime/UMG/Public/Components/EditableTextBox.h
--- a/Engine/Source/Runtime/UMG/Public/Components/EditableTextBox.h
+++ b/Engine/Source/Runtime/UMG/Public/Components/EditableTextBox.h
@@ -47,10 +47,6 @@ public:
 	UPROPERTY(EditAnywhere, BlueprintReadWrite, Category=Appearance, meta=(DisplayName="Style", ShowOnlyInnerProperties))
 	FEditableTextBoxStyle WidgetStyle;
 
-	/** Style used for the text box */
-	UPROPERTY()
-	TObjectPtr<USlateWidgetStyleAsset> Style_DEPRECATED;
-
 	/** Hint text that appears when there is no text in the text box */
 	UPROPERTY(EditAnywhere, BlueprintReadOnly, Category=Content)
 	FText HintText;
@@ -59,22 +55,6 @@ public:
 	UPROPERTY()
 	FGetText HintTextDelegate;
 
-	/** Font color and opacity (overrides Style) */
-	UPROPERTY()
-	FSlateFontInfo Font_DEPRECATED;
-
-	/** Text color and opacity (overrides Style) */
-	UPROPERTY()
-	FLinearColor ForegroundColor_DEPRECATED;
-
-	/** The color of the background/border around the editable text (overrides Style) */
-	UPROPERTY()
-	FLinearColor BackgroundColor_DEPRECATED;
-
-	/** Text color and opacity when read-only (overrides Style) */
-	UPROPERTY()
-	FLinearColor ReadOnlyForegroundColor_DEPRECATED;
-
 	/** Sets whether this text box can actually be modified interactively by the user */
 	UPROPERTY(EditAnywhere, BlueprintReadOnly, Category=Appearance)
 	bool IsReadOnly;
@@ -87,10 +67,6 @@ public:
 	UPROPERTY(EditAnywhere, BlueprintReadOnly, Category=Appearance)
 	float MinimumDesiredWidth;
 
-	/** Padding between the box/border and the text widget inside (overrides Style) */
-	UPROPERTY()
-	FMargin Padding_DEPRECATED;
-
 	/** Workaround as we lose focus when the auto completion closes. */
 	UPROPERTY(EditAnywhere, BlueprintReadOnly, Category=Behavior, AdvancedDisplay)
 	bool IsCaretMovedWhenGainFocus;
@@ -202,10 +178,6 @@ public:
 	virtual void ReleaseSlateResources(bool bReleaseChildren) override;
 	//~ End UVisual Interface
 
-	//~ Begin UObject Interface
-	virtual void PostLoad() override;
-	//~ End UObject Interface
-
 #if WITH_EDITOR
 	virtual const FText GetPaletteCategory() override;
 #endif
diff --git a/Engine/Source/Runtime/UMG/Public/Components/Image.h b/Engine/Source/Runtime/UMG/Public/Components/Image.h
--- a/Engine/Source/Runtime/UMG/Public/Components/Image.h
+++ b/Engine/Source/Runtime/UMG/Public/Components/Image.h
@@ -33,12 +33,6 @@ class UMG_API UImage : public UWidget
 
 public:
 
-#if WITH_EDITORONLY_DATA
-	/** Image to draw */
-	UPROPERTY()
-	TObjectPtr<USlateBrushAsset> Image_DEPRECATED;
-#endif
-
 	/** Image to draw */
 	UPROPERTY(EditAnywhere, BlueprintReadOnly, Category=Appearance)
 	FSlateBrush Brush;
@@ -153,12 +147,6 @@ public:
 	virtual void ReleaseSlateResources(bool bReleaseChildren) override;
 	//~ End UVisual Interface
 
-#if WITH_EDITORONLY_DATA
-	//~ Begin UObject Interface
-	virtual void PostLoad() override;
-	//~ End UObject Interface
-#endif
-
 #if WITH_EDITOR
 	//~ Begin UWidget Interface
 	virtual const FText GetPaletteCategory() override;
diff --git a/Engine/Source/Runtime/UMG/Public/Components/InputKeySelector.h b/Engine/Source/Runtime/UMG/Public/Components/InputKeySelector.h
--- a/Engine/Source/Runtime/UMG/Public/Components/InputKeySelector.h
+++ b/Engine/Source/Runtime/UMG/Public/Components/InputKeySelector.h
@@ -39,16 +39,10 @@ public:
 	UPROPERTY( BlueprintReadOnly, Category = "Key Selection" )
 	FInputChord SelectedKey;
 
-	UPROPERTY()
-	FSlateFontInfo Font_DEPRECATED;
-
 	/** The amount of blank space around the text used to display the currently selected key. */
 	UPROPERTY( EditAnywhere, BlueprintReadOnly, Category = Appearance )
 	FMargin Margin;
 
-	UPROPERTY()
-	FLinearColor ColorAndOpacity_DEPRECATED;
-
 	/** Sets the text which is displayed while selecting keys. */
 	UPROPERTY( EditAnywhere, BlueprintReadOnly, Category = Appearance )
 	FText KeySelectionText;
@@ -123,7 +117,6 @@ public:
 
 protected:
 	//~ Begin UObject Interface
-	virtual void PostLoad() override;
 	virtual void Serialize(FArchive& Ar) override;
 	//~ End UObject Interface
 
diff --git a/Engine/Source/Runtime/UMG/Public/Components/ListViewBase.h b/Engine/Source/Runtime/UMG/Public/Components/ListViewBase.h
--- a/Engine/Source/Runtime/UMG/Public/Components/ListViewBase.h
+++ b/Engine/Source/Runtime/UMG/Public/Components/ListViewBase.h
@@ -341,70 +341,6 @@ protected:
 			.OnExpansionChanged_UObject(Implementer, &UListViewBaseT::HandleExpansionChanged);
 	}
 
-	template <template<typename> class ListViewT = SListView, typename UListViewBaseT>
-	UE_DEPRECATED(4.23, "Use ConstructListView with FListViewConstructArgs instead of individual arguments.")
-	static TSharedRef<ListViewT<ItemType>> ConstructListView(UListViewBaseT* Implementer,
-		const TArray<ItemType>& ListItems,
-		bool bAllowFocus,
-		ESelectionMode::Type SelectionMode = ESelectionMode::Single,
-		bool bClearSelectionOnClick = false,
-		EConsumeMouseWheel ConsumeMouseWheel = EConsumeMouseWheel::WhenScrollingPossible,
-		bool bReturnFocusToSelection = false,
-		EOrientation Orientation = Orient_Vertical)
-	{
-		FListViewConstructArgs Args;
-		Args.bAllowFocus = bAllowFocus;
-		Args.SelectionMode = SelectionMode;
-		Args.bClearSelectionOnClick = bClearSelectionOnClick;
-		Args.ConsumeMouseWheel = ConsumeMouseWheel;
-		Args.bReturnFocusToSelection = bReturnFocusToSelection;
-		Args.Orientation = Orientation;
-
-		return ConstructListView<ListViewT>(Implementer, ListItems, Args);
-	}
-
-	template <template<typename> class TileViewT = STileView, typename UListViewBaseT>
-	UE_DEPRECATED(4.23, "Use ConstructTileView with FTileViewConstructArgs instead of individual arguments.")
-	static TSharedRef<TileViewT<ItemType>> ConstructTileView(UListViewBaseT* Implementer,
-		const TArray<ItemType>& ListItems,
-		EListItemAlignment TileAlignment,
-		TAttribute<float> TileHeight,
-		TAttribute<float> TileWidth,
-		ESelectionMode::Type SelectionMode = ESelectionMode::Single,
-		bool bClearSelectionOnClick = false,
-		bool bWrapDirectionalNavigation = false,
-		EConsumeMouseWheel ConsumeMouseWheel = EConsumeMouseWheel::WhenScrollingPossible,
-		EOrientation Orientation = Orient_Vertical)
-	{
-		FTileViewConstructArgs Args;
-		Args.SelectionMode = SelectionMode;
-		Args.bClearSelectionOnClick = bClearSelectionOnClick;
-		Args.ConsumeMouseWheel = ConsumeMouseWheel;
-		Args.Orientation = Orientation;
-		Args.bWrapDirectionalNavigation = bWrapDirectionalNavigation;
-		Args.TileHeight = TileHeight;
-		Args.TileWidth = TileWidth;
-		Args.TileAlignment = TileAlignment;
-
-		return ConstructTileView<TileViewT>(Implementer, ListItems, Args);
-	}
-
-	template <template<typename> class TreeViewT = STreeView, typename UListViewBaseT>
-	UE_DEPRECATED(4.23, "Use ConstructTreeView with FTreeViewConstructArgs instead of individual arguments.")
-	static TSharedRef<TreeViewT<ItemType>> ConstructTreeView(UListViewBaseT* Implementer,
-		const TArray<ItemType>& ListItems,
-		ESelectionMode::Type SelectionMode,
-		bool bClearSelectionOnClick = false,
-		EConsumeMouseWheel ConsumeMouseWheel = EConsumeMouseWheel::WhenScrollingPossible)
-	{
-		FTreeViewConstructArgs Args;
-		Args.SelectionMode = SelectionMode;
-		Args.bClearSelectionOnClick = bClearSelectionOnClick;
-		Args.ConsumeMouseWheel = ConsumeMouseWheel;
-
-		return ConstructTreeView<TreeViewT>(Implementer, ListItems, Args);
-	}
-
 protected:
 	/** Gets the SObjectTableRow underlying the UMG EntryWidget that represents the given item (if one exists) */
 	template <template<typename> class ObjectRowT = SObjectTableRow>
diff --git a/Engine/Source/Runtime/UMG/Public/Components/MultiLineEditableText.h b/Engine/Source/Runtime/UMG/Public/Components/MultiLineEditableText.h
--- a/Engine/Source/Runtime/UMG/Public/Components/MultiLineEditableText.h
+++ b/Engine/Source/Runtime/UMG/Public/Components/MultiLineEditableText.h
@@ -48,10 +48,6 @@ public:
 	UPROPERTY(EditAnywhere, BlueprintReadOnly, Category="Appearance")
 	bool bIsReadOnly;
 
-	/** Font color and opacity (overrides Style) */
-	UPROPERTY()
-	FSlateFontInfo Font_DEPRECATED;
-
 	/** Whether to select all text when the user clicks to give focus on the widget */
 	UPROPERTY(EditAnywhere, Category=Behavior, AdvancedDisplay)
 	bool SelectAllTextWhenFocused;
@@ -124,10 +120,6 @@ public:
 	virtual void ReleaseSlateResources(bool bReleaseChildren) override;
 	//~ End UVisual Interface
 
-	//~ Begin UObject Interface
-	virtual void PostLoad() override;
-	//~ End UObject Interface
-
 #if WITH_EDITOR
 	virtual const FText GetPaletteCategory() override;
 #endif
diff --git a/Engine/Source/Runtime/UMG/Public/Components/MultiLineEditableTextBox.h b/Engine/Source/Runtime/UMG/Public/Components/MultiLineEditableTextBox.h
--- a/Engine/Source/Runtime/UMG/Public/Components/MultiLineEditableTextBox.h
+++ b/Engine/Source/Runtime/UMG/Public/Components/MultiLineEditableTextBox.h
@@ -66,25 +66,6 @@ public:
 	UPROPERTY(EditAnywhere, Category=Behavior, AdvancedDisplay)
 	EVirtualKeyboardDismissAction VirtualKeyboardDismissAction;
 
-	UPROPERTY()
-	TObjectPtr<USlateWidgetStyleAsset> Style_DEPRECATED;
-
-	/** Font color and opacity (overrides Style) */
-	UPROPERTY()
-	FSlateFontInfo Font_DEPRECATED;
-
-	/** Text color and opacity (overrides Style) */
-	UPROPERTY()
-	FLinearColor ForegroundColor_DEPRECATED;
-
-	/** The color of the background/border around the editable text (overrides Style) */
-	UPROPERTY()
-	FLinearColor BackgroundColor_DEPRECATED;
-
-	/** Text color and opacity when read-only (overrides Style) */
-	UPROPERTY()
-	FLinearColor ReadOnlyForegroundColor_DEPRECATED;
-
 	/** Called whenever the text is changed programmatically or interactively by the user */
 	UPROPERTY(BlueprintAssignable, Category="Widget Event", meta=(DisplayName="OnTextChanged (Multi-Line Text Box)"))
 	FOnMultiLineEditableTextBoxChangedEvent OnTextChanged;
@@ -143,10 +124,6 @@ public:
 	virtual void ReleaseSlateResources(bool bReleaseChildren) override;
 	//~ End UVisual Interface
 
-	//~ Begin UObject Interface
-	virtual void PostLoad() override;
-	//~ End UObject Interface
-
 #if WITH_EDITOR
 	virtual const FText GetPaletteCategory() override;
 #endif
diff --git a/Engine/Source/Runtime/UMG/Public/Components/ProgressBar.h b/Engine/Source/Runtime/UMG/Public/Components/ProgressBar.h
--- a/Engine/Source/Runtime/UMG/Public/Components/ProgressBar.h
+++ b/Engine/Source/Runtime/UMG/Public/Components/ProgressBar.h
@@ -28,22 +28,6 @@ public:
 	UPROPERTY(EditAnywhere, BlueprintReadWrite, Category="Style", meta=( DisplayName="Style" ))
 	FProgressBarStyle WidgetStyle;
 
-	/** Style used for the progress bar */
-	UPROPERTY()
-	TObjectPtr<USlateWidgetStyleAsset> Style_DEPRECATED;
-
-	/** The brush to use as the background of the progress bar */
-	UPROPERTY()
-	TObjectPtr<USlateBrushAsset> BackgroundImage_DEPRECATED;
-	
-	/** The brush to use as the fill image */
-	UPROPERTY()
-	TObjectPtr<USlateBrushAsset> FillImage_DEPRECATED;
-	
-	/** The brush to use as the marquee image */
-	UPROPERTY()
-	TObjectPtr<USlateBrushAsset> MarqueeImage_DEPRECATED;
-
 	/** Used to determine the fill position of the progress bar ranging 0..1 */
 	UPROPERTY(EditAnywhere, BlueprintReadOnly, Category=Progress, meta=( UIMin = "0", UIMax = "1" ))
 	float Percent;
@@ -100,10 +84,6 @@ public:
 	virtual void ReleaseSlateResources(bool bReleaseChildren) override;
 	//~ End UVisual Interface
 
-	//~ Begin UObject Interface
-	virtual void PostLoad() override;
-	//~ End UObject Interface
-
 #if WITH_EDITOR
 	//~ Begin UWidget Interface
 	virtual const FText GetPaletteCategory() override;
diff --git a/Engine/Source/Runtime/UMG/Public/Components/ScrollBar.h b/Engine/Source/Runtime/UMG/Public/Components/ScrollBar.h
--- a/Engine/Source/Runtime/UMG/Public/Components/ScrollBar.h
+++ b/Engine/Source/Runtime/UMG/Public/Components/ScrollBar.h
@@ -22,9 +22,6 @@ public:
 	UPROPERTY(EditAnywhere, BlueprintReadWrite, Category="Style", meta=( DisplayName="Style" ))
 	FScrollBarStyle WidgetStyle;
 
-	UPROPERTY()
-	TObjectPtr<USlateWidgetStyleAsset> Style_DEPRECATED;
-
 	/**  */
 	UPROPERTY(EditAnywhere, Category="Behavior")
 	bool bAlwaysShowScrollbar;
@@ -84,7 +81,6 @@ public:
 	//~ Begin UObject Interface
 #if WITH_EDITORONLY_DATA
 	virtual void Serialize(FArchive& Ar) override;
-	virtual void PostLoad() override;
 #endif // if WITH_EDITORONLY_DATA
 	//~ End UObject Interface
 
diff --git a/Engine/Source/Runtime/UMG/Public/Components/ScrollBox.h b/Engine/Source/Runtime/UMG/Public/Components/ScrollBox.h
--- a/Engine/Source/Runtime/UMG/Public/Components/ScrollBox.h
+++ b/Engine/Source/Runtime/UMG/Public/Components/ScrollBox.h
@@ -31,12 +31,6 @@ public:
 	UPROPERTY(EditAnywhere, BlueprintReadWrite, Category="Style", meta=( DisplayName="Bar Style" ))
 	FScrollBarStyle WidgetBarStyle;
 
-	UPROPERTY()
-	TObjectPtr<USlateWidgetStyleAsset> Style_DEPRECATED;
-
-	UPROPERTY()
-	TObjectPtr<USlateWidgetStyleAsset> BarStyle_DEPRECATED;
-
 	/** The orientation of the scrolling and stacking in the box. */
 	UPROPERTY(EditAnywhere, BlueprintReadOnly, Category="Scroll")
 	TEnumAsByte<EOrientation> Orientation;
@@ -191,7 +185,6 @@ public:
 	//~ Begin UObject Interface
 #if WITH_EDITORONLY_DATA
 	virtual void Serialize(FArchive& Ar) override;
-	virtual void PostLoad() override;
 #endif // if WITH_EDITORONLY_DATA
 	//~ End UObject Interface
 
diff --git a/Engine/Source/Runtime/UMG/Public/Components/SpinBox.h b/Engine/Source/Runtime/UMG/Public/Components/SpinBox.h
--- a/Engine/Source/Runtime/UMG/Public/Components/SpinBox.h
+++ b/Engine/Source/Runtime/UMG/Public/Components/SpinBox.h
@@ -40,9 +40,6 @@ public:
 	UPROPERTY(EditAnywhere, BlueprintReadWrite, Category="Style", meta=( DisplayName="Style" ))
 	FSpinBoxStyle WidgetStyle;
 
-	UPROPERTY()
-	TObjectPtr<USlateWidgetStyleAsset> Style_DEPRECATED;
-
 	/** The minimum required fractional digits - default 1 */
 	UPROPERTY(EditAnywhere, BlueprintReadWrite, BlueprintSetter = SetMinFractionalDigits, BlueprintGetter = GetMinFractionalDigits, Category = "Slider", meta = (ClampMin = 0, UIMin = 0))
 	int32 MinFractionalDigits;
@@ -209,10 +206,6 @@ public:
 	virtual void ReleaseSlateResources(bool bReleaseChildren) override;
 	//~ End UVisual Interface
 
-	//~ Begin UObject Interface
-	virtual void PostLoad() override;
-	//~ End UObject Interface
-
 #if WITH_EDITOR
 	virtual const FText GetPaletteCategory() override;
 #endif
diff --git a/Engine/Source/Runtime/UMG/Public/Components/TextBlock.h b/Engine/Source/Runtime/UMG/Public/Components/TextBlock.h
--- a/Engine/Source/Runtime/UMG/Public/Components/TextBlock.h
+++ b/Engine/Source/Runtime/UMG/Public/Components/TextBlock.h
@@ -158,10 +158,6 @@ public:
 	UPROPERTY(EditAnywhere, BlueprintReadOnly, Category = Performance, AdvancedDisplay)
 	bool bWrapWithInvalidationPanel;
 
-	/** Whether the text should automatically wrap */
-	UPROPERTY(meta = (DeprecatedProperty, DeprecationMessage = "bAutoWrapText is deprecated. Please use AutoWrapText instead."))
-	bool bAutoWrapText_DEPRECATED;
-
 	/** The text transformation policy to apply to this text block */
 	UPROPERTY(EditAnywhere, BlueprintReadOnly, Category=Appearance, meta=(DisplayName="Transform Policy"))
 	ETextTransformPolicy TextTransformPolicy;
@@ -218,10 +214,6 @@ protected:
 	UPROPERTY(EditAnywhere, BlueprintReadOnly, Category=Performance, AdvancedDisplay, meta=(AllowPrivateAccess, DesignerRebuild))
 	bool bSimpleTextMode;
 
-	//~ Begin UObject Interface
-	virtual void PostLoad() override;
-	//~ End UObject Interface
-
 	//~ Begin UWidget Interface
 	virtual TSharedRef<SWidget> RebuildWidget() override;
 	virtual void OnBindingChanged(const FName& Property) override;
diff --git a/Engine/Source/Runtime/UMG/Public/Components/Throbber.h b/Engine/Source/Runtime/UMG/Public/Components/Throbber.h
--- a/Engine/Source/Runtime/UMG/Public/Components/Throbber.h
+++ b/Engine/Source/Runtime/UMG/Public/Components/Throbber.h
@@ -38,10 +38,6 @@ public:
 	UPROPERTY(EditAnywhere, BlueprintReadOnly, Category=Appearance)
 	bool bAnimateOpacity;
 
-	/** Image to use for each segment of the throbber */
-	UPROPERTY()
-	TObjectPtr<USlateBrushAsset> PieceImage_DEPRECATED;
-
 	UPROPERTY(EditAnywhere, BlueprintReadOnly, Category=Appearance)
 	FSlateBrush Image;
 
@@ -71,10 +67,6 @@ public:
 	virtual void ReleaseSlateResources(bool bReleaseChildren) override;
 	//~ End UVisual Interface
 
-	//~ Begin UObject Interface
-	virtual void PostLoad() override;
-	//~ End UObject Interface
-
 #if WITH_EDITOR
 	virtual const FText GetPaletteCategory() override;
 #endif
diff --git a/Engine/Source/Runtime/UMG/Public/Components/Widget.h b/Engine/Source/Runtime/UMG/Public/Components/Widget.h
--- a/Engine/Source/Runtime/UMG/Public/Components/Widget.h
+++ b/Engine/Source/Runtime/UMG/Public/Components/Widget.h
@@ -637,16 +637,6 @@ public:
 	UFUNCTION(BlueprintCallable, Category = "Widget")
 	void SetAllNavigationRules(EUINavigationRule Rule, FName WidgetToFocus);
 
-	/**
-	 *	Sets the widget navigation rules for a specific direction. This can only be called on widgets that are in a widget tree.
-	 *	@param Direction
-	 *	@param Rule The rule to use when navigation is taking place
-	 *	@param WidgetToFocus When using the Explicit rule, focus on this widget
-	 */
-	UE_DEPRECATED(4.23, "SetNavigationRule is deprecated. Please use either SetNavigationRuleBase or SetNavigationRuleExplicit or SetNavigationRuleCustom or SetNavigationRuleCustomBoundary.")
-	UFUNCTION(BlueprintCallable, Category = "Widget")
-	void SetNavigationRule(EUINavigation Direction, EUINavigationRule Rule, FName WidgetToFocus);
-
 	/**
 	 *	Sets the widget navigation rules for a specific direction. This can only be called on widgets that are in a widget tree. This works only for non Explicit, non Custom and non CustomBoundary Rules.
 	 *	@param Direction
@@ -1000,8 +990,6 @@ protected:
 
 	/** Function called after the underlying SWidget is constructed. */
 	virtual void OnWidgetRebuilt();
-	
-#if WITH_EDITOR
 	virtual TSharedRef<SWidget> RebuildDesignWidget(TSharedRef<SWidget> Content);
 
 	TSharedRef<SWidget> CreateDesignerOutline(TSharedRef<SWidget> Content) const;
diff --git a/Engine/Source/Runtime/UMG/Public/Components/WrapBox.h b/Engine/Source/Runtime/UMG/Public/Components/WrapBox.h
--- a/Engine/Source/Runtime/UMG/Public/Components/WrapBox.h
+++ b/Engine/Source/Runtime/UMG/Public/Components/WrapBox.h
@@ -28,22 +28,10 @@ public:
 	UPROPERTY(EditAnywhere, BlueprintReadOnly, Category="Content Layout")
 	FVector2D InnerSlotPadding;
 
-#if WITH_EDITORONLY_DATA
-	/** DEPRECATED value replaced by WrapSize, When this width is exceeded, elements will start appearing on the next line. */
-	UPROPERTY()
-	float WrapWidth_DEPRECATED;
-#endif
-
 	/** When this size is exceeded, elements will start appearing on the next line. */
 	UPROPERTY(EditAnywhere, BlueprintReadOnly, Category = "Content Layout", meta=(EditCondition = "bExplicitWrapSize"))
 	float WrapSize;
 
-#if WITH_EDITORONLY_DATA
-	/** DEPRECATED value replaced by bExplicitWrapSize, Use explicit wrap width whenever possible. It greatly simplifies layout calculations and reduces likelihood of "wiggling UI" */
-	UPROPERTY()
-	bool bExplicitWrapWidth_DEPRECATED;
-#endif
-
 	/** Use explicit wrap size whenever possible. It greatly simplifies layout calculations and reduces likelihood of "wiggling UI" */
 	UPROPERTY(EditAnywhere, BlueprintReadOnly, Category = "Content Layout")
 	bool bExplicitWrapSize;
@@ -64,10 +52,6 @@ public:
 	void SetHorizontalAlignment(EHorizontalAlignment InHorizontalAlignment);
 
 public:
-
-	UE_DEPRECATED(4.22, "Deprecated, please use AddChildToWrapBox() instead")
-	UWrapBoxSlot* AddChildWrapBox(UWidget* Content);
-
 	UFUNCTION(BlueprintCallable, Category="Panel")
 	UWrapBoxSlot* AddChildToWrapBox(UWidget* Content);
 
@@ -77,8 +61,6 @@ public:
 	// End UWidget interface
 #endif
 
-	virtual void PostLoad() override;
-
 protected:
 
 	// UPanelWidget
