Fix mad auto merge
#rnx

diff --git a/Engine/Source/Runtime/UMG/Public/Components/Widget.h b/Engine/Source/Runtime/UMG/Public/Components/Widget.h
--- a/Engine/Source/Runtime/UMG/Public/Components/Widget.h
+++ b/Engine/Source/Runtime/UMG/Public/Components/Widget.h
@@ -990,6 +990,8 @@ protected:
 
 	/** Function called after the underlying SWidget is constructed. */
 	virtual void OnWidgetRebuilt();
+	
+#if WITH_EDITOR
 	virtual TSharedRef<SWidget> RebuildDesignWidget(TSharedRef<SWidget> Content);
 
 	TSharedRef<SWidget> CreateDesignerOutline(TSharedRef<SWidget> Content) const;
