UMG: Bring back deprecated function still in use in production.

diff --git a/Engine/Source/Runtime/UMG/Private/Components/Widget.cpp b/Engine/Source/Runtime/UMG/Private/Components/Widget.cpp
--- a/Engine/Source/Runtime/UMG/Private/Components/Widget.cpp
+++ b/Engine/Source/Runtime/UMG/Private/Components/Widget.cpp
@@ -686,6 +686,12 @@ void UWidget::SetNavigationRuleInternal(EUINavigation Direction, EUINavigationRu
 	}
 }
 
+void UWidget::SetNavigationRule(EUINavigation Direction, EUINavigationRule Rule, FName WidgetToFocus)
+{
+	SetNavigationRuleInternal(Direction, Rule, WidgetToFocus);
+	BuildNavigation();
+}
+
 void UWidget::SetNavigationRuleBase(EUINavigation Direction, EUINavigationRule Rule)
 {
 	if (Rule == EUINavigationRule::Explicit || Rule == EUINavigationRule::Custom || Rule == EUINavigationRule::CustomBoundary)
diff --git a/Engine/Source/Runtime/UMG/Public/Components/Widget.h b/Engine/Source/Runtime/UMG/Public/Components/Widget.h
--- a/Engine/Source/Runtime/UMG/Public/Components/Widget.h
+++ b/Engine/Source/Runtime/UMG/Public/Components/Widget.h
@@ -637,6 +637,16 @@ public:
 	UFUNCTION(BlueprintCallable, Category = "Widget")
 	void SetAllNavigationRules(EUINavigationRule Rule, FName WidgetToFocus);
 
+	/**
+	 *	Sets the widget navigation rules for a specific direction. This can only be called on widgets that are in a widget tree.
+	 *	@param Direction
+	 *	@param Rule The rule to use when navigation is taking place
+	 *	@param WidgetToFocus When using the Explicit rule, focus on this widget
+	 */
+	UE_DEPRECATED(4.23, "SetNavigationRule is deprecated. Please use either SetNavigationRuleBase or SetNavigationRuleExplicit or SetNavigationRuleCustom or SetNavigationRuleCustomBoundary.")
+	UFUNCTION(BlueprintCallable, Category = "Widget")
+	void SetNavigationRule(EUINavigation Direction, EUINavigationRule Rule, FName WidgetToFocus);
+
 	/**
 	 *	Sets the widget navigation rules for a specific direction. This can only be called on widgets that are in a widget tree. This works only for non Explicit, non Custom and non CustomBoundary Rules.
 	 *	@param Direction
