Added OneWayToSource icon to UMGStyle.

diff --git a/Engine/Source/Runtime/UMG/Private/UMGStyle.cpp b/Engine/Source/Runtime/UMG/Private/UMGStyle.cpp
--- a/Engine/Source/Runtime/UMG/Private/UMGStyle.cpp
+++ b/Engine/Source/Runtime/UMG/Private/UMGStyle.cpp
@@ -161,6 +161,7 @@ TSharedRef< FSlateStyleSet > FUMGStyle::Create()
 	Style->Set("ClassThumbnail.SinglePropertyView", new IMAGE_BRUSH_SVG(TEXT("SinglePropertyView_64"), Icon64x64));
 
 	Style->Set("MVVM.OneWay", new IMAGE_BRUSH_SVG("ViewModel_OneWay", Icon16x16));
+	Style->Set("MVVM.OneWayToSource", new IMAGE_BRUSH_SVG("ViewModel_OneWayToSource", Icon16x16));
 	Style->Set("MVVM.TwoWay", new IMAGE_BRUSH_SVG("ViewModel_TwoWay", Icon16x16));
 	Style->Set("MVVM.OneTime", new IMAGE_BRUSH_SVG("ViewModel_OneTime", Icon16x16));
 
