URichTextBlock:
- Added GetTextStyleSet function.

diff --git a/Engine/Source/Runtime/UMG/Private/Components/RichTextBlock.cpp b/Engine/Source/Runtime/UMG/Private/Components/RichTextBlock.cpp
--- a/Engine/Source/Runtime/UMG/Private/Components/RichTextBlock.cpp
+++ b/Engine/Source/Runtime/UMG/Private/Components/RichTextBlock.cpp
@@ -159,7 +159,12 @@ void URichTextBlock::RebuildStyleInstance()
 	}
 }
 
-void URichTextBlock::SetTextStyleSet(class UDataTable* NewTextStyleSet)
+UDataTable* URichTextBlock::GetTextStyleSet() const
+{
+	return TextStyleSet;
+}
+
+void URichTextBlock::SetTextStyleSet(UDataTable* NewTextStyleSet)
 {
 	if (TextStyleSet != NewTextStyleSet)
 	{
diff --git a/Engine/Source/Runtime/UMG/Public/Components/RichTextBlock.h b/Engine/Source/Runtime/UMG/Public/Components/RichTextBlock.h
--- a/Engine/Source/Runtime/UMG/Public/Components/RichTextBlock.h
+++ b/Engine/Source/Runtime/UMG/Public/Components/RichTextBlock.h
@@ -12,6 +12,7 @@
 #include "RichTextBlock.generated.h"
 
 class SRichTextBlock;
+class UDataTable;
 class UMaterialInstanceDynamic;
 class URichTextBlockDecorator;
 
@@ -159,7 +160,10 @@ public:
 	virtual void SetText(const FText& InText);
 
 	UFUNCTION(BlueprintCallable, Category = "Widget")
-	void SetTextStyleSet(class UDataTable* NewTextStyleSet);
+	UDataTable* GetTextStyleSet() const;
+
+	UFUNCTION(BlueprintCallable, Category = "Widget")
+	void SetTextStyleSet(UDataTable* NewTextStyleSet);
 
 	const FTextBlockStyle& GetDefaultTextStyle() const;
 	const FTextBlockStyle& GetCurrentDefaultTextStyle() const;
@@ -193,7 +197,7 @@ protected:
 
 	/**  */
 	UPROPERTY(EditAnywhere, Category=Appearance, meta=(RequiredAssetDataTags = "RowStructure=RichTextStyleRow"))
-	TObjectPtr<class UDataTable> TextStyleSet;
+	TObjectPtr<UDataTable> TextStyleSet;
 
 	/**  */
 	UPROPERTY(EditAnywhere, Category=Appearance)
