UMG: Add UserWidget Extension.

diff --git a/Engine/Source/Runtime/UMG/Private/UserWidget.cpp b/Engine/Source/Runtime/UMG/Private/UserWidget.cpp
--- a/Engine/Source/Runtime/UMG/Private/UserWidget.cpp
+++ b/Engine/Source/Runtime/UMG/Private/UserWidget.cpp
@@ -14,6 +14,7 @@
 #include "Blueprint/WidgetBlueprintGeneratedClass.h"
 #include "Animation/UMGSequencePlayer.h"
 #include "Animation/UMGSequenceTickManager.h"
+#include "Extensions/UserWidgetExtension.h"
 #include "Extensions/WidgetBlueprintGeneratedClassExtension.h"
 #include "UObject/UnrealType.h"
 #include "Blueprint/WidgetNavigation.h"
@@ -153,6 +154,11 @@ bool UUserWidget::Initialize()
 			InitializeNamedSlots(bReparentToWidgetTree);
 		}
 
+		for (UUserWidgetExtension* Extension : Extensions)
+		{
+			Extension->Initialize();
+		}
+
 		if (!IsDesignTime() && PlayerContext.IsValid())
 		{
 			NativeOnInitialized();
@@ -228,6 +234,11 @@ void UUserWidget::BeginDestroy()
 {
 	Super::BeginDestroy();
 
+	for (UUserWidgetExtension* Extension : Extensions)
+	{
+		Extension->Destruct();
+	}
+
 	TearDownAnimations();
 
 	if (AnimationTickManager)
@@ -488,7 +499,7 @@ UUMGSequencePlayer* UUserWidget::PlayAnimation(UWidgetAnimation* InAnimation, fl
 		OnAnimationStartedPlaying(*Player);
 
 		UpdateCanTick();
-	}
+}
 
 	return Player;
 }
@@ -654,7 +665,7 @@ void UUserWidget::SetPlaybackSpeed(const UWidgetAnimation* InAnimation, float Pl
 void UUserWidget::ReverseAnimation(const UWidgetAnimation* InAnimation)
 {
 	if (UUMGSequencePlayer* FoundPlayer = GetSequencePlayer(InAnimation))
-		{
+	{
 		FoundPlayer->Reverse();
 	}
 }
@@ -1410,6 +1421,11 @@ void UUserWidget::NativeTick(const FGeometry& MyGeometry, float InDeltaTime)
 	{
 		GInitRunaway();
 
+		for (UUserWidgetExtension* Extension : Extensions)
+		{
+			Extension->Tick(MyGeometry, InDeltaTime);
+		}
+
 #if WITH_EDITOR
 		const bool bTickAnimations = !IsDesignTime();
 #else
@@ -1641,6 +1657,18 @@ void UUserWidget::UpdateCanTick()
 			bCanTick |= bHasScriptImplementedTick;
 			bCanTick |= World->GetLatentActionManager().GetNumActionsForObject(this) != 0;
 			bCanTick |= ActiveSequencePlayers.Num() > 0;
+
+			if (!bCanTick)
+			{
+				for(UUserWidgetExtension* Extension : Extensions)
+				{
+					if (Extension->RequiresTick())
+					{
+						bCanTick = true;
+						break;
+					}
+				}
+			}
 		}
 
 		SafeGCWidget->SetCanTick(bCanTick);
@@ -2066,6 +2094,43 @@ void UUserWidget::OnLatentActionsChanged(UObject* ObjectWhichChanged, ELatentAct
 	}
 }
 
+UUserWidgetExtension* UUserWidget::GetExtension(TSubclassOf<UUserWidgetExtension> InExtensionType) const
+{
+	for (UUserWidgetExtension* Extension : Extensions)
+	{
+		if (Extension->IsA(InExtensionType))
+		{
+			return Extension;
+		}
+	}
+	return nullptr;
+}
+
+TArray<UUserWidgetExtension*> UUserWidget::GetExtensions(TSubclassOf<UUserWidgetExtension> InExtensionType) const
+{
+	TArray<UUserWidgetExtension*> Result;
+	for (UUserWidgetExtension* Extension : Extensions)
+	{
+		if (Extension->IsA(InExtensionType))
+		{
+			Result.Add(Extension);
+		}
+	}
+	return Result;
+}
+
+UUserWidgetExtension* UUserWidget::AddExtension(TSubclassOf<UUserWidgetExtension> InExtensionType)
+{
+	UUserWidgetExtension* Extension = NewObject<UUserWidgetExtension>(this, InExtensionType);
+	Extensions.Add(Extension);
+	return Extension;
+}
+
+void UUserWidget::RemoveExtension(UUserWidgetExtension* InExtension)
+{
+	Extensions.RemoveSingle(InExtension);
+}
+
 
 /////////////////////////////////////////////////////
 
diff --git a/Engine/Source/Runtime/UMG/Public/Blueprint/UserWidget.h b/Engine/Source/Runtime/UMG/Public/Blueprint/UserWidget.h
--- a/Engine/Source/Runtime/UMG/Public/Blueprint/UserWidget.h
+++ b/Engine/Source/Runtime/UMG/Public/Blueprint/UserWidget.h
@@ -37,6 +37,7 @@ class UUMGSequenceTickManager;
 class UWidgetAnimation;
 class UWidgetTree;
 class UNamedSlot;
+class UUserWidgetExtension;
 
 /** Determines what strategy we use to determine when and if the widget ticks. */
 UENUM()
@@ -226,6 +227,8 @@ public:
 
 	UWidgetBlueprintGeneratedClass* GetWidgetTreeOwningClass() const;
 
+	void UpdateCanTick();
+
 protected:
 	/** The function is implemented only in nativized widgets (automatically converted from BP to c++) */
 	virtual void InitializeNativeClassData() {}
@@ -1057,6 +1060,36 @@ public:
 	UFUNCTION(BlueprintCallable, BlueprintCosmetic, Category = "User Interface|Animation")
 	void FlushAnimations();
 
+	/** Find the first extension of the requested type. */
+	template<typename ExtensionType>
+	ExtensionType* GetExtension() const
+	{
+		return CastChecked<ExtensionType>(GetExtension(ExtensionType::StaticClass()), ECastCheckedType::NullAllowed);
+	}
+
+	/** Find the first extension of the requested type. */
+	UFUNCTION(BlueprintCallable, Category = "User Interface|Extension")
+	UUserWidgetExtension* GetExtension(TSubclassOf<UUserWidgetExtension> InExtensionType) const;
+
+	/** Find the extensions of the requested type. */
+	UFUNCTION(BlueprintCallable, Category = "User Interface|Extension")
+	TArray<UUserWidgetExtension*> GetExtensions(TSubclassOf<UUserWidgetExtension> InExtensionType) const;
+
+	/** Add the extension of the requested type. */
+	template<typename ExtensionType>
+	ExtensionType* AddExtension()
+	{
+		return CastChecked<ExtensionType>(AddExtension(ExtensionType::StaticClass()), ECastCheckedType::NullAllowed);
+	}
+
+	/** Add the extension of the requested type. */
+	UFUNCTION(BlueprintCallable, Category = "User Interface|Extension")
+	UUserWidgetExtension* AddExtension(TSubclassOf<UUserWidgetExtension> InExtensionType);
+
+	/** Remove the extension. */
+	UFUNCTION(BlueprintCallable, Category = "User Interface|Extension")
+	void RemoveExtension(UUserWidgetExtension* InExtension);
+
 	/**
 	 * Plays a sound through the UI
 	 *
@@ -1172,6 +1205,10 @@ private:
 	UPROPERTY()
 	TArray<FNamedSlotBinding> NamedSlotBindings;
 
+	/** The UserWidget extensions */
+	UPROPERTY()
+	TArray<TObjectPtr<UUserWidgetExtension>> Extensions;
+
 public:
 	/** The widget tree contained inside this user widget initialized by the blueprint */
 	UPROPERTY(Transient, DuplicateTransient, TextExportTransient)
@@ -1409,8 +1446,6 @@ private:
 	static bool bTemplateInitializing;
 	static uint32 bInitializingFromWidgetTree;
 
-	void UpdateCanTick();
-
 protected:
 
 	PROPERTY_BINDING_IMPLEMENTATION(FLinearColor, ColorAndOpacity);
diff --git a/Engine/Source/Runtime/UMG/Public/Extensions/UserWidgetExtension.h b/Engine/Source/Runtime/UMG/Public/Extensions/UserWidgetExtension.h
new file mode 100644
--- /dev/null
+++ b/Engine/Source/Runtime/UMG/Public/Extensions/UserWidgetExtension.h
@@ -0,0 +1,50 @@
+// Copyright Epic Games, Inc. All Rights Reserved.
+
+#pragma once
+
+#include "CoreMinimal.h"
+
+#include "UserWidgetExtension.generated.h"
+
+
+class UUserWidget;
+
+/**
+ * UserWidgetExtension is the base class for components that define reusable behavior that can be added to different types of Widgets.
+ */
+UCLASS(Abstract, DefaultToInstanced, Within=UserWidget)
+class UMG_API UUserWidgetExtension : public UObject
+{
+	GENERATED_BODY()
+	
+public:
+	/** Extend the UUserWidget::Initialize function */
+	virtual void Initialize()
+	{
+	}
+
+	/** Extend the UUserWidget::Destruct function */
+	virtual void Destruct()
+	{
+	}
+
+	/** Does the extension requires tick. */
+	virtual bool RequiresTick() const
+	{
+		return false;
+	}
+
+	/**
+	 * Extend the UUserWidget::Tick function.
+	 * If the UserWidget ticks, then all extensions will tick regardless of RequiresTick.
+	 */
+	virtual void Tick(const FGeometry& MyGeometry, float InDeltaTime)
+	{
+	}
+
+protected:
+	UUserWidget* GetUserWidget() const
+	{
+		return GetOuterUUserWidget();
+	}
+};
