Fix compilation unit.
#rnx

diff --git a/Engine/Source/Runtime/UMG/Public/Extensions/UserWidgetExtension.h b/Engine/Source/Runtime/UMG/Public/Extensions/UserWidgetExtension.h
--- a/Engine/Source/Runtime/UMG/Public/Extensions/UserWidgetExtension.h
+++ b/Engine/Source/Runtime/UMG/Public/Extensions/UserWidgetExtension.h
@@ -3,6 +3,9 @@
 #pragma once
 
 #include "CoreMinimal.h"
+#include "Layout/Geometry.h"
+#include "UObject/ObjectMacros.h"
+#include "UObject/Object.h"
 
 #include "UserWidgetExtension.generated.h"
 
