UMG: New feature. Add the interface IFieldValueChanged. It returns the list of all the properties or function that will send a notification when the value changes.

diff --git a/Engine/Source/Runtime/UMG/Private/Binding/WidgetFieldNotificationExtension.h b/Engine/Source/Runtime/UMG/Private/Binding/WidgetFieldNotificationExtension.h
new file mode 100644
--- /dev/null
+++ b/Engine/Source/Runtime/UMG/Private/Binding/WidgetFieldNotificationExtension.h
@@ -0,0 +1,48 @@
+// Copyright Epic Games, Inc. All Rights Reserved.
+
+#pragma once
+
+#include "CoreMinimal.h"
+#include "Components/Visual.h"
+#include "Extensions/UserWidgetExtension.h"
+#include "FieldNotification/FieldMulticastDelegate.h"
+
+#include "WidgetFieldNotificationExtension.generated.h"
+
+class UVisual;
+
+UCLASS(MinimalAPI, Transient)
+class UWidgetFieldNotificationExtension : public UUserWidgetExtension
+{
+	GENERATED_BODY()
+
+public:
+	using FDelegate = UE::FieldNotification::FFieldMulticastDelegate::FDelegate;
+	using FFieldId = UE::FieldNotification::FFieldId;
+	using FRemoveFromResult = UE::FieldNotification::FFieldMulticastDelegate::FRemoveFromResult;
+	using FRemoveAllResult = UE::FieldNotification::FFieldMulticastDelegate::FRemoveAllResult;
+
+public:
+	FDelegateHandle AddFieldValueChangedDelegate(const UVisual* InObject, FFieldId InFieldId, FDelegate InNewDelegate)
+	{
+		return Delegates.Add(InObject, InFieldId, MoveTemp(InNewDelegate));
+	}
+
+	FRemoveFromResult RemoveFieldValueChangedDelegate(const UVisual* InObject, FFieldId InFieldId, FDelegateHandle InDelegateHandle)
+	{
+		return Delegates.RemoveFrom(InObject, InFieldId, InDelegateHandle);
+	}
+
+	FRemoveAllResult RemoveAllFieldValueChangedDelegates(const UVisual* InObject, const void* InUserObject)
+	{
+		return Delegates.RemoveAll(InObject, InUserObject);
+	}
+
+	void BroadcastFieldValueChanged(UVisual* InObject, UE::FieldNotification::FFieldId InFieldId)
+	{
+		Delegates.Broadcast(InObject, InFieldId);
+	}
+
+private:
+	UE::FieldNotification::FFieldMulticastDelegate Delegates;
+};
diff --git a/Engine/Source/Runtime/UMG/Private/Components/Widget.cpp b/Engine/Source/Runtime/UMG/Private/Components/Widget.cpp
--- a/Engine/Source/Runtime/UMG/Private/Components/Widget.cpp
+++ b/Engine/Source/Runtime/UMG/Private/Components/Widget.cpp
@@ -19,6 +19,7 @@
 #include "Widgets/Layout/SSpacer.h"
 #include "Widgets/SToolTip.h"
 #include "Binding/PropertyBinding.h"
+#include "Binding/WidgetFieldNotificationExtension.h"
 #include "Logging/MessageLog.h"
 #include "Blueprint/UserWidget.h"
 #include "Blueprint/UserWidgetBlueprint.h"
@@ -152,6 +153,11 @@ namespace
 
 TArray<TSubclassOf<UPropertyBinding>> UWidget::BinderClasses;
 
+
+/////////////////////////////////////////////////////
+UE_FIELD_NOTIFICATION_IMPLEMENT_CLASS_DESCRIPTOR_ThreeFields(UWidget, ToolTipText, Visibility, bIsEnabled);
+
+
 UWidget::UWidget(const FObjectInitializer& ObjectInitializer)
 	: Super(ObjectInitializer)
 {
@@ -250,7 +256,11 @@ bool UWidget::GetIsEnabled() const
 
 void UWidget::SetIsEnabled(bool bInIsEnabled)
 {
-	bIsEnabled = bInIsEnabled;
+	if (bIsEnabled != bInIsEnabled)
+	{
+		bIsEnabled = bInIsEnabled;
+		BroadcastFieldValueChanged(FFieldNotificationClassDescriptor::bIsEnabled);
+	}
 
 	TSharedPtr<SWidget> SafeWidget = GetCachedWidget();
 	if (SafeWidget.IsValid())
@@ -317,7 +327,11 @@ ESlateVisibility UWidget::GetVisibility() const
 
 void UWidget::SetVisibility(ESlateVisibility InVisibility)
 {
-	Visibility = InVisibility;
+	if (Visibility != InVisibility)
+	{
+		Visibility = InVisibility;
+		BroadcastFieldValueChanged(FFieldNotificationClassDescriptor::Visibility);
+	}
 
 	TSharedPtr<SWidget> SafeWidget = GetCachedWidget();
 	if (SafeWidget.IsValid())
@@ -384,6 +398,7 @@ void UWidget::ForceVolatile(bool bForce)
 void UWidget::SetToolTipText(const FText& InToolTipText)
 {
 	ToolTipText = InToolTipText;
+	BroadcastFieldValueChanged(FFieldNotificationClassDescriptor::ToolTipText);
 
 	TSharedPtr<SWidget> SafeWidget = GetCachedWidget();
 	if (SafeWidget.IsValid())
@@ -1629,4 +1644,89 @@ void UWidget::OnBindingChanged(const FName& Property)
 
 }
 
+
+namespace UE::UMG::Private
+{
+	UWidgetFieldNotificationExtension* FindOrAddWidgetNotifyExtension(UWidget* Widget)
+	{
+		if (UUserWidget* UserWidget = Cast<UUserWidget>(Widget))
+		{
+			if (UWidgetFieldNotificationExtension* Extension = UserWidget->GetExtension<UWidgetFieldNotificationExtension>())
+			{
+				return Extension;
+			}
+			return UserWidget->AddExtension<UWidgetFieldNotificationExtension>();
+		}
+		else if (UWidgetTree* WidgetTree = Cast<UWidgetTree>(Widget->GetOuter()))
+		{
+			if (UUserWidget* InnerUserWidget = Cast<UUserWidget>(WidgetTree->GetOuter()))
+			{
+				if (UWidgetFieldNotificationExtension* Extension = InnerUserWidget->GetExtension<UWidgetFieldNotificationExtension>())
+				{
+					return Extension;
+				}
+				return InnerUserWidget->AddExtension<UWidgetFieldNotificationExtension>();
+			}
+		}
+		return nullptr;
+	}
+}
+
+FDelegateHandle UWidget::AddFieldValueChangedDelegate(UE::FieldNotification::FFieldId InFieldId, FFieldValueChangedDelegate InNewDelegate)
+{
+	FDelegateHandle Result;
+	if (InFieldId.IsValid())
+	{
+		if (UWidgetFieldNotificationExtension* Extension = UE::UMG::Private::FindOrAddWidgetNotifyExtension(this))
+		{
+			Result = Extension->AddFieldValueChangedDelegate(this, InFieldId, MoveTemp(InNewDelegate));
+			if (Result.IsValid())
+			{
+				EnabledFieldNotifications.PadToNum(InFieldId.GetIndex() + 1, false);
+				EnabledFieldNotifications[InFieldId.GetIndex()] = true;
+			}
+		}
+	}
+	return Result;
+}
+
+bool UWidget::RemoveFieldValueChangedDelegate(UE::FieldNotification::FFieldId InFieldId, FDelegateHandle InHandle)
+{
+	bool bResult = false;
+	if (InFieldId.IsValid() && InHandle.IsValid() && EnabledFieldNotifications.IsValidIndex(InFieldId.GetIndex()) && EnabledFieldNotifications[InFieldId.GetIndex()])
+	{
+		UWidgetFieldNotificationExtension* Extension = UE::UMG::Private::FindOrAddWidgetNotifyExtension(this);
+		checkf(Extension, TEXT("If the EnabledFieldNotifications is valid, then the Extension must also be valid."));
+		UWidgetFieldNotificationExtension::FRemoveFromResult RemoveResult = Extension->RemoveFieldValueChangedDelegate(this, InFieldId, InHandle);
+		bResult = RemoveResult.bRemoved;
+		EnabledFieldNotifications[InFieldId.GetIndex()] = RemoveResult.bHasOtherBoundDelegates;
+	}
+	return bResult;
+}
+
+int32 UWidget::RemoveAllFieldValueChangedDelegates(const void* InUserObject)
+{
+	int32 bResult = 0;
+	if (InUserObject)
+	{
+		if (UWidgetFieldNotificationExtension* Extension = UE::UMG::Private::FindOrAddWidgetNotifyExtension(this))
+		{
+			UWidgetFieldNotificationExtension::FRemoveAllResult RemoveResult = Extension->RemoveAllFieldValueChangedDelegates(this, InUserObject);
+			bResult = RemoveResult.RemoveCount;
+			EnabledFieldNotifications = RemoveResult.HasFields;
+		}
+	}
+	return bResult;
+}
+
+void UWidget::BroadcastFieldValueChanged(UE::FieldNotification::FFieldId InFieldId)
+{
+	if (InFieldId.IsValid() && EnabledFieldNotifications.IsValidIndex(InFieldId.GetIndex()) && EnabledFieldNotifications[InFieldId.GetIndex()])
+	{
+		UWidgetFieldNotificationExtension* Extension = UE::UMG::Private::FindOrAddWidgetNotifyExtension(this);
+		checkf(Extension, TEXT("If the EnabledFieldNotifications is valid, then the Extension must also be valid."));
+		Extension->BroadcastFieldValueChanged(this, InFieldId);
+	}
+}
+
 #undef LOCTEXT_NAMESPACE
diff --git a/Engine/Source/Runtime/UMG/Private/FieldNotification/FieldMulticastDelegate.cpp b/Engine/Source/Runtime/UMG/Private/FieldNotification/FieldMulticastDelegate.cpp
new file mode 100644
--- /dev/null
+++ b/Engine/Source/Runtime/UMG/Private/FieldNotification/FieldMulticastDelegate.cpp
@@ -0,0 +1,384 @@
+// Copyright Epic Games, Inc. All Rights Reserved.
+
+#include "FieldNotification/FieldMulticastDelegate.h"
+
+
+namespace UE::FieldNotification
+{
+
+FDelegateHandle FFieldMulticastDelegate::Add(const UObject* InObject, FFieldId InFieldId, FDelegate InNewDelegate)
+{
+	FDelegateHandle Result = InNewDelegate.GetHandle();
+	FInvocationKey InvocationKey{ InObject, InFieldId };
+	if (DelegateLockCount > 0)
+	{
+		const int32 Index = Delegates.Emplace(FInvocationElement{ InvocationKey, MoveTemp(InNewDelegate) });
+		AddedEmplaceAt = FMath::Min(AddedEmplaceAt, (uint16)Index);
+	}
+	else
+	{
+		const int32 FoundIndex = UpperBound(InvocationKey);
+		Delegates.EmplaceAt(FoundIndex, FInvocationElement{ InvocationKey, MoveTemp(InNewDelegate) });
+	}
+	return Result;
+}
+
+
+FFieldMulticastDelegate::FRemoveResult FFieldMulticastDelegate::Remove(FDelegateHandle InDelegate)
+{
+	FRemoveResult Result;
+
+	for (int32 Index = Delegates.Num() - 1; Index >= 0; --Index)
+	{
+		FInvocationElement& Element = Delegates[Index];
+		if (Element.Delegate.GetHandle() == InDelegate)
+		{
+			Result.Object = Element.Key.Object.Get();
+			Result.FieldId = Element.Key.Id;
+			Result.bRemoved = Element.Delegate.IsBound();
+			if (DelegateLockCount > 0)
+			{
+				Element.Delegate.Unbind();
+				++CompactionCount;
+			}
+			else
+			{
+				Delegates.RemoveAt(Index);
+			}
+			break;
+		}
+	}
+
+	if (Result.FieldId.IsValid())
+	{
+		// Search in the normal sorted list
+		const int32 FoundIndex = UpperBound(Result.FieldId) - 1;
+		if (Delegates.IsValidIndex(FoundIndex))
+		{
+			for (int32 Index = FoundIndex; Index >= 0; --Index)
+			{
+				const FInvocationElement& Element = Delegates[Index];
+				if (Element.Key.Id != Result.FieldId)
+				{
+					break;
+				}
+				else if (Element.Key.Object == Result.Object && Element.Delegate.IsBound())
+				{
+					Result.bHasOtherBoundDelegates = true;
+					break;
+				}
+			}
+		}
+
+		if (!Result.bHasOtherBoundDelegates)
+		{
+			for (int32 Index = AddedEmplaceAt; Index < Delegates.Num(); ++Index)
+			{
+				const FInvocationElement& Element = Delegates[Index];
+				if (Element.Key.Id == Result.FieldId && Element.Key.Object == Result.Object && Element.Delegate.IsBound())
+				{
+					Result.bHasOtherBoundDelegates = true;
+					break;
+				}
+			}
+		}
+	}
+
+	return Result;
+}
+
+
+FFieldMulticastDelegate::FRemoveFromResult FFieldMulticastDelegate::RemoveFrom(const UObject* InObject, FFieldId InFieldId, FDelegateHandle InDelegate)
+{
+	bool bRemoved = false;
+	bool bFieldPresent = false;
+
+	FFieldMulticastDelegate* Self = this;
+	auto RemoveOrUnbind = [Self, &InDelegate, &bRemoved, &bFieldPresent, InObject](FInvocationElement& Element, int32 Index)
+	{
+		if (Element.Delegate.GetHandle() == InDelegate)
+		{
+			if (Self->DelegateLockCount > 0)
+			{
+				Element.Delegate.Unbind();
+				++Self->CompactionCount;
+			}
+			else
+			{
+				Self->Delegates.RemoveAt(Index);
+			}
+			bRemoved = true;
+			return !bFieldPresent;
+		}
+		else if (Element.Key.Object == InObject && Element.Delegate.IsBound())
+		{
+			bFieldPresent = true;
+			return !bRemoved;
+		}
+		return true;
+	};
+
+	// Search in the normal sorted list
+	const int32 FoundIndex = UpperBound(InFieldId)-1;
+	if (Delegates.IsValidIndex(FoundIndex))
+	{
+		for (int32 Index = FoundIndex; Index >= 0; --Index)
+		{
+			FInvocationElement& Element = Delegates[Index];
+			if (Element.Key.Id != InFieldId)
+			{
+				break;
+			}
+			if (!RemoveOrUnbind(Element, Index))
+			{
+				break;
+			}
+		}
+	}
+
+	// The item may be at the end of the array, if added while broadcasting
+	if (!bFieldPresent)
+	{
+		for (int32 Index = Delegates.Num() - 1; Index >= AddedEmplaceAt; --Index)
+		{
+			FInvocationElement& Element = Delegates[Index];
+			if (Element.Key.Id == InFieldId)
+			{
+				if (!RemoveOrUnbind(Element, Index))
+				{
+					break;
+				}
+			}
+		}
+	}
+
+	return FRemoveFromResult{ bRemoved, bFieldPresent };
+}
+
+
+FFieldMulticastDelegate::FRemoveAllResult FFieldMulticastDelegate::RemoveAll(const UObject* InObject, const void* InUserObject)
+{
+	FRemoveAllResult Result;
+
+	auto SetHasFields = [&Result, InObject](FInvocationElement& Element)
+	{
+		Result.HasFields.PadToNum(Element.Key.Id.GetIndex(), false);
+		Result.HasFields[Element.Key.Id.GetIndex()] = true;
+	};
+
+	if (DelegateLockCount > 0)
+	{
+		for (FInvocationElement& Element : Delegates)
+		{
+			if (const IDelegateInstance* DelegateInstance = GetDelegateInstance(Element.Delegate))
+			{
+				if (Element.Key.Object == InObject)
+				{
+					if (DelegateInstance->HasSameObject(InUserObject))
+					{
+						Element.Delegate.Unbind();
+						++CompactionCount;
+						++Result.RemoveCount;
+					}
+					else if (!DelegateInstance->IsCompactable())
+					{
+						SetHasFields(Element);
+					}
+				}
+			}
+		}
+	}
+	else
+	{
+		for (int32 Index = Delegates.Num() - 1; Index >= 0; --Index)
+		{
+			FInvocationElement& Element = Delegates[Index];
+			const IDelegateInstance* DelegateInstance = GetDelegateInstance(Element.Delegate);
+			if (DelegateInstance == nullptr || DelegateInstance->IsCompactable())
+			{
+				Delegates.RemoveAt(Index);
+			}
+			else if (Element.Key.Object == InObject)
+			{
+				if (DelegateInstance->HasSameObject(InUserObject))
+				{
+					Delegates.RemoveAt(Index);
+					++Result.RemoveCount;
+				}
+				else
+				{
+					SetHasFields(Element);
+				}
+			}
+		}
+		CompactionCount = 0;
+	}
+
+	return Result;
+}
+
+
+FFieldMulticastDelegate::FRemoveAllResult FFieldMulticastDelegate::RemoveAll(const UObject* InObject, FFieldId InFieldId, const void* InUserObject)
+{
+	FRemoveAllResult Result;
+
+	auto SetHasFields = [&Result, InObject](FInvocationElement& Element)
+	{
+		Result.HasFields.PadToNum(Element.Key.Id.GetIndex(), false);
+		Result.HasFields[Element.Key.Id.GetIndex()] = true;
+	};
+
+	if (DelegateLockCount > 0)
+	{
+		for (FInvocationElement& Element : Delegates)
+		{
+			if (const IDelegateInstance* DelegateInstance = GetDelegateInstance(Element.Delegate))
+			{
+				if (Element.Key.Id == InFieldId && Element.Key.Object == InObject)
+				{
+					if (DelegateInstance->HasSameObject(InUserObject))
+					{
+						Element.Delegate.Unbind();
+						++CompactionCount;
+						++Result.RemoveCount;
+					}
+					else if (!DelegateInstance->IsCompactable())
+					{
+						SetHasFields(Element);
+					}
+				}
+			}
+		}
+	}
+	else
+	{
+		for (int32 Index = Delegates.Num() - 1; Index >= 0; --Index)
+		{
+			FInvocationElement& Element = Delegates[Index];
+			const IDelegateInstance* DelegateInstance = GetDelegateInstance(Element.Delegate);
+			if (DelegateInstance == nullptr || DelegateInstance->IsCompactable())
+			{
+				Delegates.RemoveAt(Index);
+			}
+			else if (Element.Key.Id == InFieldId && Element.Key.Object == InObject)
+			{
+				if (DelegateInstance->HasSameObject(InUserObject))
+				{
+					Delegates.RemoveAt(Index);
+					++Result.RemoveCount;
+				}
+				else
+				{
+					SetHasFields(Element);
+				}
+			}
+		}
+		CompactionCount = 0;
+	}
+
+	return Result;
+}
+
+
+void FFieldMulticastDelegate::Broadcast(UObject* InObject, UE::FieldNotification::FFieldId InFieldId)
+{
+	++DelegateLockCount;
+
+	// Search in the normal sorted list
+	const int32 FoundIndex = UpperBound(InFieldId) - 1;
+	if (Delegates.IsValidIndex(FoundIndex))
+	{
+		for (int32 Index = FoundIndex; Index >= 0; --Index)
+		{
+			FInvocationElement& Element = Delegates[Index];
+			if (Element.Key.Id != InFieldId)
+			{
+				break;
+			}
+			if (Element.Key.Object == InObject)
+			{
+				Element.Delegate.ExecuteIfBound(InObject, InFieldId);
+			}
+		}
+	}
+
+	// The item may be at the end of the array, if added while broadcasting
+	for (int32 Index = Delegates.Num() - 1; Index >= AddedEmplaceAt; --Index)
+	{
+		FInvocationElement& Element = Delegates[Index];
+		if (Element.Key.Id == InFieldId && Element.Key.Object == InObject)
+		{
+			Element.Delegate.ExecuteIfBound(InObject, InFieldId);
+		}
+	}
+
+	--DelegateLockCount;
+	ExecuteLockOperations();
+}
+
+
+void FFieldMulticastDelegate::Reset()
+{
+	if (DelegateLockCount > 0)
+	{
+		for (FInvocationElement& Element : Delegates)
+		{
+			CompactionCount += Element.Delegate.IsBound() ? 1 : 0;
+			Element.Delegate.Unbind();
+		}
+	}
+	else
+	{
+		Delegates.Reset();
+		DelegateLockCount = 0;
+		CompactionCount = 0;
+		AddedEmplaceAt = std::numeric_limits<uint16>::max();
+	}
+}
+
+
+void FFieldMulticastDelegate::ExecuteLockOperations()
+{
+	if (DelegateLockCount <= 0)
+	{
+		// Remove items that got removed while broadcasting
+		if (CompactionCount > 2)
+		{
+			for (int32 Index = Delegates.Num() - 1; Index >= 0; --Index)
+			{
+				FInvocationElement& Element = Delegates[Index];
+				const IDelegateInstance* DelegateInstance = GetDelegateInstance(Element.Delegate);
+				if (DelegateInstance == nullptr || DelegateInstance->IsCompactable())
+				{
+					Delegates.RemoveAt(Index);
+				}
+			}
+			CompactionCount = 0;
+		}
+
+		// Sort item that were added while broadcasting
+		for (; AddedEmplaceAt < Delegates.Num();)
+		{
+			if (Delegates[AddedEmplaceAt].Delegate.IsBound())
+			{
+				const int32 FoundIndex = UpperBound(Delegates[AddedEmplaceAt].Key);
+				if (FoundIndex != AddedEmplaceAt)
+				{
+					FInvocationElement Element = MoveTemp(Delegates[AddedEmplaceAt]);
+					Delegates.RemoveAtSwap(AddedEmplaceAt);
+					Delegates.EmplaceAt(FoundIndex, MoveTemp(Element));
+				}
+				++AddedEmplaceAt;
+			}
+			else
+			{
+				Delegates.RemoveAtSwap(AddedEmplaceAt);
+			}
+		}
+		AddedEmplaceAt = std::numeric_limits<uint16>::max();
+
+		DelegateLockCount = 0;
+	}
+}
+
+} //namespace
\ No newline at end of file
diff --git a/Engine/Source/Runtime/UMG/Private/Tests/FieldNotificationMulticastDelegateTest.cpp b/Engine/Source/Runtime/UMG/Private/Tests/FieldNotificationMulticastDelegateTest.cpp
new file mode 100644
--- /dev/null
+++ b/Engine/Source/Runtime/UMG/Private/Tests/FieldNotificationMulticastDelegateTest.cpp
@@ -0,0 +1,155 @@
+// Copyright Epic Games, Inc. All Rights Reserved.
+
+#include "Misc/AutomationTest.h"
+#include "FieldNotification/FieldMulticastDelegate.h"
+
+
+#if WITH_AUTOMATION_WORKER
+
+#define LOCTEXT_NAMESPACE "FieldNotificationMulticastDelegateTest"
+
+IMPLEMENT_SIMPLE_AUTOMATION_TEST(FFieldNotificationMulticastDelegateTest, "System.UMG.FieldNotificationMulticastDelegateTest", EAutomationTestFlags::ApplicationContextMask | EAutomationTestFlags::EngineFilter);
+
+
+bool FFieldNotificationMulticastDelegateTest::RunTest(const FString& Parameters)
+{
+	struct FLocal
+	{
+		FLocal(FFieldNotificationMulticastDelegateTest* Value)
+			: Self(Value)
+		{}
+		FFieldNotificationMulticastDelegateTest* Self = nullptr;
+		UE::FieldNotification::FFieldMulticastDelegate Delegate;
+		int32 FooACount = 0;
+		int32 FooBCount = 0;
+		int32 FooCCount = 0;
+		UE::FieldNotification::FFieldId ExpectFieldA;
+		UE::FieldNotification::FFieldId ExpectFieldB;
+		UE::FieldNotification::FFieldId ExpectFieldC;
+
+		void Reset()
+		{
+			FooACount = 0;
+			FooBCount = 0;
+			FooCCount = 0;
+			ExpectFieldA = UE::FieldNotification::FFieldId();
+			ExpectFieldB = UE::FieldNotification::FFieldId();
+			ExpectFieldC = UE::FieldNotification::FFieldId();
+		}
+		void FooA(UObject* InObject, UE::FieldNotification::FFieldId Id)
+		{
+			++FooACount;
+			if (ExpectFieldA.IsValid() && Id != ExpectFieldA)
+			{
+				Self->AddError(TEXT("The expected field A is not valid."));
+			}
+		}
+		void FooB(UObject* InObject, UE::FieldNotification::FFieldId Id)
+		{
+			++FooBCount;
+			if (ExpectFieldB.IsValid() && Id != ExpectFieldB)
+			{
+				Self->AddError(TEXT("The expected field B is not valid."));
+			}
+		}
+		void FooC(UObject* InObject, UE::FieldNotification::FFieldId Id)
+		{
+			++FooCCount;
+			if (ExpectFieldC.IsValid() && Id != ExpectFieldC)
+			{
+				Self->AddError(TEXT("The expected field C is not valid."));
+			}
+		}
+	};
+
+	UE::FieldNotification::FFieldId FieldA{ "FieldA", 0 };
+	UE::FieldNotification::FFieldId FieldB{ "FieldB", 1 };
+	UE::FieldNotification::FFieldId FieldC{ "FieldC", 2 };
+	UE::FieldNotification::FFieldId FieldD{ "FieldD", 3 };
+
+	FLocal Local{ this };
+	auto RunTest1 = [&Local, &FieldA, &FieldB, &FieldC, this]()
+	{
+		{
+			Local.Reset();
+			FDelegateHandle Handle1 = Local.Delegate.Add(UObject::StaticClass(), FieldA, UE::FieldNotification::FFieldMulticastDelegate::FDelegate::CreateRaw(&Local, &FLocal::FooA));
+			Local.ExpectFieldA = FieldA;
+			Local.Delegate.Broadcast(UObject::StaticClass(), FieldA);
+			AddErrorIfFalse(Local.FooACount == 1, TEXT("Bad amount of FooA"));
+			AddErrorIfFalse(Local.FooBCount == 0, TEXT("Bad amount of FooB"));
+
+			Local.Reset();
+			FDelegateHandle Handle2 = Local.Delegate.Add(UObject::StaticClass(), FieldA, UE::FieldNotification::FFieldMulticastDelegate::FDelegate::CreateRaw(&Local, &FLocal::FooA));
+			Local.ExpectFieldA = FieldA;
+			Local.Delegate.Broadcast(UObject::StaticClass(), FieldA);
+			AddErrorIfFalse(Local.FooACount == 2, TEXT("Bad amount of FooA"));
+			AddErrorIfFalse(Local.FooBCount == 0, TEXT("Bad amount of FooB"));
+
+			UE::FieldNotification::FFieldMulticastDelegate::FRemoveFromResult Result1 = Local.Delegate.RemoveFrom(UObject::StaticClass(), FieldA, Handle1);
+			AddErrorIfFalse(Result1.bRemoved, TEXT("Bad Remove of Handle1"));
+			AddErrorIfFalse(Result1.bHasOtherBoundDelegates, TEXT("Remove of Handle1 should have other bound delegates."));
+			UE::FieldNotification::FFieldMulticastDelegate::FRemoveFromResult Result2 = Local.Delegate.RemoveFrom(UObject::StaticClass(), FieldA, Handle1);
+			AddErrorIfFalse(!Result2.bRemoved, TEXT("Bad Remove of Handle1"));
+			AddErrorIfFalse(Result2.bHasOtherBoundDelegates, TEXT("Remove of Handle1 should not have other bound delegates."));
+			UE::FieldNotification::FFieldMulticastDelegate::FRemoveFromResult Result3 = Local.Delegate.RemoveFrom(UObject::StaticClass(), FieldA, Handle2);
+			AddErrorIfFalse(Result3.bRemoved, TEXT("Bad Remove of Handle2"));
+			AddErrorIfFalse(!Result3.bHasOtherBoundDelegates, TEXT("Remove of Handle2 should not have other bound delegates."));
+		}
+
+		{
+			Local.Reset();
+			FDelegateHandle Handle1 = Local.Delegate.Add(UObject::StaticClass(), FieldA, UE::FieldNotification::FFieldMulticastDelegate::FDelegate::CreateRaw(&Local, &FLocal::FooA));
+			FDelegateHandle Handle2 = Local.Delegate.Add(UClass::StaticClass(), FieldA, UE::FieldNotification::FFieldMulticastDelegate::FDelegate::CreateRaw(&Local, &FLocal::FooA));
+			Local.Delegate.Add(UClass::StaticClass(), FieldB, UE::FieldNotification::FFieldMulticastDelegate::FDelegate::CreateRaw(&Local, &FLocal::FooB));
+			Local.Delegate.Add(UClass::StaticClass(), FieldC, UE::FieldNotification::FFieldMulticastDelegate::FDelegate::CreateRaw(&Local, &FLocal::FooC));
+			Local.ExpectFieldA = FieldA;
+			Local.ExpectFieldB = FieldB;
+			Local.ExpectFieldC = FieldC;
+			Local.Delegate.Broadcast(UObject::StaticClass(), FieldA);
+			AddErrorIfFalse(Local.FooACount == 1, TEXT("Bad amount of FooA"));
+			AddErrorIfFalse(Local.FooBCount == 0, TEXT("Bad amount of FooB"));
+			AddErrorIfFalse(Local.FooCCount == 0, TEXT("Bad amount of FooC"));
+			Local.Delegate.Broadcast(UClass::StaticClass(), FieldA);
+			AddErrorIfFalse(Local.FooACount == 2, TEXT("Bad amount of FooA"));
+			AddErrorIfFalse(Local.FooBCount == 0, TEXT("Bad amount of FooB"));
+			AddErrorIfFalse(Local.FooCCount == 0, TEXT("Bad amount of FooC"));
+			Local.Delegate.Broadcast(UClass::StaticClass(), FieldB);
+			AddErrorIfFalse(Local.FooACount == 2, TEXT("Bad amount of FooA"));
+			AddErrorIfFalse(Local.FooBCount == 1, TEXT("Bad amount of FooB"));
+			AddErrorIfFalse(Local.FooCCount == 0, TEXT("Bad amount of FooC"));
+			Local.Delegate.Broadcast(UClass::StaticClass(), FieldC);
+			AddErrorIfFalse(Local.FooACount == 2, TEXT("Bad amount of FooA"));
+			AddErrorIfFalse(Local.FooBCount == 1, TEXT("Bad amount of FooB"));
+			AddErrorIfFalse(Local.FooCCount == 1, TEXT("Bad amount of FooC"));
+
+			UE::FieldNotification::FFieldMulticastDelegate::FRemoveResult Result1 = Local.Delegate.Remove(Handle2);
+			AddErrorIfFalse(Result1.Object == UClass::StaticClass(), TEXT("Failed Remove Handle2, bad Object"));
+			AddErrorIfFalse(Result1.FieldId == FieldA, TEXT("Failed Remove Handle2, bad FieldId"));
+			AddErrorIfFalse(Result1.bRemoved, TEXT("Failed Remove Handle2, bad bRemoved"));
+			AddErrorIfFalse(!Result1.bHasOtherBoundDelegates, TEXT("Failed Remove Handle2, bad bHasOtherBoundDelegates"));
+
+			UE::FieldNotification::FFieldMulticastDelegate::FRemoveAllResult Result2 = Local.Delegate.RemoveAll(UClass::StaticClass(), &Local);
+			AddErrorIfFalse(Result2.RemoveCount == 2, TEXT("Failed RemoveAll, bad count"));
+			AddErrorIfFalse(Result2.HasFields.CountSetBits() == 0, TEXT("Failed RemoveAll, bad HasFields"));
+
+			UE::FieldNotification::FFieldMulticastDelegate::FRemoveFromResult Result3 = Local.Delegate.RemoveFrom(UObject::StaticClass(), FieldA, Handle1);
+			AddErrorIfFalse(Result3.bRemoved, TEXT("Bad Remove of FieldA"));
+			AddErrorIfFalse(!Result3.bHasOtherBoundDelegates, TEXT("Remove of Handle1 should have other bound delegates."));
+		}
+	};
+	RunTest1();
+
+	Local.Delegate.Add(
+			UObject::StaticClass(),
+			FieldD,
+			UE::FieldNotification::FFieldMulticastDelegate::FDelegate::CreateLambda([&RunTest1](UObject* InObject, UE::FieldNotification::FFieldId Id)
+			{
+				RunTest1();
+			}));
+	Local.Delegate.Broadcast(UObject::StaticClass(), FieldD);
+
+	return true;
+}
+
+#undef LOCTEXT_NAMESPACE 
+#endif //WITH_AUTOMATION_WORKER
\ No newline at end of file
diff --git a/Engine/Source/Runtime/UMG/Public/Components/Widget.h b/Engine/Source/Runtime/UMG/Public/Components/Widget.h
--- a/Engine/Source/Runtime/UMG/Public/Components/Widget.h
+++ b/Engine/Source/Runtime/UMG/Public/Components/Widget.h
@@ -4,6 +4,8 @@
 
 #include "CoreMinimal.h"
 #include "UObject/ObjectMacros.h"
+#include "FieldNotification/FieldNotificationDeclaration.h"
+#include "FieldNotification/IFieldValueChanged.h"
 #include "Misc/Attribute.h"
 #include "Templates/SubclassOf.h"
 #include "UObject/ScriptMacros.h"
@@ -217,11 +219,20 @@ public:
  * This is the base class for all wrapped Slate controls that are exposed to UObjects.
  */
 UCLASS(Abstract, BlueprintType, Blueprintable)
-class UMG_API UWidget : public UVisual
+class UMG_API UWidget : public UVisual, public INotifyFieldValueChanged
 {
 	GENERATED_UCLASS_BODY()
 
 public:
+	UE_FIELD_NOTIFICATION_DECLARE_CLASS_DESCRIPTOR_BASE_BEGIN(UMG_API)
+		UE_FIELD_NOTIFICATION_DECLARE_FIELD(ToolTipText)
+		UE_FIELD_NOTIFICATION_DECLARE_FIELD(Visibility)
+		UE_FIELD_NOTIFICATION_DECLARE_FIELD(bIsEnabled)
+		UE_FIELD_NOTIFICATION_DECLARE_ENUM_FIELD_BEGIN(ToolTipText)
+		UE_FIELD_NOTIFICATION_DECLARE_ENUM_FIELD(Visibility)
+		UE_FIELD_NOTIFICATION_DECLARE_ENUM_FIELD(bIsEnabled)
+		UE_FIELD_NOTIFICATION_DECLARE_ENUM_FIELD_END()
+	UE_FIELD_NOTIFICATION_DECLARE_CLASS_DESCRIPTOR_BASE_END();
 
 	// Common Bindings - If you add any new common binding, you must provide a UPropertyBinding for it.
 	//                   all primitive binding in UMG goes through native binding evaluators to prevent
@@ -259,7 +270,7 @@ public:
 	FGetBool bIsEnabledDelegate;
 
 	/** Tooltip text to show when the user hovers over the widget with the mouse */
-	UPROPERTY(EditAnywhere, BlueprintReadOnly, Category="Behavior", meta=(MultiLine=true))
+	UPROPERTY(EditAnywhere, BlueprintReadWrite, BlueprintSetter="SetToolTipText", Category="Behavior", meta=(MultiLine=true))
 	FText ToolTipText;
 
 	/** A bindable delegate for ToolTipText */
@@ -644,7 +655,7 @@ public:
 	 *	@param WidgetToFocus When using the Explicit rule, focus on this widget
 	 */
 	UE_DEPRECATED(4.23, "SetNavigationRule is deprecated. Please use either SetNavigationRuleBase or SetNavigationRuleExplicit or SetNavigationRuleCustom or SetNavigationRuleCustomBoundary.")
-	UFUNCTION(BlueprintCallable, Category = "Widget")
+	UFUNCTION(BlueprintCallable, Category = "Widget", meta = (DeprecatedFunction, DeprecatedMessage = "Please use either SetNavigationRuleBase or SetNavigationRuleExplicit or SetNavigationRuleCustom or SetNavigationRuleCustomBoundary."))
 	void SetNavigationRule(EUINavigation Direction, EUINavigationRule Rule, FName WidgetToFocus);
 
 	/**
@@ -707,6 +718,18 @@ public:
 
 	UFUNCTION(BlueprintCallable, Category="Widget")
 	const FGeometry& GetPaintSpaceGeometry() const;
+
+	//~ Begin INotifyFieldValueChanged Interface
+public:
+	virtual FDelegateHandle AddFieldValueChangedDelegate(UE::FieldNotification::FFieldId InFieldId, FFieldValueChangedDelegate InNewDelegate) override final;
+	virtual bool RemoveFieldValueChangedDelegate(UE::FieldNotification::FFieldId InFieldId, FDelegateHandle InHandle) override final;
+	virtual int32 RemoveAllFieldValueChangedDelegates(const void* InUserObject) override final;
+	//~ End INotifyFieldValueChanged Interface
+
+public:
+	void BroadcastFieldValueChanged(UE::FieldNotification::FFieldId InFieldId);
+
+public:
 	/**
 	 * Gets the underlying slate widget or constructs it if it doesn't exist.  If you're looking to replace
 	 * what slate widget gets constructed look for RebuildWidget.  For extremely special cases where you actually
@@ -1055,6 +1078,7 @@ protected:
 	static TArray<TSubclassOf<UPropertyBinding>> BinderClasses;
 
 private:
+	TBitArray<> EnabledFieldNotifications;
 
 #if WITH_EDITORONLY_DATA
 	/** Any flags used by the designer at edit time. */
diff --git a/Engine/Source/Runtime/UMG/Public/FieldNotification/FieldId.h b/Engine/Source/Runtime/UMG/Public/FieldNotification/FieldId.h
new file mode 100644
--- /dev/null
+++ b/Engine/Source/Runtime/UMG/Public/FieldNotification/FieldId.h
@@ -0,0 +1,79 @@
+// Copyright Epic Games, Inc. All Rights Reserved.
+
+#pragma once
+
+#include "CoreMinimal.h"
+#include "FieldNotification/FieldVariant.h"
+#include "UObject/Object.h"
+
+
+namespace UE::FieldNotification
+{
+	
+struct FFieldId
+{
+public:
+	FFieldId() = default;
+
+	explicit FFieldId(FName InFieldName, int32 InBitNumber)
+		: FieldName(InFieldName)
+		, BitNumber(InBitNumber)
+	{
+	}
+
+
+public:
+	bool IsValid() const
+	{
+		return !FieldName.IsNone();
+	}
+
+	int32 GetIndex() const
+	{
+		return BitNumber;
+	}
+	
+	FName GetName() const
+	{
+		return FieldName;
+	}
+
+	FFieldVariant ToVariant(UObject* InContainer) const
+	{
+		if (InContainer && IsValid())
+		{
+			if (UFunction* Function = InContainer->GetClass()->FindFunctionByName(FieldName))
+			{
+				return FFieldVariant(Function);
+			}
+			else if (FProperty* Property = InContainer->GetClass()->FindPropertyByName(FieldName))
+			{
+				return FFieldVariant(Property);
+			}
+		}
+		return FFieldVariant();
+	}
+
+	bool operator==(const FFieldId& Other) const
+	{
+		return FieldName == Other.FieldName;
+	}
+
+	bool operator!=(const FFieldId& Other) const
+	{
+		return FieldName != Other.FieldName;
+	}
+
+	friend int32 GetTypeHash(const FFieldId& Value)
+	{
+		return GetTypeHash(Value.FieldName);
+	}
+
+private:
+	/** Name of the field. It can be a FProperty or UFunction. */
+	FName FieldName;
+	/** The bit this field is linked to. */
+	int32 BitNumber;
+};
+
+} //namespace
\ No newline at end of file
diff --git a/Engine/Source/Runtime/UMG/Public/FieldNotification/FieldMulticastDelegate.h b/Engine/Source/Runtime/UMG/Public/FieldNotification/FieldMulticastDelegate.h
new file mode 100644
--- /dev/null
+++ b/Engine/Source/Runtime/UMG/Public/FieldNotification/FieldMulticastDelegate.h
@@ -0,0 +1,135 @@
+// Copyright Epic Games, Inc. All Rights Reserved.
+
+#pragma once
+
+#include "CoreMinimal.h"
+#include "Delegates/Delegate.h"
+#include "FieldNotification/IFieldValueChanged.h"
+#include "UObject/Object.h"
+
+
+namespace UE::FieldNotification
+{
+	
+class UMG_API FFieldMulticastDelegate
+{
+public:
+	using FDelegate = INotifyFieldValueChanged::FFieldValueChangedDelegate;
+
+private:
+	struct FInstanceProtectedDelegate : public TDelegateBase<FDefaultDelegateUserPolicy>
+	{
+		using Super = TDelegateBase<FDefaultDelegateUserPolicy>;
+		using Super::GetDelegateInstanceProtected;
+	};
+
+	struct FInvocationKey
+	{
+		TWeakObjectPtr<const UObject> Object;
+		FFieldId Id;
+		bool operator<(const FInvocationKey& Element) const
+		{
+			return Id.GetName().FastLess(Element.Id.GetName());
+		}
+	};
+
+	struct FInvocationElement
+	{
+		FInvocationKey Key;
+		FDelegate Delegate;
+		bool operator<(const FInvocationElement& InElement) const
+		{
+			return Key < InElement.Key;
+		}
+	};
+
+	friend bool operator<(const FInvocationElement& A, const FInvocationKey& B)
+	{
+		return A.Key < B;
+	}
+	friend bool operator<(const FInvocationKey& A, const FInvocationElement& B)
+	{
+		return A < B.Key;
+	}
+	friend bool operator<(const FInvocationElement& A, const FFieldId& B)
+	{
+		return A.Key.Id.GetName().FastLess(B.GetName());
+	}
+	friend bool operator<(const FFieldId& A, const FInvocationElement& B)
+	{
+		return A.GetName().FastLess(B.Key.Id.GetName());
+	}
+
+	using InvocationListType = TArray<FInvocationElement>;
+
+public:
+	FDelegateHandle Add(const UObject* InObject, FFieldId InFieldId, FDelegate InNewDelegate);
+
+	struct FRemoveResult
+	{
+		bool bRemoved = false;
+		bool bHasOtherBoundDelegates = false;
+		const UObject* Object = nullptr;
+		FFieldId FieldId;
+	};
+	FRemoveResult Remove(FDelegateHandle InDelegate);
+
+	struct FRemoveFromResult
+	{
+		bool bRemoved = false;
+		bool bHasOtherBoundDelegates = false;
+	};
+	FRemoveFromResult RemoveFrom(const UObject* InObject, FFieldId InFieldId, FDelegateHandle InDelegate);
+
+	struct FRemoveAllResult
+	{
+		int32 RemoveCount = 0;
+		TBitArray<> HasFields;
+	};
+	FRemoveAllResult RemoveAll(const UObject* InObject, const void* InUserObject);
+	FRemoveAllResult RemoveAll(const UObject* InObject, FFieldId InFieldId, const void* InUserObject);
+
+	void Broadcast(UObject* InObject, FFieldId InFieldId);
+
+	void Reset();
+
+private:
+	int32 LowerBound(FFieldId InFieldId) const
+	{
+		int32 Num = AddedEmplaceAt > GetNum(Delegates) ? GetNum(Delegates) : AddedEmplaceAt;
+		return AlgoImpl::LowerBoundInternal(GetData(Delegates), Num, InFieldId, FIdentityFunctor(), TLess<>());
+	}
+
+	int32 LowerBound(const FInvocationKey& InKey) const
+	{
+		int32 Num = AddedEmplaceAt > GetNum(Delegates) ? GetNum(Delegates) : AddedEmplaceAt;
+		return AlgoImpl::LowerBoundInternal(GetData(Delegates), Num, InKey, FIdentityFunctor(), TLess<>());
+	}
+
+	int32 UpperBound(FFieldId InFieldId) const
+	{
+		int32 Num = AddedEmplaceAt > GetNum(Delegates) ? GetNum(Delegates) : AddedEmplaceAt;
+		return AlgoImpl::UpperBoundInternal(GetData(Delegates), Num, InFieldId, FIdentityFunctor(), TLess<>());
+	}
+
+	int32 UpperBound(const FInvocationKey& InKey) const
+	{
+		int32 Num = AddedEmplaceAt > GetNum(Delegates) ? GetNum(Delegates) : AddedEmplaceAt;
+		return AlgoImpl::UpperBoundInternal(GetData(Delegates), Num, InKey, FIdentityFunctor(), TLess<>());
+	}
+
+	const IDelegateInstance* GetDelegateInstance(FDelegate& InDelegate) const
+	{
+		return static_cast<FInstanceProtectedDelegate&>(static_cast<TDelegateBase<FDefaultDelegateUserPolicy>&>(InDelegate)).GetDelegateInstanceProtected();
+	}
+
+	void ExecuteLockOperations();
+
+private:
+	InvocationListType Delegates;
+	int16 DelegateLockCount = 0;
+	int16 CompactionCount = 0;
+	uint16 AddedEmplaceAt = std::numeric_limits<uint16>::max();
+};
+
+} //namespace
\ No newline at end of file
diff --git a/Engine/Source/Runtime/UMG/Public/FieldNotification/FieldNotificationDeclaration.h b/Engine/Source/Runtime/UMG/Public/FieldNotification/FieldNotificationDeclaration.h
new file mode 100644
--- /dev/null
+++ b/Engine/Source/Runtime/UMG/Public/FieldNotification/FieldNotificationDeclaration.h
@@ -0,0 +1,385 @@
+// Copyright Epic Games, Inc. All Rights Reserved.
+
+#pragma once
+
+#include "FieldNotification/FieldId.h"
+#include "FieldNotification/IClassDescriptor.h"
+
+
+/*
+ *	struct FFieldNotificationClassDescriptor : Super::FFieldNotificationClassDescriptor
+ *	{
+ *		static const ::UE::FieldNotification::FFieldId FieldA;
+ *		static const ::UE::FieldNotification::FFieldId FieldB;
+ *		enum
+ *		{
+ *			IndexOf_FieldA = Super::FFieldNotificationClassDescriptor::Max_IndexOf_ + 0,
+ *			IndexOf_FieldB,
+ *			Max_IndexOf_,
+ *		};
+ *		virtual int32 GetNumberOfField() const override { reutrn Max_IndexOf_; }
+ *		virtual ::UE::FieldNotification::FFieldId GetField(FName InFieldName) const override;
+ *		virtual ::UE::FieldNotification::FFieldId GetField(int32 InFieldNumber) const override;
+ *		virtual const ::UE::FieldNotification::IClassDescriptor& GetFieldNotificationDescriptor() const override;
+ *	};
+ */
+
+
+#define UE_FIELD_NOTIFICATION_DECLARE_CLASS_DESCRIPTOR_BEGIN(DllApi) \
+	struct DllApi FFieldNotificationClassDescriptor : public Super::FFieldNotificationClassDescriptor \
+	{ \
+	private: \
+		using SuperDescriptor = Super::FFieldNotificationClassDescriptor; \
+		static const ::UE::FieldNotification::FFieldId* AllFields[]; \
+		friend ThisClass; \
+	public:
+
+
+#define UE_FIELD_NOTIFICATION_DECLARE_CLASS_DESCRIPTOR_END() \
+		virtual int32 GetNumberOfField() const override \
+		{ \
+			return Max_IndexOf_; \
+		} \
+		virtual ::UE::FieldNotification::FFieldId GetField(FName InFieldName) const override \
+		{ \
+			for (int32 Index = 0; Index < Max_IndexOf_-SuperDescriptor::Max_IndexOf_; ++Index) \
+			{ \
+				if (AllFields[Index]->GetName() == InFieldName) \
+				{ \
+					return *(AllFields[Index]); \
+				} \
+			} \
+			return SuperDescriptor::GetField(InFieldName); \
+		} \
+		virtual ::UE::FieldNotification::FFieldId GetField(int32 InFieldNumber) const override \
+		{ \
+			check(InFieldNumber >= 0 && InFieldNumber < Max_IndexOf_); \
+			if (InFieldNumber < SuperDescriptor::Max_IndexOf_) \
+			{ \
+				return SuperDescriptor::GetField(InFieldNumber); \
+			} \
+			return *(AllFields[InFieldNumber - SuperDescriptor::Max_IndexOf_]); \
+		} \
+	}; \
+	virtual const ::UE::FieldNotification::IClassDescriptor& GetFieldNotificationDescriptor() const override
+
+
+#define UE_FIELD_NOTIFICATION_DECLARE_CLASS_DESCRIPTOR_BASE_BEGIN(DllApi) \
+	struct DllApi FFieldNotificationClassDescriptor : public ::UE::FieldNotification::IClassDescriptor \
+	{ \
+	private: \
+		using SuperDescriptor = ::UE::FieldNotification::IClassDescriptor; \
+		static const ::UE::FieldNotification::FFieldId* AllFields[]; \
+		friend ThisClass; \
+	public:
+
+
+#define UE_FIELD_NOTIFICATION_DECLARE_CLASS_DESCRIPTOR_BASE_END() \
+		virtual int32 GetNumberOfField() const override \
+		{ \
+			return Max_IndexOf_; \
+		} \
+		virtual ::UE::FieldNotification::FFieldId GetField(FName InFieldName) const override \
+		{ \
+			for (int32 Index = 0; Index < Max_IndexOf_; ++Index) \
+			{ \
+				if (AllFields[Index]->GetName() == InFieldName) \
+				{ \
+					return *(AllFields[Index]); \
+				} \
+			} \
+			return ::UE::FieldNotification::FFieldId(); \
+		} \
+		virtual ::UE::FieldNotification::FFieldId GetField(int32 InFieldNumber) const override \
+		{ \
+			check(InFieldNumber >= 0 && InFieldNumber < Max_IndexOf_); \
+			return *(AllFields[InFieldNumber]); \
+		} \
+	}; \
+	virtual const ::UE::FieldNotification::IClassDescriptor& GetFieldNotificationDescriptor() const override
+
+
+#define UE_FIELD_NOTIFICATION_DECLARE_ENUM_FIELD_BEGIN(Name) \
+		enum EField \
+		{ \
+			IndexOf_##Name = SuperDescriptor::Max_IndexOf_ + 0,
+
+
+#define UE_FIELD_NOTIFICATION_DECLARE_ENUM_FIELD_END() \
+			Max_IndexOf_, \
+		};
+
+
+#define UE_FIELD_NOTIFICATION_DECLARE_FIELD(Name) \
+	static const ::UE::FieldNotification::FFieldId Name;
+
+
+#define UE_FIELD_NOTIFICATION_DECLARE_ENUM_FIELD(Name) \
+			IndexOf_##Name,
+
+
+#define UE_FIELD_NOTIFICATION_DECLARE_CLASS_DESCRIPTOR_OneField(DllApi, Field1) \
+	UE_FIELD_NOTIFICATION_DECLARE_CLASS_DESCRIPTOR_BEGIN(DllApi) \
+	UE_FIELD_NOTIFICATION_DECLARE_FIELD(Field1) \
+	UE_FIELD_NOTIFICATION_DECLARE_ENUM_FIELD_BEGIN(Field1) \
+	UE_FIELD_NOTIFICATION_DECLARE_ENUM_FIELD_END() \
+	UE_FIELD_NOTIFICATION_DECLARE_CLASS_DESCRIPTOR_END()
+
+
+#define UE_FIELD_NOTIFICATION_DECLARE_CLASS_DESCRIPTOR_TwoFields(DllApi, Field1, Field2) \
+	UE_FIELD_NOTIFICATION_DECLARE_CLASS_DESCRIPTOR_BEGIN(DllApi) \
+	UE_FIELD_NOTIFICATION_DECLARE_FIELD(Field1) \
+	UE_FIELD_NOTIFICATION_DECLARE_FIELD(Field2) \
+	UE_FIELD_NOTIFICATION_DECLARE_ENUM_FIELD_BEGIN(Field1) \
+	UE_FIELD_NOTIFICATION_DECLARE_ENUM_FIELD(Field2) \
+	UE_FIELD_NOTIFICATION_DECLARE_ENUM_FIELD_END() \
+	UE_FIELD_NOTIFICATION_DECLARE_CLASS_DESCRIPTOR_END()
+
+
+#define UE_FIELD_NOTIFICATION_DECLARE_CLASS_DESCRIPTOR_ThreeFields(DllApi, Field1, Field2, Field3) \
+	UE_FIELD_NOTIFICATION_DECLARE_CLASS_DESCRIPTOR_BEGIN(DllApi) \
+	UE_FIELD_NOTIFICATION_DECLARE_FIELD(Field1) \
+	UE_FIELD_NOTIFICATION_DECLARE_FIELD(Field2) \
+	UE_FIELD_NOTIFICATION_DECLARE_FIELD(Field3) \
+	UE_FIELD_NOTIFICATION_DECLARE_ENUM_FIELD_BEGIN(Field1) \
+	UE_FIELD_NOTIFICATION_DECLARE_ENUM_FIELD(Field2) \
+	UE_FIELD_NOTIFICATION_DECLARE_ENUM_FIELD(Field3) \
+	UE_FIELD_NOTIFICATION_DECLARE_ENUM_FIELD_END() \
+	UE_FIELD_NOTIFICATION_DECLARE_CLASS_DESCRIPTOR_END()
+
+
+#define UE_FIELD_NOTIFICATION_DECLARE_CLASS_DESCRIPTOR_FourFields(DllApi, Field1, Field2, Field3, Field4) \
+	UE_FIELD_NOTIFICATION_DECLARE_CLASS_DESCRIPTOR_BEGIN(DllApi) \
+	UE_FIELD_NOTIFICATION_DECLARE_FIELD(Field1) \
+	UE_FIELD_NOTIFICATION_DECLARE_FIELD(Field2) \
+	UE_FIELD_NOTIFICATION_DECLARE_FIELD(Field3) \
+	UE_FIELD_NOTIFICATION_DECLARE_FIELD(Field4) \
+	UE_FIELD_NOTIFICATION_DECLARE_ENUM_FIELD_BEGIN(Field1) \
+	UE_FIELD_NOTIFICATION_DECLARE_ENUM_FIELD(Field2) \
+	UE_FIELD_NOTIFICATION_DECLARE_ENUM_FIELD(Field3) \
+	UE_FIELD_NOTIFICATION_DECLARE_ENUM_FIELD(Field4) \
+	UE_FIELD_NOTIFICATION_DECLARE_ENUM_FIELD_END() \
+	UE_FIELD_NOTIFICATION_DECLARE_CLASS_DESCRIPTOR_END()
+
+
+#define UE_FIELD_NOTIFICATION_DECLARE_CLASS_DESCRIPTOR_FiveFields(DllApi, Field1, Field2, Field3, Field4, Field5) \
+	UE_FIELD_NOTIFICATION_DECLARE_CLASS_DESCRIPTOR_BEGIN(DllApi) \
+	UE_FIELD_NOTIFICATION_DECLARE_FIELD(Field1) \
+	UE_FIELD_NOTIFICATION_DECLARE_FIELD(Field2) \
+	UE_FIELD_NOTIFICATION_DECLARE_FIELD(Field3) \
+	UE_FIELD_NOTIFICATION_DECLARE_FIELD(Field4) \
+	UE_FIELD_NOTIFICATION_DECLARE_FIELD(Field5) \
+	UE_FIELD_NOTIFICATION_DECLARE_ENUM_FIELD_BEGIN(Field1) \
+	UE_FIELD_NOTIFICATION_DECLARE_ENUM_FIELD(Field2) \
+	UE_FIELD_NOTIFICATION_DECLARE_ENUM_FIELD(Field3) \
+	UE_FIELD_NOTIFICATION_DECLARE_ENUM_FIELD(Field4) \
+	UE_FIELD_NOTIFICATION_DECLARE_ENUM_FIELD(Field5) \
+	UE_FIELD_NOTIFICATION_DECLARE_ENUM_FIELD_END() \
+	UE_FIELD_NOTIFICATION_DECLARE_CLASS_DESCRIPTOR_END()
+
+
+#define UE_FIELD_NOTIFICATION_DECLARE_CLASS_DESCRIPTOR_SixFields(DllApi, Field1, Field2, Field3, Field4, Field5, Field6) \
+	UE_FIELD_NOTIFICATION_DECLARE_CLASS_DESCRIPTOR_BEGIN(DllApi) \
+	UE_FIELD_NOTIFICATION_DECLARE_FIELD(Field1) \
+	UE_FIELD_NOTIFICATION_DECLARE_FIELD(Field2) \
+	UE_FIELD_NOTIFICATION_DECLARE_FIELD(Field3) \
+	UE_FIELD_NOTIFICATION_DECLARE_FIELD(Field4) \
+	UE_FIELD_NOTIFICATION_DECLARE_FIELD(Field5) \
+	UE_FIELD_NOTIFICATION_DECLARE_FIELD(Field6) \
+	UE_FIELD_NOTIFICATION_DECLARE_ENUM_FIELD_BEGIN(Field1) \
+	UE_FIELD_NOTIFICATION_DECLARE_ENUM_FIELD(Field2) \
+	UE_FIELD_NOTIFICATION_DECLARE_ENUM_FIELD(Field3) \
+	UE_FIELD_NOTIFICATION_DECLARE_ENUM_FIELD(Field4) \
+	UE_FIELD_NOTIFICATION_DECLARE_ENUM_FIELD(Field5) \
+	UE_FIELD_NOTIFICATION_DECLARE_ENUM_FIELD(Field6) \
+	UE_FIELD_NOTIFICATION_DECLARE_ENUM_FIELD_END() \
+	UE_FIELD_NOTIFICATION_DECLARE_CLASS_DESCRIPTOR_END()
+
+
+#define UE_FIELD_NOTIFICATION_DECLARE_CLASS_DESCRIPTOR_SevenFields(DllApi, Field1, Field2, Field3, Field4, Field5, Field6, Field7) \
+	UE_FIELD_NOTIFICATION_DECLARE_CLASS_DESCRIPTOR_BEGIN(DllApi) \
+	UE_FIELD_NOTIFICATION_DECLARE_FIELD(Field1) \
+	UE_FIELD_NOTIFICATION_DECLARE_FIELD(Field2) \
+	UE_FIELD_NOTIFICATION_DECLARE_FIELD(Field3) \
+	UE_FIELD_NOTIFICATION_DECLARE_FIELD(Field4) \
+	UE_FIELD_NOTIFICATION_DECLARE_FIELD(Field5) \
+	UE_FIELD_NOTIFICATION_DECLARE_FIELD(Field6) \
+	UE_FIELD_NOTIFICATION_DECLARE_FIELD(Field7) \
+	UE_FIELD_NOTIFICATION_DECLARE_ENUM_FIELD_BEGIN(Field1) \
+	UE_FIELD_NOTIFICATION_DECLARE_ENUM_FIELD(Field2) \
+	UE_FIELD_NOTIFICATION_DECLARE_ENUM_FIELD(Field3) \
+	UE_FIELD_NOTIFICATION_DECLARE_ENUM_FIELD(Field4) \
+	UE_FIELD_NOTIFICATION_DECLARE_ENUM_FIELD(Field5) \
+	UE_FIELD_NOTIFICATION_DECLARE_ENUM_FIELD(Field6) \
+	UE_FIELD_NOTIFICATION_DECLARE_ENUM_FIELD(Field7) \
+	UE_FIELD_NOTIFICATION_DECLARE_ENUM_FIELD_END() \
+	UE_FIELD_NOTIFICATION_DECLARE_CLASS_DESCRIPTOR_END()
+
+
+#define UE_FIELD_NOTIFICATION_DECLARE_CLASS_DESCRIPTOR_EightFields(DllApi, Field1, Field2, Field3, Field4, Field5, Field6, Field7, Field8) \
+	UE_FIELD_NOTIFICATION_DECLARE_CLASS_DESCRIPTOR_BEGIN(DllApi) \
+	UE_FIELD_NOTIFICATION_DECLARE_FIELD(Field1) \
+	UE_FIELD_NOTIFICATION_DECLARE_FIELD(Field2) \
+	UE_FIELD_NOTIFICATION_DECLARE_FIELD(Field3) \
+	UE_FIELD_NOTIFICATION_DECLARE_FIELD(Field4) \
+	UE_FIELD_NOTIFICATION_DECLARE_FIELD(Field5) \
+	UE_FIELD_NOTIFICATION_DECLARE_FIELD(Field6) \
+	UE_FIELD_NOTIFICATION_DECLARE_FIELD(Field7) \
+	UE_FIELD_NOTIFICATION_DECLARE_FIELD(Field8) \
+	UE_FIELD_NOTIFICATION_DECLARE_ENUM_FIELD_BEGIN(Field1) \
+	UE_FIELD_NOTIFICATION_DECLARE_ENUM_FIELD(Field2) \
+	UE_FIELD_NOTIFICATION_DECLARE_ENUM_FIELD(Field3) \
+	UE_FIELD_NOTIFICATION_DECLARE_ENUM_FIELD(Field4) \
+	UE_FIELD_NOTIFICATION_DECLARE_ENUM_FIELD(Field5) \
+	UE_FIELD_NOTIFICATION_DECLARE_ENUM_FIELD(Field6) \
+	UE_FIELD_NOTIFICATION_DECLARE_ENUM_FIELD(Field7) \
+	UE_FIELD_NOTIFICATION_DECLARE_ENUM_FIELD(Field8) \
+	UE_FIELD_NOTIFICATION_DECLARE_ENUM_FIELD_END() \
+	UE_FIELD_NOTIFICATION_DECLARE_CLASS_DESCRIPTOR_END()
+
+
+#define UE_FIELD_NOTIFICATION_IMPLEMENTATION_BEGIN(ClassName) \
+	const ::UE::FieldNotification::FFieldId* ClassName::FFieldNotificationClassDescriptor::AllFields[] = {
+
+
+#if DO_CHECK
+#define UE_FIELD_NOTIFICATION_IMPLEMENTATION_END(ClassName) \
+	}; \
+	const ::UE::FieldNotification::IClassDescriptor& ClassName::GetFieldNotificationDescriptor() const \
+	{ \
+		static_assert(UE_ARRAY_COUNT(ClassName::FFieldNotificationClassDescriptor::AllFields) == ClassName::FFieldNotificationClassDescriptor::Max_IndexOf_-ClassName::FFieldNotificationClassDescriptor::SuperDescriptor::Max_IndexOf_, "The descriptor for class " #ClassName " doesn't not implement the same number of field as declared in in the constructor."); \
+		struct FLocal \
+		{ \
+			ClassName::FFieldNotificationClassDescriptor Instance; \
+			FLocal() \
+			{ \
+				int32 NumberOfField = Instance.GetNumberOfField(); \
+				for (int32 Index = 0; Index < NumberOfField; ++Index) \
+				{ \
+					::UE::FieldNotification::FFieldId Id = Instance.GetField(Index); \
+					ensureAlwaysMsgf(Id.IsValid(), TEXT("The FFieldId '%d' for the class '%s' is invalid"), Index, *ClassName::StaticClass()->GetName()); \
+					ensureAlwaysMsgf(Id.GetIndex() == Index, TEXT("The FFieldId at index '%d' doesn't match the id index '%d'."), Index, Id.GetIndex()); \
+				} \
+			} \
+		}; \
+		static FLocal Local; \
+		return Local.Instance; \
+	}
+#else
+#define UE_FIELD_NOTIFICATION_IMPLEMENTATION_END(ClassName) \
+	}; \
+	const ::UE::FieldNotification::IClassDescriptor& ClassName::GetFieldNotificationDescriptor() const \
+	{ \
+		static FFieldNotificationClassDescriptor Instance; \
+		return Instance; \
+	}
+#endif
+
+
+#define UE_FIELD_NOTIFICATION_IMPLEMENT_FIELD(ClassName, MemberName) const ::UE::FieldNotification::FFieldId ClassName::FFieldNotificationClassDescriptor::MemberName( FName(TEXT(#MemberName)), ClassName::FFieldNotificationClassDescriptor::IndexOf_##MemberName );
+#define UE_FIELD_NOTIFICATION_IMPLEMENT_ENUM_FIELD(ClassName, MemberName) &ClassName::FFieldNotificationClassDescriptor::MemberName,
+
+
+#define UE_FIELD_NOTIFICATION_IMPLEMENT_CLASS_DESCRIPTOR_OneField(ClassName, Field1) \
+	UE_FIELD_NOTIFICATION_IMPLEMENT_FIELD(ClassName, Field1) \
+	UE_FIELD_NOTIFICATION_IMPLEMENTATION_BEGIN(ClassName) \
+	UE_FIELD_NOTIFICATION_IMPLEMENT_ENUM_FIELD(ClassName, Field1) \
+	UE_FIELD_NOTIFICATION_IMPLEMENTATION_END(ClassName)
+
+#define UE_FIELD_NOTIFICATION_IMPLEMENT_CLASS_DESCRIPTOR_TwoFields(ClassName, Field1, Field2) \
+	UE_FIELD_NOTIFICATION_IMPLEMENT_FIELD(ClassName, Field1) \
+	UE_FIELD_NOTIFICATION_IMPLEMENT_FIELD(ClassName, Field2) \
+	UE_FIELD_NOTIFICATION_IMPLEMENTATION_BEGIN(ClassName) \
+	UE_FIELD_NOTIFICATION_IMPLEMENT_ENUM_FIELD(ClassName, Field1) \
+	UE_FIELD_NOTIFICATION_IMPLEMENT_ENUM_FIELD(ClassName, Field2) \
+	UE_FIELD_NOTIFICATION_IMPLEMENTATION_END(ClassName)
+
+#define UE_FIELD_NOTIFICATION_IMPLEMENT_CLASS_DESCRIPTOR_ThreeFields(ClassName, Field1, Field2, Field3) \
+	UE_FIELD_NOTIFICATION_IMPLEMENT_FIELD(ClassName, Field1) \
+	UE_FIELD_NOTIFICATION_IMPLEMENT_FIELD(ClassName, Field2) \
+	UE_FIELD_NOTIFICATION_IMPLEMENT_FIELD(ClassName, Field3) \
+	UE_FIELD_NOTIFICATION_IMPLEMENTATION_BEGIN(ClassName) \
+	UE_FIELD_NOTIFICATION_IMPLEMENT_ENUM_FIELD(ClassName, Field1) \
+	UE_FIELD_NOTIFICATION_IMPLEMENT_ENUM_FIELD(ClassName, Field2) \
+	UE_FIELD_NOTIFICATION_IMPLEMENT_ENUM_FIELD(ClassName, Field3) \
+	UE_FIELD_NOTIFICATION_IMPLEMENTATION_END(ClassName)
+
+#define UE_FIELD_NOTIFICATION_IMPLEMENT_CLASS_DESCRIPTOR_FourFields(ClassName, Field1, Field2, Field3, Field4) \
+	UE_FIELD_NOTIFICATION_IMPLEMENT_FIELD(ClassName, Field1) \
+	UE_FIELD_NOTIFICATION_IMPLEMENT_FIELD(ClassName, Field2) \
+	UE_FIELD_NOTIFICATION_IMPLEMENT_FIELD(ClassName, Field3) \
+	UE_FIELD_NOTIFICATION_IMPLEMENT_FIELD(ClassName, Field4) \
+	UE_FIELD_NOTIFICATION_IMPLEMENTATION_BEGIN(ClassName) \
+	UE_FIELD_NOTIFICATION_IMPLEMENT_ENUM_FIELD(ClassName, Field1) \
+	UE_FIELD_NOTIFICATION_IMPLEMENT_ENUM_FIELD(ClassName, Field2) \
+	UE_FIELD_NOTIFICATION_IMPLEMENT_ENUM_FIELD(ClassName, Field3) \
+	UE_FIELD_NOTIFICATION_IMPLEMENT_ENUM_FIELD(ClassName, Field4) \
+	UE_FIELD_NOTIFICATION_IMPLEMENTATION_END(ClassName)
+
+#define UE_FIELD_NOTIFICATION_IMPLEMENT_CLASS_DESCRIPTOR_FiveFields(ClassName, Field1, Field2, Field3, Field4, Field5) \
+	UE_FIELD_NOTIFICATION_IMPLEMENT_FIELD(ClassName, Field1) \
+	UE_FIELD_NOTIFICATION_IMPLEMENT_FIELD(ClassName, Field2) \
+	UE_FIELD_NOTIFICATION_IMPLEMENT_FIELD(ClassName, Field3) \
+	UE_FIELD_NOTIFICATION_IMPLEMENT_FIELD(ClassName, Field4) \
+	UE_FIELD_NOTIFICATION_IMPLEMENT_FIELD(ClassName, Field5) \
+	UE_FIELD_NOTIFICATION_IMPLEMENTATION_BEGIN(ClassName) \
+	UE_FIELD_NOTIFICATION_IMPLEMENT_ENUM_FIELD(ClassName, Field1) \
+	UE_FIELD_NOTIFICATION_IMPLEMENT_ENUM_FIELD(ClassName, Field2) \
+	UE_FIELD_NOTIFICATION_IMPLEMENT_ENUM_FIELD(ClassName, Field3) \
+	UE_FIELD_NOTIFICATION_IMPLEMENT_ENUM_FIELD(ClassName, Field4) \
+	UE_FIELD_NOTIFICATION_IMPLEMENT_ENUM_FIELD(ClassName, Field5) \
+	UE_FIELD_NOTIFICATION_IMPLEMENTATION_END(ClassName)
+
+#define UE_FIELD_NOTIFICATION_IMPLEMENT_CLASS_DESCRIPTOR_SixFields(ClassName, Field1, Field2, Field3, Field4, Field5, Field6) \
+	UE_FIELD_NOTIFICATION_IMPLEMENT_FIELD(ClassName, Field1) \
+	UE_FIELD_NOTIFICATION_IMPLEMENT_FIELD(ClassName, Field2) \
+	UE_FIELD_NOTIFICATION_IMPLEMENT_FIELD(ClassName, Field3) \
+	UE_FIELD_NOTIFICATION_IMPLEMENT_FIELD(ClassName, Field4) \
+	UE_FIELD_NOTIFICATION_IMPLEMENT_FIELD(ClassName, Field5) \
+	UE_FIELD_NOTIFICATION_IMPLEMENT_FIELD(ClassName, Field6) \
+	UE_FIELD_NOTIFICATION_IMPLEMENTATION_BEGIN(ClassName) \
+	UE_FIELD_NOTIFICATION_IMPLEMENT_ENUM_FIELD(ClassName, Field1) \
+	UE_FIELD_NOTIFICATION_IMPLEMENT_ENUM_FIELD(ClassName, Field2) \
+	UE_FIELD_NOTIFICATION_IMPLEMENT_ENUM_FIELD(ClassName, Field3) \
+	UE_FIELD_NOTIFICATION_IMPLEMENT_ENUM_FIELD(ClassName, Field4) \
+	UE_FIELD_NOTIFICATION_IMPLEMENT_ENUM_FIELD(ClassName, Field5) \
+	UE_FIELD_NOTIFICATION_IMPLEMENT_ENUM_FIELD(ClassName, Field6) \
+	UE_FIELD_NOTIFICATION_IMPLEMENTATION_END(ClassName)
+
+#define UE_FIELD_NOTIFICATION_IMPLEMENT_CLASS_DESCRIPTOR_SevenFields(ClassName, Field1, Field2, Field3, Field4, Field5, Field6, Field7) \
+	UE_FIELD_NOTIFICATION_IMPLEMENT_FIELD(ClassName, Field1) \
+	UE_FIELD_NOTIFICATION_IMPLEMENT_FIELD(ClassName, Field2) \
+	UE_FIELD_NOTIFICATION_IMPLEMENT_FIELD(ClassName, Field3) \
+	UE_FIELD_NOTIFICATION_IMPLEMENT_FIELD(ClassName, Field4) \
+	UE_FIELD_NOTIFICATION_IMPLEMENT_FIELD(ClassName, Field5) \
+	UE_FIELD_NOTIFICATION_IMPLEMENT_FIELD(ClassName, Field6) \
+	UE_FIELD_NOTIFICATION_IMPLEMENT_FIELD(ClassName, Field7) \
+	UE_FIELD_NOTIFICATION_IMPLEMENTATION_BEGIN(ClassName) \
+	UE_FIELD_NOTIFICATION_IMPLEMENT_ENUM_FIELD(ClassName, Field1) \
+	UE_FIELD_NOTIFICATION_IMPLEMENT_ENUM_FIELD(ClassName, Field2) \
+	UE_FIELD_NOTIFICATION_IMPLEMENT_ENUM_FIELD(ClassName, Field3) \
+	UE_FIELD_NOTIFICATION_IMPLEMENT_ENUM_FIELD(ClassName, Field4) \
+	UE_FIELD_NOTIFICATION_IMPLEMENT_ENUM_FIELD(ClassName, Field5) \
+	UE_FIELD_NOTIFICATION_IMPLEMENT_ENUM_FIELD(ClassName, Field6) \
+	UE_FIELD_NOTIFICATION_IMPLEMENT_ENUM_FIELD(ClassName, Field7) \
+	UE_FIELD_NOTIFICATION_IMPLEMENTATION_END(ClassName)
+
+#define UE_FIELD_NOTIFICATION_IMPLEMENT_CLASS_DESCRIPTOR_EightFields(ClassName, Field1, Field2, Field3, Field4, Field5, Field6, Field7, Field8) \
+	UE_FIELD_NOTIFICATION_IMPLEMENT_FIELD(ClassName, Field1) \
+	UE_FIELD_NOTIFICATION_IMPLEMENT_FIELD(ClassName, Field2) \
+	UE_FIELD_NOTIFICATION_IMPLEMENT_FIELD(ClassName, Field3) \
+	UE_FIELD_NOTIFICATION_IMPLEMENT_FIELD(ClassName, Field4) \
+	UE_FIELD_NOTIFICATION_IMPLEMENT_FIELD(ClassName, Field5) \
+	UE_FIELD_NOTIFICATION_IMPLEMENT_FIELD(ClassName, Field6) \
+	UE_FIELD_NOTIFICATION_IMPLEMENT_FIELD(ClassName, Field7) \
+	UE_FIELD_NOTIFICATION_IMPLEMENT_FIELD(ClassName, Field8) \
+	UE_FIELD_NOTIFICATION_IMPLEMENTATION_BEGIN(ClassName) \
+	UE_FIELD_NOTIFICATION_IMPLEMENT_ENUM_FIELD(ClassName, Field1) \
+	UE_FIELD_NOTIFICATION_IMPLEMENT_ENUM_FIELD(ClassName, Field2) \
+	UE_FIELD_NOTIFICATION_IMPLEMENT_ENUM_FIELD(ClassName, Field3) \
+	UE_FIELD_NOTIFICATION_IMPLEMENT_ENUM_FIELD(ClassName, Field4) \
+	UE_FIELD_NOTIFICATION_IMPLEMENT_ENUM_FIELD(ClassName, Field5) \
+	UE_FIELD_NOTIFICATION_IMPLEMENT_ENUM_FIELD(ClassName, Field6) \
+	UE_FIELD_NOTIFICATION_IMPLEMENT_ENUM_FIELD(ClassName, Field7) \
+	UE_FIELD_NOTIFICATION_IMPLEMENT_ENUM_FIELD(ClassName, Field8) \
+	UE_FIELD_NOTIFICATION_IMPLEMENTATION_END(ClassName)
diff --git a/Engine/Source/Runtime/UMG/Public/FieldNotification/FieldVariant.h b/Engine/Source/Runtime/UMG/Public/FieldNotification/FieldVariant.h
new file mode 100644
--- /dev/null
+++ b/Engine/Source/Runtime/UMG/Public/FieldNotification/FieldVariant.h
@@ -0,0 +1,68 @@
+// Copyright Epic Games, Inc. All Rights Reserved.
+
+#pragma once
+
+#include "CoreMinimal.h"
+#include "UObject/Field.h"
+
+namespace UE::FieldNotification
+{
+
+struct FFieldVariant : private ::FFieldVariant
+{
+public:
+	FFieldVariant() = default;
+
+	explicit FFieldVariant(const FProperty* InProperty)
+		: ::FFieldVariant(InProperty)
+	{
+	}
+
+	explicit FFieldVariant(const UFunction* InFunction)
+		: ::FFieldVariant(InFunction)
+	{
+	}
+
+public:
+	using ::FFieldVariant::IsValid;
+	using ::FFieldVariant::GetOwnerClass;
+	using ::FFieldVariant::GetFName;
+	using ::FFieldVariant::operator bool;
+	using ::FFieldVariant::operator ==;
+	using ::FFieldVariant::operator !=;
+#if WITH_EDITORONLY_DATA
+	using ::FFieldVariant::HasMetaData;
+#endif
+
+	inline bool IsProperty() const
+	{
+		return IsValid() && !IsUObject();
+	}
+
+	inline bool IsFunction() const
+	{
+		return IsValid() && IsUObject();
+	}
+
+	inline FProperty* GetProperty()
+	{
+		return IsProperty() ? static_cast<FProperty*>(ToFieldUnsafe()) : nullptr;
+	}
+	
+	inline const FProperty* GetProperty() const
+	{
+		return IsProperty() ? static_cast<const FProperty*>(ToFieldUnsafe()) : nullptr;
+	}
+
+	inline UFunction* GetFunction()
+	{
+		return IsFunction() ? static_cast<UFunction*>(ToUObjectUnsafe()) : nullptr;
+	}
+
+	inline const UFunction* GetFunction() const
+	{
+		return IsFunction() ? static_cast<const UFunction*>(ToUObjectUnsafe()) : nullptr;
+	}
+};
+
+} //namespace
diff --git a/Engine/Source/Runtime/UMG/Public/FieldNotification/IClassDescriptor.h b/Engine/Source/Runtime/UMG/Public/FieldNotification/IClassDescriptor.h
new file mode 100644
--- /dev/null
+++ b/Engine/Source/Runtime/UMG/Public/FieldNotification/IClassDescriptor.h
@@ -0,0 +1,24 @@
+// Copyright Epic Games, Inc. All Rights Reserved.
+
+#pragma once
+
+#include "CoreMinimal.h"
+#include "FieldNotification/FieldId.h"
+
+namespace UE::FieldNotification
+{
+	
+	struct IClassDescriptor
+	{
+		enum
+		{
+			Max_IndexOf_ = 0,
+		};
+
+		virtual int32 GetNumberOfField() const = 0;
+		virtual FFieldId GetField(FName InFieldName) const = 0;
+		virtual FFieldId GetField(int32 InFieldNumber) const = 0;
+		virtual ~IClassDescriptor() = default;
+	};
+
+} // namespace
\ No newline at end of file
diff --git a/Engine/Source/Runtime/UMG/Public/FieldNotification/IFieldValueChanged.h b/Engine/Source/Runtime/UMG/Public/FieldNotification/IFieldValueChanged.h
new file mode 100644
--- /dev/null
+++ b/Engine/Source/Runtime/UMG/Public/FieldNotification/IFieldValueChanged.h
@@ -0,0 +1,40 @@
+// Copyright Epic Games, Inc. All Rights Reserved.
+
+#pragma once
+
+#include "CoreMinimal.h"
+#include "FieldNotification/FieldId.h"
+#include "FieldNotification/IClassDescriptor.h"
+#include "UObject/Interface.h"
+
+#include "IFieldValueChanged.generated.h"
+
+
+UINTERFACE(MinimalAPI, NotBlueprintable)
+class UNotifyFieldValueChanged : public UInterface
+{
+	GENERATED_BODY()
+};
+
+
+class INotifyFieldValueChanged : public IInterface
+{
+	GENERATED_BODY()
+
+public:
+	DECLARE_MULTICAST_DELEGATE_TwoParams(FFieldValueChangedMulticastDelegate, UObject*, UE::FieldNotification::FFieldId);
+	using FFieldValueChangedDelegate = FFieldValueChangedMulticastDelegate::FDelegate;
+
+public:
+	/** Add a delegate that will be notified when the FieldId is value changed. */
+	virtual FDelegateHandle AddFieldValueChangedDelegate(UE::FieldNotification::FFieldId FieldId, FFieldValueChangedDelegate InNewDelegate) = 0;
+
+	/** Remove a delegate that was added. */
+	virtual bool RemoveFieldValueChangedDelegate(UE::FieldNotification::FFieldId FieldId, FDelegateHandle Handle) = 0;
+
+	/** Remove all the delegate that are bound to the specified UserObject. */
+	virtual int32 RemoveAllFieldValueChangedDelegates(const void* InUserObject) = 0;
+
+	/** @returns the list of all the field that can notify when their value changes. */
+	virtual const UE::FieldNotification::IClassDescriptor& GetFieldNotificationDescriptor() const = 0;
+};
