UMG: Add FieldNotification to Checkbox & ProgressBar.

diff --git a/Engine/Source/Runtime/UMG/Private/Components/CheckBox.cpp b/Engine/Source/Runtime/UMG/Private/Components/CheckBox.cpp
--- a/Engine/Source/Runtime/UMG/Private/Components/CheckBox.cpp
+++ b/Engine/Source/Runtime/UMG/Private/Components/CheckBox.cpp
@@ -12,6 +12,8 @@
 /////////////////////////////////////////////////////
 // UCheckBox
 
+UE_FIELD_NOTIFICATION_IMPLEMENT_CLASS_DESCRIPTOR_OneField(UCheckBox, CheckedState);
+
 static FCheckBoxStyle* DefaultCheckboxStyle = nullptr;
 
 #if WITH_EDITOR
@@ -175,7 +177,13 @@ ECheckBoxState UCheckBox::GetCheckedState() const
 
 void UCheckBox::SetIsChecked(bool InIsChecked)
 {
-	CheckedState = InIsChecked ? ECheckBoxState::Checked : ECheckBoxState::Unchecked;
+	ECheckBoxState NewState = InIsChecked ? ECheckBoxState::Checked : ECheckBoxState::Unchecked;
+	if (NewState != CheckedState)
+	{
+		CheckedState = NewState;
+		BroadcastFieldValueChanged(FFieldNotificationClassDescriptor::CheckedState);
+	}
+
 	if ( MyCheckbox.IsValid() )
 	{
 		MyCheckbox->SetIsChecked(PROPERTY_BINDING(ECheckBoxState, CheckedState));
@@ -184,7 +192,12 @@ void UCheckBox::SetIsChecked(bool InIsChecked)
 
 void UCheckBox::SetCheckedState(ECheckBoxState InCheckedState)
 {
-	CheckedState = InCheckedState;
+	if (CheckedState != InCheckedState)
+	{
+		CheckedState = InCheckedState;
+		BroadcastFieldValueChanged(FFieldNotificationClassDescriptor::CheckedState);
+	}
+
 	if ( MyCheckbox.IsValid() )
 	{
 		MyCheckbox->SetIsChecked(PROPERTY_BINDING(ECheckBoxState, CheckedState));
@@ -193,7 +206,11 @@ void UCheckBox::SetCheckedState(ECheckBoxState InCheckedState)
 
 void UCheckBox::SlateOnCheckStateChangedCallback(ECheckBoxState NewState)
 {
-	CheckedState = NewState;
+	if (CheckedState != NewState)
+	{
+		CheckedState = NewState;
+		BroadcastFieldValueChanged(FFieldNotificationClassDescriptor::CheckedState);
+	}
 
 	//@TODO: Choosing to treat Undetermined as Checked
 	const bool bWantsToBeChecked = NewState != ECheckBoxState::Unchecked;
diff --git a/Engine/Source/Runtime/UMG/Private/Components/ProgressBar.cpp b/Engine/Source/Runtime/UMG/Private/Components/ProgressBar.cpp
--- a/Engine/Source/Runtime/UMG/Private/Components/ProgressBar.cpp
+++ b/Engine/Source/Runtime/UMG/Private/Components/ProgressBar.cpp
@@ -9,6 +9,10 @@
 /////////////////////////////////////////////////////
 // UProgressBar
 
+
+UE_FIELD_NOTIFICATION_IMPLEMENT_CLASS_DESCRIPTOR_ThreeFields(UProgressBar, Percent, FillColorAndOpacity, bIsMarquee);
+
+
 static FProgressBarStyle* DefaultProgressBarStyle = nullptr;
 
 #if WITH_EDITOR
@@ -88,7 +92,12 @@ void UProgressBar::SynchronizeProperties()
 
 void UProgressBar::SetIsMarquee(bool InbIsMarquee)
 {
-	bIsMarquee = InbIsMarquee;
+	if (bIsMarquee != InbIsMarquee)
+	{
+		bIsMarquee = InbIsMarquee;
+		BroadcastFieldValueChanged(FFieldNotificationClassDescriptor::bIsMarquee);
+	}
+
 	if (MyProgressBar.IsValid())
 	{
 		MyProgressBar->SetPercent(bIsMarquee ? TOptional<float>() : Percent);
@@ -97,7 +106,12 @@ void UProgressBar::SetIsMarquee(bool InbIsMarquee)
 
 void UProgressBar::SetFillColorAndOpacity(FLinearColor Color)
 {
-	FillColorAndOpacity = Color;
+	if (FillColorAndOpacity != Color)
+	{
+		FillColorAndOpacity = Color;
+		BroadcastFieldValueChanged(FFieldNotificationClassDescriptor::FillColorAndOpacity);
+	}
+
 	if (MyProgressBar.IsValid())
 	{
 		MyProgressBar->SetFillColorAndOpacity(FillColorAndOpacity);
@@ -106,7 +120,12 @@ void UProgressBar::SetFillColorAndOpacity(FLinearColor Color)
 
 void UProgressBar::SetPercent(float InPercent)
 {
-	Percent = InPercent;
+	if (Percent != InPercent)
+	{
+		Percent = InPercent;
+		BroadcastFieldValueChanged(FFieldNotificationClassDescriptor::Percent);
+	}
+
 	if (MyProgressBar.IsValid())
 	{
 		MyProgressBar->SetPercent(InPercent);
diff --git a/Engine/Source/Runtime/UMG/Public/Components/CheckBox.h b/Engine/Source/Runtime/UMG/Public/Components/CheckBox.h
--- a/Engine/Source/Runtime/UMG/Public/Components/CheckBox.h
+++ b/Engine/Source/Runtime/UMG/Public/Components/CheckBox.h
@@ -31,6 +31,9 @@ class UMG_API UCheckBox : public UContentWidget
 {
 	GENERATED_UCLASS_BODY()
 
+public:
+	UE_FIELD_NOTIFICATION_DECLARE_CLASS_DESCRIPTOR_OneField(CheckedState);
+
 public:
 	/** Whether the check box is currently in a checked state */
 	UPROPERTY(EditAnywhere, Category=Appearance)
diff --git a/Engine/Source/Runtime/UMG/Public/Components/ProgressBar.h b/Engine/Source/Runtime/UMG/Public/Components/ProgressBar.h
--- a/Engine/Source/Runtime/UMG/Public/Components/ProgressBar.h
+++ b/Engine/Source/Runtime/UMG/Public/Components/ProgressBar.h
@@ -23,7 +23,9 @@ class UMG_API UProgressBar : public UWidget
 	GENERATED_UCLASS_BODY()
 	
 public:
+	UE_FIELD_NOTIFICATION_DECLARE_CLASS_DESCRIPTOR_ThreeFields(Percent, FillColorAndOpacity, bIsMarquee);
 
+public:
 	/** The progress bar style */
 	UPROPERTY(EditAnywhere, BlueprintReadWrite, Category="Style", meta=( DisplayName="Style" ))
 	FProgressBarStyle WidgetStyle;
diff --git a/Engine/Source/Runtime/UMG/Public/Components/Widget.h b/Engine/Source/Runtime/UMG/Public/Components/Widget.h
--- a/Engine/Source/Runtime/UMG/Public/Components/Widget.h
+++ b/Engine/Source/Runtime/UMG/Public/Components/Widget.h
@@ -224,7 +224,7 @@ class UMG_API UWidget : public UVisual, public INotifyFieldValueChanged
 	GENERATED_UCLASS_BODY()
 
 public:
-	UE_FIELD_NOTIFICATION_DECLARE_CLASS_DESCRIPTOR_BASE_BEGIN(UMG_API)
+	UE_FIELD_NOTIFICATION_DECLARE_CLASS_DESCRIPTOR_BASE_BEGIN()
 		UE_FIELD_NOTIFICATION_DECLARE_FIELD(ToolTipText)
 		UE_FIELD_NOTIFICATION_DECLARE_FIELD(Visibility)
 		UE_FIELD_NOTIFICATION_DECLARE_FIELD(bIsEnabled)
diff --git a/Engine/Source/Runtime/UMG/Public/FieldNotification/FieldNotificationDeclaration.h b/Engine/Source/Runtime/UMG/Public/FieldNotification/FieldNotificationDeclaration.h
--- a/Engine/Source/Runtime/UMG/Public/FieldNotification/FieldNotificationDeclaration.h
+++ b/Engine/Source/Runtime/UMG/Public/FieldNotification/FieldNotificationDeclaration.h
@@ -25,8 +25,8 @@
  */
 
 
-#define UE_FIELD_NOTIFICATION_DECLARE_CLASS_DESCRIPTOR_BEGIN(DllApi) \
-	struct DllApi FFieldNotificationClassDescriptor : public Super::FFieldNotificationClassDescriptor \
+#define UE_FIELD_NOTIFICATION_DECLARE_CLASS_DESCRIPTOR_BEGIN() \
+	struct FFieldNotificationClassDescriptor : public Super::FFieldNotificationClassDescriptor \
 	{ \
 	private: \
 		using SuperDescriptor = Super::FFieldNotificationClassDescriptor; \
@@ -61,11 +61,15 @@
 			return *(AllFields[InFieldNumber - SuperDescriptor::Max_IndexOf_]); \
 		} \
 	}; \
-	virtual const ::UE::FieldNotification::IClassDescriptor& GetFieldNotificationDescriptor() const override
+	virtual const ::UE::FieldNotification::IClassDescriptor& GetFieldNotificationDescriptor() const override \
+ 	{ \
+ 		static FFieldNotificationClassDescriptor Instance; \
+ 		return Instance; \
+ 	}
 
 
-#define UE_FIELD_NOTIFICATION_DECLARE_CLASS_DESCRIPTOR_BASE_BEGIN(DllApi) \
-	struct DllApi FFieldNotificationClassDescriptor : public ::UE::FieldNotification::IClassDescriptor \
+#define UE_FIELD_NOTIFICATION_DECLARE_CLASS_DESCRIPTOR_BASE_BEGIN() \
+	struct FFieldNotificationClassDescriptor : public ::UE::FieldNotification::IClassDescriptor \
 	{ \
 	private: \
 		using SuperDescriptor = ::UE::FieldNotification::IClassDescriptor; \
@@ -96,7 +100,11 @@
 			return *(AllFields[InFieldNumber]); \
 		} \
 	}; \
-	virtual const ::UE::FieldNotification::IClassDescriptor& GetFieldNotificationDescriptor() const override
+	virtual const ::UE::FieldNotification::IClassDescriptor& GetFieldNotificationDescriptor() const \
+	{ \
+ 		static FFieldNotificationClassDescriptor Instance; \
+ 		return Instance; \
+ 	}
 
 
 #define UE_FIELD_NOTIFICATION_DECLARE_ENUM_FIELD_BEGIN(Name) \
@@ -118,16 +126,16 @@
 			IndexOf_##Name,
 
 
-#define UE_FIELD_NOTIFICATION_DECLARE_CLASS_DESCRIPTOR_OneField(DllApi, Field1) \
-	UE_FIELD_NOTIFICATION_DECLARE_CLASS_DESCRIPTOR_BEGIN(DllApi) \
+#define UE_FIELD_NOTIFICATION_DECLARE_CLASS_DESCRIPTOR_OneField(Field1) \
+	UE_FIELD_NOTIFICATION_DECLARE_CLASS_DESCRIPTOR_BEGIN() \
 	UE_FIELD_NOTIFICATION_DECLARE_FIELD(Field1) \
 	UE_FIELD_NOTIFICATION_DECLARE_ENUM_FIELD_BEGIN(Field1) \
 	UE_FIELD_NOTIFICATION_DECLARE_ENUM_FIELD_END() \
 	UE_FIELD_NOTIFICATION_DECLARE_CLASS_DESCRIPTOR_END()
 
 
-#define UE_FIELD_NOTIFICATION_DECLARE_CLASS_DESCRIPTOR_TwoFields(DllApi, Field1, Field2) \
-	UE_FIELD_NOTIFICATION_DECLARE_CLASS_DESCRIPTOR_BEGIN(DllApi) \
+#define UE_FIELD_NOTIFICATION_DECLARE_CLASS_DESCRIPTOR_TwoFields(Field1, Field2) \
+	UE_FIELD_NOTIFICATION_DECLARE_CLASS_DESCRIPTOR_BEGIN() \
 	UE_FIELD_NOTIFICATION_DECLARE_FIELD(Field1) \
 	UE_FIELD_NOTIFICATION_DECLARE_FIELD(Field2) \
 	UE_FIELD_NOTIFICATION_DECLARE_ENUM_FIELD_BEGIN(Field1) \
@@ -136,8 +144,8 @@
 	UE_FIELD_NOTIFICATION_DECLARE_CLASS_DESCRIPTOR_END()
 
 
-#define UE_FIELD_NOTIFICATION_DECLARE_CLASS_DESCRIPTOR_ThreeFields(DllApi, Field1, Field2, Field3) \
-	UE_FIELD_NOTIFICATION_DECLARE_CLASS_DESCRIPTOR_BEGIN(DllApi) \
+#define UE_FIELD_NOTIFICATION_DECLARE_CLASS_DESCRIPTOR_ThreeFields(Field1, Field2, Field3) \
+	UE_FIELD_NOTIFICATION_DECLARE_CLASS_DESCRIPTOR_BEGIN() \
 	UE_FIELD_NOTIFICATION_DECLARE_FIELD(Field1) \
 	UE_FIELD_NOTIFICATION_DECLARE_FIELD(Field2) \
 	UE_FIELD_NOTIFICATION_DECLARE_FIELD(Field3) \
@@ -148,8 +156,8 @@
 	UE_FIELD_NOTIFICATION_DECLARE_CLASS_DESCRIPTOR_END()
 
 
-#define UE_FIELD_NOTIFICATION_DECLARE_CLASS_DESCRIPTOR_FourFields(DllApi, Field1, Field2, Field3, Field4) \
-	UE_FIELD_NOTIFICATION_DECLARE_CLASS_DESCRIPTOR_BEGIN(DllApi) \
+#define UE_FIELD_NOTIFICATION_DECLARE_CLASS_DESCRIPTOR_FourFields(Field1, Field2, Field3, Field4) \
+	UE_FIELD_NOTIFICATION_DECLARE_CLASS_DESCRIPTOR_BEGIN() \
 	UE_FIELD_NOTIFICATION_DECLARE_FIELD(Field1) \
 	UE_FIELD_NOTIFICATION_DECLARE_FIELD(Field2) \
 	UE_FIELD_NOTIFICATION_DECLARE_FIELD(Field3) \
@@ -162,8 +170,8 @@
 	UE_FIELD_NOTIFICATION_DECLARE_CLASS_DESCRIPTOR_END()
 
 
-#define UE_FIELD_NOTIFICATION_DECLARE_CLASS_DESCRIPTOR_FiveFields(DllApi, Field1, Field2, Field3, Field4, Field5) \
-	UE_FIELD_NOTIFICATION_DECLARE_CLASS_DESCRIPTOR_BEGIN(DllApi) \
+#define UE_FIELD_NOTIFICATION_DECLARE_CLASS_DESCRIPTOR_FiveFields(Field1, Field2, Field3, Field4, Field5) \
+	UE_FIELD_NOTIFICATION_DECLARE_CLASS_DESCRIPTOR_BEGIN() \
 	UE_FIELD_NOTIFICATION_DECLARE_FIELD(Field1) \
 	UE_FIELD_NOTIFICATION_DECLARE_FIELD(Field2) \
 	UE_FIELD_NOTIFICATION_DECLARE_FIELD(Field3) \
@@ -178,8 +186,8 @@
 	UE_FIELD_NOTIFICATION_DECLARE_CLASS_DESCRIPTOR_END()
 
 
-#define UE_FIELD_NOTIFICATION_DECLARE_CLASS_DESCRIPTOR_SixFields(DllApi, Field1, Field2, Field3, Field4, Field5, Field6) \
-	UE_FIELD_NOTIFICATION_DECLARE_CLASS_DESCRIPTOR_BEGIN(DllApi) \
+#define UE_FIELD_NOTIFICATION_DECLARE_CLASS_DESCRIPTOR_SixFields(Field1, Field2, Field3, Field4, Field5, Field6) \
+	UE_FIELD_NOTIFICATION_DECLARE_CLASS_DESCRIPTOR_BEGIN() \
 	UE_FIELD_NOTIFICATION_DECLARE_FIELD(Field1) \
 	UE_FIELD_NOTIFICATION_DECLARE_FIELD(Field2) \
 	UE_FIELD_NOTIFICATION_DECLARE_FIELD(Field3) \
@@ -196,8 +204,8 @@
 	UE_FIELD_NOTIFICATION_DECLARE_CLASS_DESCRIPTOR_END()
 
 
-#define UE_FIELD_NOTIFICATION_DECLARE_CLASS_DESCRIPTOR_SevenFields(DllApi, Field1, Field2, Field3, Field4, Field5, Field6, Field7) \
-	UE_FIELD_NOTIFICATION_DECLARE_CLASS_DESCRIPTOR_BEGIN(DllApi) \
+#define UE_FIELD_NOTIFICATION_DECLARE_CLASS_DESCRIPTOR_SevenFields(Field1, Field2, Field3, Field4, Field5, Field6, Field7) \
+	UE_FIELD_NOTIFICATION_DECLARE_CLASS_DESCRIPTOR_BEGIN() \
 	UE_FIELD_NOTIFICATION_DECLARE_FIELD(Field1) \
 	UE_FIELD_NOTIFICATION_DECLARE_FIELD(Field2) \
 	UE_FIELD_NOTIFICATION_DECLARE_FIELD(Field3) \
@@ -216,8 +224,8 @@
 	UE_FIELD_NOTIFICATION_DECLARE_CLASS_DESCRIPTOR_END()
 
 
-#define UE_FIELD_NOTIFICATION_DECLARE_CLASS_DESCRIPTOR_EightFields(DllApi, Field1, Field2, Field3, Field4, Field5, Field6, Field7, Field8) \
-	UE_FIELD_NOTIFICATION_DECLARE_CLASS_DESCRIPTOR_BEGIN(DllApi) \
+#define UE_FIELD_NOTIFICATION_DECLARE_CLASS_DESCRIPTOR_EightFields(Field1, Field2, Field3, Field4, Field5, Field6, Field7, Field8) \
+	UE_FIELD_NOTIFICATION_DECLARE_CLASS_DESCRIPTOR_BEGIN() \
 	UE_FIELD_NOTIFICATION_DECLARE_FIELD(Field1) \
 	UE_FIELD_NOTIFICATION_DECLARE_FIELD(Field2) \
 	UE_FIELD_NOTIFICATION_DECLARE_FIELD(Field3) \
@@ -242,38 +250,40 @@
 	const ::UE::FieldNotification::FFieldId* ClassName::FFieldNotificationClassDescriptor::AllFields[] = {
 
 
-#if DO_CHECK
-#define UE_FIELD_NOTIFICATION_IMPLEMENTATION_END(ClassName) \
-	}; \
-	const ::UE::FieldNotification::IClassDescriptor& ClassName::GetFieldNotificationDescriptor() const \
-	{ \
-		static_assert(UE_ARRAY_COUNT(ClassName::FFieldNotificationClassDescriptor::AllFields) == ClassName::FFieldNotificationClassDescriptor::Max_IndexOf_-ClassName::FFieldNotificationClassDescriptor::SuperDescriptor::Max_IndexOf_, "The descriptor for class " #ClassName " doesn't not implement the same number of field as declared in in the constructor."); \
-		struct FLocal \
-		{ \
-			ClassName::FFieldNotificationClassDescriptor Instance; \
-			FLocal() \
-			{ \
-				int32 NumberOfField = Instance.GetNumberOfField(); \
-				for (int32 Index = 0; Index < NumberOfField; ++Index) \
-				{ \
-					::UE::FieldNotification::FFieldId Id = Instance.GetField(Index); \
-					ensureAlwaysMsgf(Id.IsValid(), TEXT("The FFieldId '%d' for the class '%s' is invalid"), Index, *ClassName::StaticClass()->GetName()); \
-					ensureAlwaysMsgf(Id.GetIndex() == Index, TEXT("The FFieldId at index '%d' doesn't match the id index '%d'."), Index, Id.GetIndex()); \
-				} \
-			} \
-		}; \
-		static FLocal Local; \
-		return Local.Instance; \
-	}
-#else
-#define UE_FIELD_NOTIFICATION_IMPLEMENTATION_END(ClassName) \
-	}; \
-	const ::UE::FieldNotification::IClassDescriptor& ClassName::GetFieldNotificationDescriptor() const \
-	{ \
-		static FFieldNotificationClassDescriptor Instance; \
-		return Instance; \
-	}
-#endif
+#define UE_FIELD_NOTIFICATION_IMPLEMENTATION_END(ClassName) }
+
+//#if DO_CHECK
+//#define UE_FIELD_NOTIFICATION_IMPLEMENTATION_END(ClassName) \
+//	}; \
+//	const ::UE::FieldNotification::IClassDescriptor& ClassName::GetFieldNotificationDescriptor() const \
+//	{ \
+//		static_assert(UE_ARRAY_COUNT(ClassName::FFieldNotificationClassDescriptor::AllFields) == ClassName::FFieldNotificationClassDescriptor::Max_IndexOf_-ClassName::FFieldNotificationClassDescriptor::SuperDescriptor::Max_IndexOf_, "The descriptor for class " #ClassName " doesn't not implement the same number of field as declared in in the constructor."); \
+//		struct FLocal \
+//		{ \
+//			ClassName::FFieldNotificationClassDescriptor Instance; \
+//			FLocal() \
+//			{ \
+//				int32 NumberOfField = Instance.GetNumberOfField(); \
+//				for (int32 Index = 0; Index < NumberOfField; ++Index) \
+//				{ \
+//					::UE::FieldNotification::FFieldId Id = Instance.GetField(Index); \
+//					ensureAlwaysMsgf(Id.IsValid(), TEXT("The FFieldId '%d' for the class '%s' is invalid"), Index, *ClassName::StaticClass()->GetName()); \
+//					ensureAlwaysMsgf(Id.GetIndex() == Index, TEXT("The FFieldId at index '%d' doesn't match the id index '%d'."), Index, Id.GetIndex()); \
+//				} \
+//			} \
+//		}; \
+//		static FLocal Local; \
+//		return Local.Instance; \
+//	}
+//#else
+//#define UE_FIELD_NOTIFICATION_IMPLEMENTATION_END(ClassName) \
+//	}; \
+//	const ::UE::FieldNotification::IClassDescriptor& ClassName::GetFieldNotificationDescriptor() const \
+//	{ \
+//		static FFieldNotificationClassDescriptor Instance; \
+//		return Instance; \
+//	}
+//#endif
 
 
 #define UE_FIELD_NOTIFICATION_IMPLEMENT_FIELD(ClassName, MemberName) const ::UE::FieldNotification::FFieldId ClassName::FFieldNotificationClassDescriptor::MemberName( FName(TEXT(#MemberName)), ClassName::FFieldNotificationClassDescriptor::IndexOf_##MemberName );
