Fix missing include introduced with CL 19272412
#rnx

diff --git a/Engine/Source/Runtime/UMG/Public/FieldNotification/FieldVariant.h b/Engine/Source/Runtime/UMG/Public/FieldNotification/FieldVariant.h
--- a/Engine/Source/Runtime/UMG/Public/FieldNotification/FieldVariant.h
+++ b/Engine/Source/Runtime/UMG/Public/FieldNotification/FieldVariant.h
@@ -3,7 +3,9 @@
 #pragma once
 
 #include "CoreMinimal.h"
+#include "UObject/Class.h"
 #include "UObject/Field.h"
+#include "UObject/UnrealType.h"
 
 namespace UE::FieldNotification
 {
