FieldNotification: Fix crash with bad indexing.
#rnx

diff --git a/Engine/Source/Runtime/UMG/Private/FieldNotification/FieldMulticastDelegate.cpp b/Engine/Source/Runtime/UMG/Private/FieldNotification/FieldMulticastDelegate.cpp
--- a/Engine/Source/Runtime/UMG/Private/FieldNotification/FieldMulticastDelegate.cpp
+++ b/Engine/Source/Runtime/UMG/Private/FieldNotification/FieldMulticastDelegate.cpp
@@ -162,7 +162,7 @@ FFieldMulticastDelegate::FRemoveAllResult FFieldMulticastDelegate::RemoveAll(con
 
 	auto SetHasFields = [&Result, InObject](FInvocationElement& Element)
 	{
-		Result.HasFields.PadToNum(Element.Key.Id.GetIndex(), false);
+		Result.HasFields.PadToNum(Element.Key.Id.GetIndex() + 1, false);
 		Result.HasFields[Element.Key.Id.GetIndex()] = true;
 	};
 
@@ -224,7 +224,7 @@ FFieldMulticastDelegate::FRemoveAllResult FFieldMulticastDelegate::RemoveAll(con
 
 	auto SetHasFields = [&Result, InObject](FInvocationElement& Element)
 	{
-		Result.HasFields.PadToNum(Element.Key.Id.GetIndex(), false);
+		Result.HasFields.PadToNum(Element.Key.Id.GetIndex() + 1, false);
 		Result.HasFields[Element.Key.Id.GetIndex()] = true;
 	};
 
