UMG: Move when the Initialize and Destroyed to be closer to the BP callback. It was previously closer to where the animation code was doing it since I've used it as a template and test case.

diff --git a/Engine/Source/Runtime/UMG/Private/UserWidget.cpp b/Engine/Source/Runtime/UMG/Private/UserWidget.cpp
--- a/Engine/Source/Runtime/UMG/Private/UserWidget.cpp
+++ b/Engine/Source/Runtime/UMG/Private/UserWidget.cpp
@@ -154,11 +154,6 @@ bool UUserWidget::Initialize()
 			InitializeNamedSlots(bReparentToWidgetTree);
 		}
 
-		for (UUserWidgetExtension* Extension : Extensions)
-		{
-			Extension->Initialize();
-		}
-
 		if (!IsDesignTime() && PlayerContext.IsValid())
 		{
 			NativeOnInitialized();
@@ -234,11 +229,6 @@ void UUserWidget::BeginDestroy()
 {
 	Super::BeginDestroy();
 
-	for (UUserWidgetExtension* Extension : Extensions)
-	{
-		Extension->Destruct();
-	}
-
 	TearDownAnimations();
 
 	if (AnimationTickManager)
@@ -1369,6 +1359,11 @@ void UUserWidget::NativeOnInitialized()
 			});
 	}
 
+	for (UUserWidgetExtension* Extension : Extensions)
+	{
+		Extension->Initialize();
+	}
+
 	OnInitialized();
 }
 
@@ -1402,16 +1397,21 @@ void UUserWidget::NativeConstruct()
 
 void UUserWidget::NativeDestruct()
 {
+	StopListeningForAllInputActions();
+	Destruct();
+
+	for (UUserWidgetExtension* Extension : Extensions)
+	{
+		Extension->Destruct();
+	}
+
 	if (UWidgetBlueprintGeneratedClass* Class = GetWidgetTreeOwningClass())
 	{
 		Class->ForEachExtension([this](UWidgetBlueprintGeneratedClassExtension* Extension)
-		{
-			Extension->Destruct(this);
-		});
+			{
+				Extension->Destruct(this);
+			});
 	}
-
-	StopListeningForAllInputActions();
-	Destruct();
 }
 
 void UUserWidget::NativeTick(const FGeometry& MyGeometry, float InDeltaTime)
@@ -2123,12 +2123,22 @@ UUserWidgetExtension* UUserWidget::AddExtension(TSubclassOf<UUserWidgetExtension
 {
 	UUserWidgetExtension* Extension = NewObject<UUserWidgetExtension>(this, InExtensionType);
 	Extensions.Add(Extension);
+	if (bInitialized)
+	{
+		Extension->Initialize();
+	}
 	return Extension;
 }
 
 void UUserWidget::RemoveExtension(UUserWidgetExtension* InExtension)
 {
-	Extensions.RemoveSingle(InExtension);
+	if (InExtension)
+	{
+		if (Extensions.RemoveSingleSwap(InExtension))
+		{
+			InExtension->Destruct();
+		}
+	}
 }
 
 
