UMG: Expose RemoveAllFieldValueChangedDelegates to IFieldValueChanged.

diff --git a/Engine/Source/Runtime/UMG/Private/Binding/WidgetFieldNotificationExtension.h b/Engine/Source/Runtime/UMG/Private/Binding/WidgetFieldNotificationExtension.h
--- a/Engine/Source/Runtime/UMG/Private/Binding/WidgetFieldNotificationExtension.h
+++ b/Engine/Source/Runtime/UMG/Private/Binding/WidgetFieldNotificationExtension.h
@@ -38,6 +38,11 @@ public:
 		return Delegates.RemoveAll(InObject, InUserObject);
 	}
 
+	FRemoveAllResult RemoveAllFieldValueChangedDelegates(const UVisual* InObject, FFieldId InFieldId, const void* InUserObject)
+	{
+		return Delegates.RemoveAll(InObject, InFieldId, InUserObject);
+	}
+
 	void BroadcastFieldValueChanged(UVisual* InObject, UE::FieldNotification::FFieldId InFieldId)
 	{
 		Delegates.Broadcast(InObject, InFieldId);
diff --git a/Engine/Source/Runtime/UMG/Private/Components/Widget.cpp b/Engine/Source/Runtime/UMG/Private/Components/Widget.cpp
--- a/Engine/Source/Runtime/UMG/Private/Components/Widget.cpp
+++ b/Engine/Source/Runtime/UMG/Private/Components/Widget.cpp
@@ -1719,6 +1719,21 @@ int32 UWidget::RemoveAllFieldValueChangedDelegates(const void* InUserObject)
 	return bResult;
 }
 
+int32 UWidget::RemoveAllFieldValueChangedDelegates(UE::FieldNotification::FFieldId InFieldId, const void* InUserObject)
+{
+	int32 bResult = 0;
+	if (InUserObject)
+	{
+		if (UWidgetFieldNotificationExtension* Extension = UE::UMG::Private::FindOrAddWidgetNotifyExtension(this))
+		{
+			UWidgetFieldNotificationExtension::FRemoveAllResult RemoveResult = Extension->RemoveAllFieldValueChangedDelegates(this, InFieldId, InUserObject);
+			bResult = RemoveResult.RemoveCount;
+			EnabledFieldNotifications = RemoveResult.HasFields;
+		}
+	}
+	return bResult;
+}
+
 void UWidget::BroadcastFieldValueChanged(UE::FieldNotification::FFieldId InFieldId)
 {
 	if (InFieldId.IsValid() && EnabledFieldNotifications.IsValidIndex(InFieldId.GetIndex()) && EnabledFieldNotifications[InFieldId.GetIndex()])
diff --git a/Engine/Source/Runtime/UMG/Public/Components/Widget.h b/Engine/Source/Runtime/UMG/Public/Components/Widget.h
--- a/Engine/Source/Runtime/UMG/Public/Components/Widget.h
+++ b/Engine/Source/Runtime/UMG/Public/Components/Widget.h
@@ -724,6 +724,7 @@ public:
 	virtual FDelegateHandle AddFieldValueChangedDelegate(UE::FieldNotification::FFieldId InFieldId, FFieldValueChangedDelegate InNewDelegate) override final;
 	virtual bool RemoveFieldValueChangedDelegate(UE::FieldNotification::FFieldId InFieldId, FDelegateHandle InHandle) override final;
 	virtual int32 RemoveAllFieldValueChangedDelegates(const void* InUserObject) override final;
+	virtual int32 RemoveAllFieldValueChangedDelegates(UE::FieldNotification::FFieldId InFieldId, const void* InUserObject) override final;
 	//~ End INotifyFieldValueChanged Interface
 
 public:
diff --git a/Engine/Source/Runtime/UMG/Public/FieldNotification/IFieldValueChanged.h b/Engine/Source/Runtime/UMG/Public/FieldNotification/IFieldValueChanged.h
--- a/Engine/Source/Runtime/UMG/Public/FieldNotification/IFieldValueChanged.h
+++ b/Engine/Source/Runtime/UMG/Public/FieldNotification/IFieldValueChanged.h
@@ -27,14 +27,17 @@ public:
 
 public:
 	/** Add a delegate that will be notified when the FieldId is value changed. */
-	virtual FDelegateHandle AddFieldValueChangedDelegate(UE::FieldNotification::FFieldId FieldId, FFieldValueChangedDelegate InNewDelegate) = 0;
+	virtual FDelegateHandle AddFieldValueChangedDelegate(UE::FieldNotification::FFieldId InFieldId, FFieldValueChangedDelegate InNewDelegate) = 0;
 
 	/** Remove a delegate that was added. */
-	virtual bool RemoveFieldValueChangedDelegate(UE::FieldNotification::FFieldId FieldId, FDelegateHandle Handle) = 0;
+	virtual bool RemoveFieldValueChangedDelegate(UE::FieldNotification::FFieldId InFieldId, FDelegateHandle InHandle) = 0;
 
 	/** Remove all the delegate that are bound to the specified UserObject. */
 	virtual int32 RemoveAllFieldValueChangedDelegates(const void* InUserObject) = 0;
 
+	/** Remove all the delegate that are bound to the specified Field and UserObject. */
+	virtual int32 RemoveAllFieldValueChangedDelegates(UE::FieldNotification::FFieldId InFieldId, const void* InUserObject) = 0;
+
 	/** @returns the list of all the field that can notify when their value changes. */
 	virtual const UE::FieldNotification::IClassDescriptor& GetFieldNotificationDescriptor() const = 0;
 };
