Save 8 bytes each on UWidget and USoundNodeWavePlayer
#rnx

diff --git a/Engine/Source/Runtime/Engine/Classes/Sound/SoundNodeWavePlayer.h b/Engine/Source/Runtime/Engine/Classes/Sound/SoundNodeWavePlayer.h
--- a/Engine/Source/Runtime/Engine/Classes/Sound/SoundNodeWavePlayer.h
+++ b/Engine/Source/Runtime/Engine/Classes/Sound/SoundNodeWavePlayer.h
@@ -35,6 +35,12 @@ private:
 
 	uint8 bAsyncLoading:1;
 
+public:
+
+	UPROPERTY(EditAnywhere, Category=WavePlayer)
+	uint8 bLooping:1;
+
+private:
 	// Set to true when we enqueue a task to the game thread to load
 	// the asset associated with this thread.
 	// This only occurs if ClearAssetReferences() was called and we still tried to play this wave player,
@@ -43,9 +49,6 @@ private:
 
 public:
 
-	UPROPERTY(EditAnywhere, Category=WavePlayer)
-	uint32 bLooping:1;
-
 	ENGINE_API USoundWave* GetSoundWave() const { return SoundWave; }
 	ENGINE_API void SetSoundWave(USoundWave* SoundWave);
 
diff --git a/Engine/Source/Runtime/UMG/Public/Components/Widget.h b/Engine/Source/Runtime/UMG/Public/Components/Widget.h
--- a/Engine/Source/Runtime/UMG/Public/Components/Widget.h
+++ b/Engine/Source/Runtime/UMG/Public/Components/Widget.h
@@ -307,6 +307,10 @@ public:
 	UPROPERTY(EditAnywhere, BlueprintReadOnly, Category="Render Transform", meta=( DisplayName="Pivot" ))
 	FVector2D RenderTransformPivot;
 
+	/** Allows you to set a new flow direction */
+	UPROPERTY(EditAnywhere, Category = "Localization")
+	EFlowDirectionPreference FlowDirectionPreference;
+
 	/**
 	 * Allows controls to be exposed as variables in a blueprint.  Not all controls need to be exposed
 	 * as variables, so this allows only the most useful ones to end up being exposed.
@@ -326,10 +330,6 @@ public:
 	UPROPERTY(EditAnywhere, Category="Behavior", meta=(InlineEditConditionToggle))
 	uint8 bOverride_Cursor : 1;
 
-	/** Allows you to set a new flow direction */
-	UPROPERTY(EditAnywhere, Category = "Localization")
-	EFlowDirectionPreference FlowDirectionPreference;
-
 #if WITH_EDITORONLY_DATA
 	// These editor-only properties exist for two reasons:
 	//   1. To make details customization easier to write, specifically in regards to the binding extension widget
@@ -376,11 +376,6 @@ public:
 	USlateAccessibleWidgetData::FGetText AccessibleSummaryTextDelegate;
 #endif
 
-private:
-	/** A custom set of accessibility rules for this widget. If null, default rules for the widget are used. */
-	UPROPERTY(Instanced)
-	TObjectPtr<USlateAccessibleWidgetData> AccessibleWidgetData;
-
 protected:
 
 	/**
@@ -431,6 +426,12 @@ public:
 	UPROPERTY(EditAnywhere, Category="Behavior")
 	float RenderOpacity;
 
+private:
+	/** A custom set of accessibility rules for this widget. If null, default rules for the widget are used. */
+	UPROPERTY(Instanced)
+	TObjectPtr<USlateAccessibleWidgetData> AccessibleWidgetData;
+
+public:
 	/**
 	 * The navigation object for this widget is optionally created if the user has configured custom
 	 * navigation rules for this widget in the widget designer.  Those rules determine how navigation transitions
