MVVM: Add an helper function to find if a UFunction can be valid for FieldNotify.

diff --git a/Engine/Source/Runtime/UMG/Public/FieldNotification/FieldNotificationHelpers.h b/Engine/Source/Runtime/UMG/Public/FieldNotification/FieldNotificationHelpers.h
new file mode 100644
--- /dev/null
+++ b/Engine/Source/Runtime/UMG/Public/FieldNotification/FieldNotificationHelpers.h
@@ -0,0 +1,20 @@
+// Copyright Epic Games, Inc. All Rights Reserved.
+
+#pragma once
+
+#include "CoreMinimal.h"
+#include "UObject/Class.h"
+
+namespace UE::FieldNotification::Helpers
+{
+
+	UE_NODISCARD bool IsValidAsField(const UFunction* InFunction)
+	{
+		return InFunction != nullptr
+			&& !InFunction->HasAnyFunctionFlags(FUNC_Net | FUNC_Event)
+			&& InFunction->HasAllFunctionFlags(FUNC_BlueprintCallable | FUNC_Const)
+			&& InFunction->NumParms == 1
+			&& InFunction->GetReturnProperty() != nullptr;
+	}
+
+} //namespace
