FieldNotification: Add support for dynamic delegate.

diff --git a/Engine/Source/Runtime/UMG/Private/Binding/WidgetFieldNotificationExtension.h b/Engine/Source/Runtime/UMG/Private/Binding/WidgetFieldNotificationExtension.h
--- a/Engine/Source/Runtime/UMG/Private/Binding/WidgetFieldNotificationExtension.h
+++ b/Engine/Source/Runtime/UMG/Private/Binding/WidgetFieldNotificationExtension.h
@@ -28,11 +28,21 @@ public:
 		return Delegates.Add(InObject, InFieldId, MoveTemp(InNewDelegate));
 	}
 
+	FDelegateHandle AddFieldValueChangedDelegate(const UVisual* InObject, FFieldId InFieldId, const FFieldValueChangedDynamicDelegate& InDelegate)
+	{
+		return Delegates.Add(InObject, InFieldId, InDelegate);
+	}
+
 	FRemoveFromResult RemoveFieldValueChangedDelegate(const UVisual* InObject, FFieldId InFieldId, FDelegateHandle InDelegateHandle)
 	{
 		return Delegates.RemoveFrom(InObject, InFieldId, InDelegateHandle);
 	}
 
+	FRemoveFromResult RemoveFieldValueChangedDelegate(const UVisual* InObject, FFieldId InFieldId, const FFieldValueChangedDynamicDelegate& InDelegate)
+	{
+		return Delegates.RemoveFrom(InObject, InFieldId, InDelegate);
+	}
+
 	FRemoveAllResult RemoveAllFieldValueChangedDelegates(const UVisual* InObject, const void* InUserObject)
 	{
 		return Delegates.RemoveAll(InObject, InUserObject);
diff --git a/Engine/Source/Runtime/UMG/Private/FieldNotification/FieldMulticastDelegate.cpp b/Engine/Source/Runtime/UMG/Private/FieldNotification/FieldMulticastDelegate.cpp
--- a/Engine/Source/Runtime/UMG/Private/FieldNotification/FieldMulticastDelegate.cpp
+++ b/Engine/Source/Runtime/UMG/Private/FieldNotification/FieldMulticastDelegate.cpp
@@ -8,8 +8,13 @@
 
 FDelegateHandle FFieldMulticastDelegate::Add(const UObject* InObject, FFieldId InFieldId, FDelegate InNewDelegate)
 {
+	if (!InNewDelegate.IsBound())
+	{
+		return FDelegateHandle();
+	}
+
 	FDelegateHandle Result = InNewDelegate.GetHandle();
-	FInvocationKey InvocationKey{ InObject, InFieldId };
+	FInvocationKey InvocationKey{ InObject, InFieldId, NAME_None };
 	if (DelegateLockCount > 0)
 	{
 		const int32 Index = Delegates.Emplace(FInvocationElement{ InvocationKey, MoveTemp(InNewDelegate) });
@@ -24,31 +29,49 @@
 }
 
 
-FFieldMulticastDelegate::FRemoveResult FFieldMulticastDelegate::Remove(FDelegateHandle InDelegate)
+FDelegateHandle FFieldMulticastDelegate::Add(const UObject* InObject, FFieldId InFieldId, const FDynamicDelegate& InDynamicDelegate)
 {
-	FRemoveResult Result;
-
-	for (int32 Index = Delegates.Num() - 1; Index >= 0; --Index)
+	if (!InDynamicDelegate.IsBound())
 	{
-		FInvocationElement& Element = Delegates[Index];
-		if (Element.Delegate.GetHandle() == InDelegate)
-		{
-			Result.Object = Element.Key.Object.Get();
-			Result.FieldId = Element.Key.Id;
-			Result.bRemoved = Element.Delegate.IsBound();
-			if (DelegateLockCount > 0)
-			{
-				Element.Delegate.Unbind();
-				++CompactionCount;
-			}
-			else
-			{
-				Delegates.RemoveAt(Index);
-			}
-			break;
-		}
+		return FDelegateHandle();
 	}
 
+	FDelegate NewDelegate = FDelegate::CreateUFunction(const_cast<UObject*>(InDynamicDelegate.GetUObject()), InDynamicDelegate.GetFunctionName());
+	FDelegateHandle Result = NewDelegate.GetHandle();
+	FInvocationKey InvocationKey{ InObject, InFieldId, InDynamicDelegate.GetFunctionName()};
+	if (DelegateLockCount > 0)
+	{
+		const int32 Index = Delegates.Emplace(FInvocationElement{ InvocationKey, MoveTemp(NewDelegate) });
+		AddedEmplaceAt = FMath::Min(AddedEmplaceAt, (uint16)Index);
+	}
+	else
+	{
+		const int32 FoundIndex = UpperBound(InvocationKey);
+		Delegates.EmplaceAt(FoundIndex, FInvocationElement{ InvocationKey, MoveTemp(NewDelegate) });
+	}
+	return Result;
+}
+
+
+void FFieldMulticastDelegate::RemoveElement(FInvocationElement& Element, int32 Index, FRemoveResult& Result)
+{
+	Result.Object = Element.Key.Object.Get();
+	Result.FieldId = Element.Key.Id;
+	Result.bRemoved = Element.Delegate.IsBound();
+	if (DelegateLockCount > 0)
+	{
+		Element.Delegate.Unbind();
+		++CompactionCount;
+	}
+	else
+	{
+		Delegates.RemoveAt(Index);
+	}
+}
+
+
+void FFieldMulticastDelegate::CompleteRemove(FRemoveResult& Result) const
+{
 	if (Result.FieldId.IsValid())
 	{
 		// Search in the normal sorted list
@@ -83,43 +106,72 @@
 			}
 		}
 	}
-
-	return Result;
 }
 
 
-FFieldMulticastDelegate::FRemoveFromResult FFieldMulticastDelegate::RemoveFrom(const UObject* InObject, FFieldId InFieldId, FDelegateHandle InDelegate)
+FFieldMulticastDelegate::FRemoveResult FFieldMulticastDelegate::Remove(FDelegateHandle InDelegate)
 {
-	bool bRemoved = false;
-	bool bFieldPresent = false;
+	FRemoveResult Result;
 
-	FFieldMulticastDelegate* Self = this;
-	auto RemoveOrUnbind = [Self, &InDelegate, &bRemoved, &bFieldPresent, InObject](FInvocationElement& Element, int32 Index)
+	if (!InDelegate.IsValid())
 	{
+		return Result;
+	}
+
+	for (int32 Index = Delegates.Num() - 1; Index >= 0; --Index)
+	{
+		FInvocationElement& Element = Delegates[Index];
 		if (Element.Delegate.GetHandle() == InDelegate)
 		{
-			if (Self->DelegateLockCount > 0)
+			RemoveElement(Element, Index, Result);
+			break;
+		}
+	}
+
+	CompleteRemove(Result);
+
+	return Result;
+}
+
+
+FFieldMulticastDelegate::FRemoveResult FFieldMulticastDelegate::Remove(const FDynamicDelegate& InDynamicDelegate)
+{
+	FRemoveResult Result;
+
+	if (!InDynamicDelegate.IsBound())
+	{
+		// The UObject may be GCed. We would like to compact the list but there is not way to know which one (since HasSameObject will return false).
+		//Result may indicate that it was not remove even if it will be removed later when we finally compact.
+		return Result;
+	}
+
+	for (int32 Index = Delegates.Num() - 1; Index >= 0; --Index)
+	{
+		FInvocationElement& Element = Delegates[Index];
+		if (Element.Key.DynamicName == InDynamicDelegate.GetFunctionName())
+		{
+			if (const IDelegateInstance* DelegateInstance = GetDelegateInstance(Element.Delegate))
 			{
-				Element.Delegate.Unbind();
-				++Self->CompactionCount;
+				if (DelegateInstance->HasSameObject(InDynamicDelegate.GetUObject()))
+				{
+					RemoveElement(Element, Index, Result);
+					break;
+				}
 			}
-			else
-			{
-				Self->Delegates.RemoveAt(Index);
-			}
-			bRemoved = true;
-			return !bFieldPresent;
 		}
-		else if (Element.Key.Object == InObject && Element.Delegate.IsBound())
-		{
-			bFieldPresent = true;
-			return !bRemoved;
-		}
-		return true;
-	};
+	}
 
+	CompleteRemove(Result);
+
+	return Result;
+}
+
+
+template<typename TRemoveOrUnbind>
+void FFieldMulticastDelegate::RemoveFrom(TRemoveOrUnbind& RemoveOrUnbind, FFieldId InFieldId, bool& bFieldPresent)
+{
 	// Search in the normal sorted list
-	const int32 FoundIndex = UpperBound(InFieldId)-1;
+	const int32 FoundIndex = UpperBound(InFieldId) - 1;
 	if (Delegates.IsValidIndex(FoundIndex))
 	{
 		for (int32 Index = FoundIndex; Index >= 0; --Index)
@@ -151,6 +203,92 @@
 			}
 		}
 	}
+}
+
+
+FFieldMulticastDelegate::FRemoveFromResult FFieldMulticastDelegate::RemoveFrom(const UObject* InObject, FFieldId InFieldId, FDelegateHandle InDelegate)
+{
+	bool bRemoved = false;
+	bool bFieldPresent = false;
+
+	if (!InDelegate.IsValid())
+	{
+		return FRemoveFromResult{ bRemoved, bFieldPresent };
+	}
+
+	FFieldMulticastDelegate* Self = this;
+	auto RemoveOrUnbind = [Self, &InDelegate, &bRemoved, &bFieldPresent, InObject](FInvocationElement& Element, int32 Index)
+	{
+		if (Element.Delegate.GetHandle() == InDelegate)
+		{
+			if (Self->DelegateLockCount > 0)
+			{
+				Element.Delegate.Unbind();
+				++Self->CompactionCount;
+			}
+			else
+			{
+				Self->Delegates.RemoveAt(Index);
+			}
+			bRemoved = true;
+			return !bFieldPresent;
+		}
+		if (Element.Key.Object == InObject && Element.Delegate.IsBound())
+		{
+			bFieldPresent = true;
+			return !bRemoved;
+		}
+		return true;
+	};
+
+	RemoveFrom(RemoveOrUnbind, InFieldId, bFieldPresent);
 
 	return FRemoveFromResult{ bRemoved, bFieldPresent };
 }
+
+
+FFieldMulticastDelegate::FRemoveFromResult FFieldMulticastDelegate::RemoveFrom(const UObject* InObject, FFieldId InFieldId, const FDynamicDelegate& InDynamicDelegate)
+{
+	bool bRemoved = false;
+	bool bFieldPresent = false;
+
+	if (!InDynamicDelegate.IsBound())
+	{
+		return FRemoveFromResult{ bRemoved, bFieldPresent };
+	}
+
+	FFieldMulticastDelegate* Self = this;
+	auto RemoveOrUnbind = [Self, &InDynamicDelegate, &bRemoved, &bFieldPresent, InObject](FInvocationElement& Element, int32 Index)
+	{
+		if (Element.Key.DynamicName == InDynamicDelegate.GetFunctionName())
+		{
+			if (const IDelegateInstance* DelegateInstance = Self->GetDelegateInstance(Element.Delegate))
+			{
+				if (DelegateInstance->HasSameObject(InDynamicDelegate.GetUObject()))
+				{
+					if (Self->DelegateLockCount > 0)
+					{
+						Element.Delegate.Unbind();
+						++Self->CompactionCount;
+					}
+					else
+					{
+						Self->Delegates.RemoveAt(Index);
+					}
+					bRemoved = true;
+					return !bFieldPresent;
+				}
+			}
+		}
+		if (Element.Key.Object == InObject && Element.Delegate.IsBound())
+		{
+			bFieldPresent = true;
+			return !bRemoved;
+		}
+		return true;
+	};
+
+	RemoveFrom(RemoveOrUnbind, InFieldId, bFieldPresent);
+
+	return FRemoveFromResult{ bRemoved, bFieldPresent };
+}
diff --git a/Engine/Source/Runtime/UMG/Public/FieldNotification/FieldId.h b/Engine/Source/Runtime/UMG/Public/FieldNotification/FieldId.h
--- a/Engine/Source/Runtime/UMG/Public/FieldNotification/FieldId.h
+++ b/Engine/Source/Runtime/UMG/Public/FieldNotification/FieldId.h
@@ -6,6 +6,8 @@
 #include "FieldNotification/FieldVariant.h"
 #include "UObject/Object.h"
 
+#include "FieldId.generated.h"
+
 
 namespace UE::FieldNotification
 {
@@ -76,4 +78,50 @@ private:
 	int32 BitNumber;
 };
 
-} //namespace
\ No newline at end of file
+} //namespace
+
+
+USTRUCT(BlueprintType)
+struct FFieldNotificationId
+{
+	GENERATED_BODY()
+
+public:
+	FFieldNotificationId() = default;
+
+	explicit FFieldNotificationId(FName InFieldName)
+		: FieldName(InFieldName)
+	{
+
+	}
+
+public:
+	bool IsValid() const
+	{
+		return FieldName.IsNone();
+	}
+
+	FName GetFieldName() const
+	{
+		return FieldName;
+	}
+
+	bool operator== (const FFieldNotificationId& Other) const
+	{
+		return FieldName == Other.FieldName;
+	}
+
+	bool operator!= (const FFieldNotificationId& Other) const
+	{
+		return FieldName != Other.FieldName;
+	}
+
+	friend int32 GetTypeHash(const FFieldNotificationId& Value)
+	{
+		return GetTypeHash(Value.FieldName);
+	}
+
+public:
+	UPROPERTY(BlueprintReadWrite, EditAnywhere, Category = "FieldNotify")
+	FName FieldName;
+};
diff --git a/Engine/Source/Runtime/UMG/Public/FieldNotification/FieldMulticastDelegate.h b/Engine/Source/Runtime/UMG/Public/FieldNotification/FieldMulticastDelegate.h
--- a/Engine/Source/Runtime/UMG/Public/FieldNotification/FieldMulticastDelegate.h
+++ b/Engine/Source/Runtime/UMG/Public/FieldNotification/FieldMulticastDelegate.h
@@ -16,6 +16,7 @@ class UMG_API FFieldMulticastDelegate
 {
 public:
 	using FDelegate = INotifyFieldValueChanged::FFieldValueChangedDelegate;
+	using FDynamicDelegate = FFieldValueChangedDynamicDelegate;
 
 private:
 	struct FInstanceProtectedDelegate : public TDelegateBase<FDefaultDelegateUserPolicy>
@@ -28,6 +29,7 @@ private:
 	{
 		TWeakObjectPtr<const UObject> Object;
 		FFieldId Id;
+		FName DynamicName;
 		bool operator<(const FInvocationKey& Element) const
 		{
 			return Id.GetName().FastLess(Element.Id.GetName());
@@ -65,6 +67,7 @@ private:
 
 public:
 	FDelegateHandle Add(const UObject* InObject, FFieldId InFieldId, FDelegate InNewDelegate);
+	FDelegateHandle Add(const UObject* InObject, FFieldId InFieldId, const FDynamicDelegate& InDynamicDelegate);
 
 	struct FRemoveResult
 	{
@@ -74,6 +77,7 @@ public:
 		FFieldId FieldId;
 	};
 	FRemoveResult Remove(FDelegateHandle InDelegate);
+	FRemoveResult Remove(const FDynamicDelegate& InDynamicDelegate);
 
 	struct FRemoveFromResult
 	{
@@ -81,6 +85,7 @@ public:
 		bool bHasOtherBoundDelegates = false;
 	};
 	FRemoveFromResult RemoveFrom(const UObject* InObject, FFieldId InFieldId, FDelegateHandle InDelegate);
+	FRemoveFromResult RemoveFrom(const UObject* InObject, FFieldId InFieldId, const FDynamicDelegate& InDynamicDelegate);
 
 	struct FRemoveAllResult
 	{
@@ -125,6 +130,10 @@ private:
 	}
 
 	void ExecuteLockOperations();
+	void RemoveElement(FInvocationElement& Element, int32 Index, FRemoveResult& Result);
+	void CompleteRemove(FRemoveResult& Result) const;
+	template<typename TRemoveOrUnbind>
+	void RemoveFrom(TRemoveOrUnbind& RemoveOrUnbind, FFieldId FieldId, bool& bFieldPresent);
 
 private:
 	InvocationListType Delegates;
diff --git a/Engine/Source/Runtime/UMG/Public/FieldNotification/IFieldValueChanged.h b/Engine/Source/Runtime/UMG/Public/FieldNotification/IFieldValueChanged.h
--- a/Engine/Source/Runtime/UMG/Public/FieldNotification/IFieldValueChanged.h
+++ b/Engine/Source/Runtime/UMG/Public/FieldNotification/IFieldValueChanged.h
@@ -10,6 +10,9 @@
 #include "IFieldValueChanged.generated.h"
 
 
+DECLARE_DYNAMIC_DELEGATE_TwoParams(FFieldValueChangedDynamicDelegate, UObject*, Object, FFieldNotificationId, Field);
+
+
 UINTERFACE(MinimalAPI, NotBlueprintable)
 class UNotifyFieldValueChanged : public UInterface
 {
