UMG: Modify the FieldNotification Descriptor to support Field generated in BP.

diff --git a/Engine/Source/Runtime/UMG/Private/Components/Widget.cpp b/Engine/Source/Runtime/UMG/Private/Components/Widget.cpp
--- a/Engine/Source/Runtime/UMG/Private/Components/Widget.cpp
+++ b/Engine/Source/Runtime/UMG/Private/Components/Widget.cpp
@@ -23,6 +23,7 @@
 #include "Logging/MessageLog.h"
 #include "Blueprint/UserWidget.h"
 #include "Blueprint/UserWidgetBlueprint.h"
+#include "Blueprint/WidgetBlueprintGeneratedClass.h"
 #include "Slate/SObjectWidget.h"
 #include "Blueprint/WidgetTree.h"
 #include "UMGStyle.h"
@@ -155,6 +156,20 @@ TArray<TSubclassOf<UPropertyBinding>> UWidget::BinderClasses;
 
 
 /////////////////////////////////////////////////////
+void UWidget::FFieldNotificationClassDescriptor::ForEachField(const UClass* Class, TFunctionRef<bool(::UE::FieldNotification::FFieldId FielId)> Callback) const
+{
+	for (int32 Index = 0; Index < Max_IndexOf_; ++Index)
+	{
+		if (!Callback(*AllFields[Index]))
+		{
+			return;
+		}
+	}
+	if (const UWidgetBlueprintGeneratedClass* WidgetBPClass = Cast<const UWidgetBlueprintGeneratedClass>(Class))
+	{
+		WidgetBPClass->ForEachField(Callback);
+	}
+}
 UE_FIELD_NOTIFICATION_IMPLEMENT_CLASS_DESCRIPTOR_ThreeFields(UWidget, ToolTipText, Visibility, bIsEnabled);
 
 
@@ -1670,6 +1685,27 @@ namespace UE::UMG::Private
 		}
 		return nullptr;
 	}
+	UWidgetFieldNotificationExtension* FindWidgetNotifyExtension(const UWidget* Widget)
+	{
+		if (const UUserWidget* UserWidget = Cast<const UUserWidget>(Widget))
+		{
+			if (UWidgetFieldNotificationExtension* Extension = UserWidget->GetExtension<UWidgetFieldNotificationExtension>())
+			{
+				return Extension;
+			}
+		}
+		else if (const UWidgetTree* WidgetTree = Cast<const UWidgetTree>(Widget->GetOuter()))
+		{
+			if (const UUserWidget* InnerUserWidget = Cast<const UUserWidget>(WidgetTree->GetOuter()))
+			{
+				if (UWidgetFieldNotificationExtension* Extension = InnerUserWidget->GetExtension<UWidgetFieldNotificationExtension>())
+				{
+					return Extension;
+				}
+			}
+		}
+		return nullptr;
+	}
 }
 
 FDelegateHandle UWidget::AddFieldValueChangedDelegate(UE::FieldNotification::FFieldId InFieldId, FFieldValueChangedDelegate InNewDelegate)
@@ -1690,12 +1726,31 @@ FDelegateHandle UWidget::AddFieldValueChangedDelegate(UE::FieldNotification::FFi
 	return Result;
 }
 
+void UWidget::K2_AddFieldValueChangedDelegate(FFieldNotificationId InFieldId, FFieldValueChangedDynamicDelegate InDelegate)
+{
+	if (InFieldId.IsValid())
+	{
+		const UE::FieldNotification::FFieldId FieldId = GetFieldNotificationDescriptor().GetField(GetClass(), InFieldId.FieldName);
+		if (ensureMsgf(FieldId.IsValid(), TEXT("The field should be compiled correctly.")))
+		{
+			if (UWidgetFieldNotificationExtension* Extension = UE::UMG::Private::FindOrAddWidgetNotifyExtension(this))
+			{
+				if (Extension->AddFieldValueChangedDelegate(this, FieldId, InDelegate).IsValid())
+				{
+					EnabledFieldNotifications.PadToNum(FieldId.GetIndex() + 1, false);
+					EnabledFieldNotifications[FieldId.GetIndex()] = true;
+				}
+			}
+		}
+	}
+}
+
 bool UWidget::RemoveFieldValueChangedDelegate(UE::FieldNotification::FFieldId InFieldId, FDelegateHandle InHandle)
 {
 	bool bResult = false;
 	if (InFieldId.IsValid() && InHandle.IsValid() && EnabledFieldNotifications.IsValidIndex(InFieldId.GetIndex()) && EnabledFieldNotifications[InFieldId.GetIndex()])
 	{
-		UWidgetFieldNotificationExtension* Extension = UE::UMG::Private::FindOrAddWidgetNotifyExtension(this);
+		UWidgetFieldNotificationExtension* Extension = UE::UMG::Private::FindWidgetNotifyExtension(this);
 		checkf(Extension, TEXT("If the EnabledFieldNotifications is valid, then the Extension must also be valid."));
 		UWidgetFieldNotificationExtension::FRemoveFromResult RemoveResult = Extension->RemoveFieldValueChangedDelegate(this, InFieldId, InHandle);
 		bResult = RemoveResult.bRemoved;
@@ -1704,12 +1759,30 @@ bool UWidget::RemoveFieldValueChangedDelegate(UE::FieldNotification::FFieldId In
 	return bResult;
 }
 
+void UWidget::K2_RemoveFieldValueChangedDelegate(FFieldNotificationId InFieldId, FFieldValueChangedDynamicDelegate InDelegate)
+{
+	if (InFieldId.IsValid())
+	{
+		const UE::FieldNotification::FFieldId FieldId = GetFieldNotificationDescriptor().GetField(GetClass(), InFieldId.FieldName);
+		if (ensureMsgf(FieldId.IsValid(), TEXT("The field should be compiled correctly.")))
+		{
+			if (EnabledFieldNotifications.IsValidIndex(FieldId.GetIndex()) && EnabledFieldNotifications[FieldId.GetIndex()])
+			{
+				UWidgetFieldNotificationExtension* Extension = UE::UMG::Private::FindWidgetNotifyExtension(this);
+				checkf(Extension, TEXT("If the EnabledFieldNotifications is valid, then the Extension must also be valid."));
+				UWidgetFieldNotificationExtension::FRemoveFromResult RemoveResult = Extension->RemoveFieldValueChangedDelegate(this, FieldId, InDelegate);
+				EnabledFieldNotifications[FieldId.GetIndex()] = RemoveResult.bHasOtherBoundDelegates;
+			}
+		}
+	}
+}
+
 int32 UWidget::RemoveAllFieldValueChangedDelegates(const void* InUserObject)
 {
 	int32 bResult = 0;
 	if (InUserObject)
 	{
-		if (UWidgetFieldNotificationExtension* Extension = UE::UMG::Private::FindOrAddWidgetNotifyExtension(this))
+		if (UWidgetFieldNotificationExtension* Extension = UE::UMG::Private::FindWidgetNotifyExtension(this))
 		{
 			UWidgetFieldNotificationExtension::FRemoveAllResult RemoveResult = Extension->RemoveAllFieldValueChangedDelegates(this, InUserObject);
 			bResult = RemoveResult.RemoveCount;
@@ -1724,7 +1797,7 @@ int32 UWidget::RemoveAllFieldValueChangedDelegates(UE::FieldNotification::FField
 	int32 bResult = 0;
 	if (InUserObject)
 	{
-		if (UWidgetFieldNotificationExtension* Extension = UE::UMG::Private::FindOrAddWidgetNotifyExtension(this))
+		if (UWidgetFieldNotificationExtension* Extension = UE::UMG::Private::FindWidgetNotifyExtension(this))
 		{
 			UWidgetFieldNotificationExtension::FRemoveAllResult RemoveResult = Extension->RemoveAllFieldValueChangedDelegates(this, InFieldId, InUserObject);
 			bResult = RemoveResult.RemoveCount;
@@ -1738,10 +1811,22 @@ void UWidget::BroadcastFieldValueChanged(UE::FieldNotification::FFieldId InField
 {
 	if (InFieldId.IsValid() && EnabledFieldNotifications.IsValidIndex(InFieldId.GetIndex()) && EnabledFieldNotifications[InFieldId.GetIndex()])
 	{
-		UWidgetFieldNotificationExtension* Extension = UE::UMG::Private::FindOrAddWidgetNotifyExtension(this);
+		UWidgetFieldNotificationExtension* Extension = UE::UMG::Private::FindWidgetNotifyExtension(this);
 		checkf(Extension, TEXT("If the EnabledFieldNotifications is valid, then the Extension must also be valid."));
 		Extension->BroadcastFieldValueChanged(this, InFieldId);
 	}
 }
 
+void UWidget::K2_BroadcastFieldValueChanged(FFieldNotificationId InFieldId)
+{
+	if (InFieldId.IsValid())
+	{
+		const UE::FieldNotification::FFieldId FieldId = GetFieldNotificationDescriptor().GetField(GetClass(), InFieldId.FieldName);
+		if (ensureMsgf(FieldId.IsValid(), TEXT("The field should be compiled correctly.")))
+		{
+			BroadcastFieldValueChanged(FieldId);
+		}
+	}
+}
+
 #undef LOCTEXT_NAMESPACE
diff --git a/Engine/Source/Runtime/UMG/Private/WidgetBlueprintGeneratedClass.cpp b/Engine/Source/Runtime/UMG/Private/WidgetBlueprintGeneratedClass.cpp
--- a/Engine/Source/Runtime/UMG/Private/WidgetBlueprintGeneratedClass.cpp
+++ b/Engine/Source/Runtime/UMG/Private/WidgetBlueprintGeneratedClass.cpp
@@ -117,6 +117,7 @@ namespace
 // UWidgetBlueprintGeneratedClass
 
 UWidgetBlueprintGeneratedClass::UWidgetBlueprintGeneratedClass()
+	: FieldNotifyStartBitNumber(INDEX_NONE)
 {
 #if WITH_EDITORONLY_DATA
 	{
@@ -349,6 +350,12 @@ void UWidgetBlueprintGeneratedClass::PostLoad()
 #endif
 }
 
+void UWidgetBlueprintGeneratedClass::PostLoadDefaultObject(UObject* Object)
+{
+	Super::PostLoadDefaultObject(Object);
+	InitializeFieldNotification(Cast<UUserWidget>(Object));
+}
+
 void UWidgetBlueprintGeneratedClass::PurgeClass(bool bRecompilingOnLoad)
 {
 	Super::PurgeClass(bRecompilingOnLoad);
@@ -433,6 +440,33 @@ UWidgetBlueprintGeneratedClass* UWidgetBlueprintGeneratedClass::FindWidgetTreeOw
 	return nullptr;
 }
 
+void UWidgetBlueprintGeneratedClass::InitializeFieldNotification(const UUserWidget* UserWidget)
+{
+	FieldNotifyStartBitNumber = 0;
+	if (UserWidget && FieldNotifyNames.Num())
+	{
+		int32 NumberOfField = 0;
+		UserWidget->GetFieldNotificationDescriptor().ForEachField(this, [&NumberOfField](::UE::FieldNotification::FFieldId FielId)
+			{
+				++NumberOfField;
+				return true;
+			});
+		FieldNotifyStartBitNumber = NumberOfField - FieldNotifyNames.Num();
+		ensureMsgf(FieldNotifyStartBitNumber >= 0, TEXT("The FieldNotifyStartIndex is negative. The number of field should be positive."));
+	}
+}
+
+void UWidgetBlueprintGeneratedClass::ForEachField(TFunctionRef<bool(::UE::FieldNotification::FFieldId FielId)> Callback) const
+{
+	ensureMsgf(FieldNotifyStartBitNumber >= 0, TEXT("The FieldNotifyStartIndex is negative. The number of field should be positive."));
+	for (int32 Index = 0; Index < FieldNotifyNames.Num(); ++Index)
+	{
+		if (!Callback(UE::FieldNotification::FFieldId(FieldNotifyNames[Index].GetFieldName(), Index + FieldNotifyStartBitNumber)))
+		{
+			break;
+		}
+	}
+}
 
 UWidgetBlueprintGeneratedClassExtension* UWidgetBlueprintGeneratedClass::GetExtension(TSubclassOf<UWidgetBlueprintGeneratedClassExtension> InExtensionType)
 {
diff --git a/Engine/Source/Runtime/UMG/Public/Blueprint/WidgetBlueprintGeneratedClass.h b/Engine/Source/Runtime/UMG/Public/Blueprint/WidgetBlueprintGeneratedClass.h
--- a/Engine/Source/Runtime/UMG/Public/Blueprint/WidgetBlueprintGeneratedClass.h
+++ b/Engine/Source/Runtime/UMG/Public/Blueprint/WidgetBlueprintGeneratedClass.h
@@ -6,6 +6,7 @@
 #include "UObject/ObjectMacros.h"
 #include "Binding/DynamicPropertyPath.h"
 #include "Engine/BlueprintGeneratedClass.h"
+#include "FieldNotification/FieldId.h"
 
 #include "WidgetBlueprintGeneratedClass.generated.h"
 
@@ -73,20 +74,21 @@ private:
 	UPROPERTY()
 	TArray<TObjectPtr<UWidgetBlueprintGeneratedClassExtension>> Extensions;
 
-#if WITH_EDITORONLY_DATA
-
-public:
-
-	UPROPERTY(Transient)
-	uint32 bCanCallPreConstruct:1;
-
-#endif
+	/** List Field Notifies. No index here on purpose to prevent saving them. */
+	UPROPERTY()
+	TArray<FFieldNotificationId> FieldNotifyNames;
 
-private:
+	int32 FieldNotifyStartBitNumber;
 
 	/** The classes native parent requires a native tick */
 	UPROPERTY()
-	uint32 bClassRequiresNativeTick:1;
+	uint32 bClassRequiresNativeTick :1;
+
+#if WITH_EDITORONLY_DATA
+public:
+	UPROPERTY(Transient)
+	uint32 bCanCallPreConstruct : 1;
+#endif
 
 public:
 	UPROPERTY()
@@ -99,19 +101,21 @@ public:
 	TArray< FName > NamedSlots;
 
 public:
-
 	UWidgetTree* GetWidgetTreeArchetype() const { return WidgetTree; }
 	void SetWidgetTreeArchetype(UWidgetTree* InWidgetTree);
 
 	// Walks up the hierarchy looking for a valid widget tree.
 	UWidgetBlueprintGeneratedClass* FindWidgetTreeOwningClass();
 
-	// UObject interface
-	virtual void Serialize(FArchive& Ar) override;
+	// Execute the callback for every FieldId defined in the BP class
+	void ForEachField(TFunctionRef<bool(::UE::FieldNotification::FFieldId FielId)> Callback) const;
 
+	//~ Begin UObject interface
+	virtual void Serialize(FArchive& Ar) override;
 	virtual void PostLoad() override;
 	virtual bool NeedsLoadForServer() const override;
-	// End UObject interface
+	virtual void PostLoadDefaultObject(UObject* Object) override;
+	//~ End UObject interface
 
 	virtual void PurgeClass(bool bRecompilingOnLoad) override;
 
@@ -160,4 +164,5 @@ public:
 private:
 	static void InitializeBindingsStatic(UUserWidget* UserWidget, const TArray< FDelegateRuntimeBinding >& InBindings);
 	static void BindAnimations(UUserWidget* Instance, const TArray< UWidgetAnimation* >& InAnimations);
+	void InitializeFieldNotification(const UUserWidget*);
 };
diff --git a/Engine/Source/Runtime/UMG/Public/Components/Widget.h b/Engine/Source/Runtime/UMG/Public/Components/Widget.h
--- a/Engine/Source/Runtime/UMG/Public/Components/Widget.h
+++ b/Engine/Source/Runtime/UMG/Public/Components/Widget.h
@@ -728,9 +728,18 @@ public:
 	virtual int32 RemoveAllFieldValueChangedDelegates(UE::FieldNotification::FFieldId InFieldId, const void* InUserObject) override final;
 	//~ End INotifyFieldValueChanged Interface
 
-public:
+	UFUNCTION(BlueprintCallable, Category = "FieldNotify", meta = (DisplayName = "Add Field Value Changed Delegate", ScriptName = "AddFieldValueChangedDelegate"))
+	void K2_AddFieldValueChangedDelegate(FFieldNotificationId FieldId, FFieldValueChangedDynamicDelegate Delegate);
+
+	UFUNCTION(BlueprintCallable, Category = "FieldNotify", meta = (DisplayName = "Remove Field Value Changed Delegate", ScriptName="RemoveFieldValueChangedDelegate"))
+	void K2_RemoveFieldValueChangedDelegate(FFieldNotificationId FieldId, FFieldValueChangedDynamicDelegate Delegate);
+
+protected:
 	void BroadcastFieldValueChanged(UE::FieldNotification::FFieldId InFieldId);
 
+	UFUNCTION(BlueprintCallable, Category="FieldNotify", meta = (DisplayName="Broadcast Field Value Changed", ScriptName="BroadcastFieldValueChanged"))
+	void K2_BroadcastFieldValueChanged(FFieldNotificationId FieldId);
+
 public:
 	/**
 	 * Gets the underlying slate widget or constructs it if it doesn't exist.  If you're looking to replace
diff --git a/Engine/Source/Runtime/UMG/Public/FieldNotification/FieldNotificationDeclaration.h b/Engine/Source/Runtime/UMG/Public/FieldNotification/FieldNotificationDeclaration.h
--- a/Engine/Source/Runtime/UMG/Public/FieldNotification/FieldNotificationDeclaration.h
+++ b/Engine/Source/Runtime/UMG/Public/FieldNotification/FieldNotificationDeclaration.h
@@ -9,6 +9,11 @@
 /*
  *	struct FFieldNotificationClassDescriptor : Super::FFieldNotificationClassDescriptor
  *	{
+ *	private:
+ *		using SuperDescriptor = Super::FFieldNotificationClassDescriptor;
+ *		static const ::UE::FieldNotification::FFieldId* AllFields[];
+ *		friend ThisClass;
+ *	public:
  *		static const ::UE::FieldNotification::FFieldId FieldA;
  *		static const ::UE::FieldNotification::FFieldId FieldB;
  *		enum
@@ -17,13 +22,17 @@
  *			IndexOf_FieldB,
  *			Max_IndexOf_,
  *		};
- *		virtual int32 GetNumberOfField() const override { reutrn Max_IndexOf_; }
- *		virtual ::UE::FieldNotification::FFieldId GetField(FName InFieldName) const override;
- *		virtual ::UE::FieldNotification::FFieldId GetField(int32 InFieldNumber) const override;
- *		virtual const ::UE::FieldNotification::IClassDescriptor& GetFieldNotificationDescriptor() const override;
+ * 	virtual void ForEachField(const UClass* Class, TFunctionRef<bool(::UE::FieldNotification::FFieldId FielId)> Callback) const override;
  *	};
+ * virtual const ::UE::FieldNotification::IClassDescriptor& GetFieldNotificationDescriptor() const override;
  */
 
+#define UE_FIELD_NOTIFICATION_OVERRIDE_GET_FIELD_NOTIFICATION_DESCRIPTOR() \
+	virtual const ::UE::FieldNotification::IClassDescriptor& GetFieldNotificationDescriptor() const \
+	{ \
+ 		static FFieldNotificationClassDescriptor Instance; \
+ 		return Instance; \
+ 	}
 
 #define UE_FIELD_NOTIFICATION_DECLARE_CLASS_DESCRIPTOR_BEGIN() \
 	struct FFieldNotificationClassDescriptor : public Super::FFieldNotificationClassDescriptor \
@@ -36,36 +45,19 @@
 
 
 #define UE_FIELD_NOTIFICATION_DECLARE_CLASS_DESCRIPTOR_END() \
-		virtual int32 GetNumberOfField() const override \
-		{ \
-			return Max_IndexOf_; \
-		} \
-		virtual ::UE::FieldNotification::FFieldId GetField(FName InFieldName) const override \
+		virtual void ForEachField(const UClass* Class, TFunctionRef<bool(::UE::FieldNotification::FFieldId FielId)> Callback) const override \
 		{ \
 			for (int32 Index = 0; Index < Max_IndexOf_-SuperDescriptor::Max_IndexOf_; ++Index) \
 			{ \
-				if (AllFields[Index]->GetName() == InFieldName) \
+				if (!Callback(*AllFields[Index])) \
 				{ \
-					return *(AllFields[Index]); \
+					return; \
 				} \
 			} \
-			return SuperDescriptor::GetField(InFieldName); \
-		} \
-		virtual ::UE::FieldNotification::FFieldId GetField(int32 InFieldNumber) const override \
-		{ \
-			check(InFieldNumber >= 0 && InFieldNumber < Max_IndexOf_); \
-			if (InFieldNumber < SuperDescriptor::Max_IndexOf_) \
-			{ \
-				return SuperDescriptor::GetField(InFieldNumber); \
-			} \
-			return *(AllFields[InFieldNumber - SuperDescriptor::Max_IndexOf_]); \
+			SuperDescriptor::ForEachField(Class, Callback); \
 		} \
 	}; \
-	virtual const ::UE::FieldNotification::IClassDescriptor& GetFieldNotificationDescriptor() const override \
- 	{ \
- 		static FFieldNotificationClassDescriptor Instance; \
- 		return Instance; \
- 	}
+	UE_FIELD_NOTIFICATION_OVERRIDE_GET_FIELD_NOTIFICATION_DESCRIPTOR()
 
 
 #define UE_FIELD_NOTIFICATION_DECLARE_CLASS_DESCRIPTOR_BASE_BEGIN() \
@@ -79,32 +71,9 @@
 
 
 #define UE_FIELD_NOTIFICATION_DECLARE_CLASS_DESCRIPTOR_BASE_END() \
-		virtual int32 GetNumberOfField() const override \
-		{ \
-			return Max_IndexOf_; \
-		} \
-		virtual ::UE::FieldNotification::FFieldId GetField(FName InFieldName) const override \
-		{ \
-			for (int32 Index = 0; Index < Max_IndexOf_; ++Index) \
-			{ \
-				if (AllFields[Index]->GetName() == InFieldName) \
-				{ \
-					return *(AllFields[Index]); \
-				} \
-			} \
-			return ::UE::FieldNotification::FFieldId(); \
-		} \
-		virtual ::UE::FieldNotification::FFieldId GetField(int32 InFieldNumber) const override \
-		{ \
-			check(InFieldNumber >= 0 && InFieldNumber < Max_IndexOf_); \
-			return *(AllFields[InFieldNumber]); \
-		} \
+		virtual void ForEachField(const UClass* Class, TFunctionRef<bool(::UE::FieldNotification::FFieldId FielId)> Callback) const override; \
 	}; \
-	virtual const ::UE::FieldNotification::IClassDescriptor& GetFieldNotificationDescriptor() const \
-	{ \
- 		static FFieldNotificationClassDescriptor Instance; \
- 		return Instance; \
- 	}
+	UE_FIELD_NOTIFICATION_OVERRIDE_GET_FIELD_NOTIFICATION_DESCRIPTOR()
 
 
 #define UE_FIELD_NOTIFICATION_DECLARE_ENUM_FIELD_BEGIN(Name) \
diff --git a/Engine/Source/Runtime/UMG/Public/FieldNotification/IClassDescriptor.h b/Engine/Source/Runtime/UMG/Public/FieldNotification/IClassDescriptor.h
--- a/Engine/Source/Runtime/UMG/Public/FieldNotification/IClassDescriptor.h
+++ b/Engine/Source/Runtime/UMG/Public/FieldNotification/IClassDescriptor.h
@@ -15,9 +15,29 @@ namespace UE::FieldNotification
 			Max_IndexOf_ = 0,
 		};
 
-		virtual int32 GetNumberOfField() const = 0;
-		virtual FFieldId GetField(FName InFieldName) const = 0;
-		virtual FFieldId GetField(int32 InFieldNumber) const = 0;
+		/** Find the FieldId by name. Returns an invalid FFIeldId if not found. */
+		FFieldId GetField(const UClass* Class, FName InFieldName) const
+		{
+			FFieldId FoundId;
+			ForEachField(Class, [&FoundId, InFieldName](const FFieldId Other)
+			{
+				if (Other.GetName() == InFieldName)
+				{
+					FoundId = Other;
+					return false;
+				}
+				return true;
+			});
+			return FoundId;
+		}
+
+		/** Execute the callback for every FieldId in the ClassDescriptor. */
+		virtual void ForEachField(const UClass* Class, TFunctionRef<bool(FFieldId FielId)> Callback) const
+		{
+
+		}
+
+		/** */
 		virtual ~IClassDescriptor() = default;
 	};
 
