Added MVVM getter/setter icons.

diff --git a/Engine/Source/Runtime/UMG/Private/UMGStyle.cpp b/Engine/Source/Runtime/UMG/Private/UMGStyle.cpp
--- a/Engine/Source/Runtime/UMG/Private/UMGStyle.cpp
+++ b/Engine/Source/Runtime/UMG/Private/UMGStyle.cpp
@@ -164,6 +164,8 @@ TSharedRef< FSlateStyleSet > FUMGStyle::Create()
 	Style->Set("MVVM.OneWayToSource", new IMAGE_BRUSH_SVG("ViewModel_OneWayToSource", Icon16x16));
 	Style->Set("MVVM.TwoWay", new IMAGE_BRUSH_SVG("ViewModel_TwoWay", Icon16x16));
 	Style->Set("MVVM.OneTime", new IMAGE_BRUSH_SVG("ViewModel_OneTime", Icon16x16));
+	Style->Set("MVVM.SetterFunction", new IMAGE_BRUSH_SVG("Function_Setter", Icon16x16));
+	Style->Set("MVVM.GetterFunction", new IMAGE_BRUSH_SVG("Function_Getter", Icon16x16));
 
 	return Style;
 }
