Force inlined IsValidAsField to avoid duplicate definition error.

#rnx

diff --git a/Engine/Source/Runtime/UMG/Public/FieldNotification/FieldNotificationHelpers.h b/Engine/Source/Runtime/UMG/Public/FieldNotification/FieldNotificationHelpers.h
--- a/Engine/Source/Runtime/UMG/Public/FieldNotification/FieldNotificationHelpers.h
+++ b/Engine/Source/Runtime/UMG/Public/FieldNotification/FieldNotificationHelpers.h
@@ -7,8 +7,7 @@
 
 namespace UE::FieldNotification::Helpers
 {
-
-	UE_NODISCARD bool IsValidAsField(const UFunction* InFunction)
+	UE_NODISCARD FORCEINLINE bool IsValidAsField(const UFunction* InFunction)
 	{
 		return InFunction != nullptr
 			&& !InFunction->HasAnyFunctionFlags(FUNC_Net | FUNC_Event)
