Prevent crash with uninitialized asset

diff --git a/Engine/Source/Runtime/UMG/Private/Slate/SlateVectorArtData.cpp b/Engine/Source/Runtime/UMG/Private/Slate/SlateVectorArtData.cpp
--- a/Engine/Source/Runtime/UMG/Private/Slate/SlateVectorArtData.cpp
+++ b/Engine/Source/Runtime/UMG/Private/Slate/SlateVectorArtData.cpp
@@ -148,7 +148,7 @@ UMaterialInstanceDynamic* USlateVectorArtData::ConvertToMaterialInstanceDynamic(
 void USlateVectorArtData::EnsureValidData()
 {
 #if WITH_EDITORONLY_DATA
-	InitFromStaticMesh(*MeshAsset);
+	//InitFromStaticMesh(*MeshAsset);
 #endif
 }
 
