UMG: Make the some internal flags private. Set the initialize flag once the initialization is completed. That fix an issue with AddExtention. Cached the FObjectProperty in a map, that prevents multiple link list traversal.

diff --git a/Engine/Source/Editor/UMGEditor/Public/WidgetBlueprintCompiler.h b/Engine/Source/Editor/UMGEditor/Public/WidgetBlueprintCompiler.h
--- a/Engine/Source/Editor/UMGEditor/Public/WidgetBlueprintCompiler.h
+++ b/Engine/Source/Editor/UMGEditor/Public/WidgetBlueprintCompiler.h
@@ -83,9 +83,10 @@ public:
 
 	struct UMGEDITOR_API FCreateVariableContext
 	{
-		FCreateVariableContext(FWidgetBlueprintCompilerContext& InContext);
 		FProperty* CreateVariable(const FName Name, const FEdGraphPinType& Type) const;
 	private:
+		friend FWidgetBlueprintCompilerContext;
+		FCreateVariableContext(FWidgetBlueprintCompilerContext& InContext);
 		FWidgetBlueprintCompilerContext& Context;
 	};
 
diff --git a/Engine/Source/Runtime/UMG/Private/WidgetBlueprintGeneratedClass.cpp b/Engine/Source/Runtime/UMG/Private/WidgetBlueprintGeneratedClass.cpp
--- a/Engine/Source/Runtime/UMG/Private/WidgetBlueprintGeneratedClass.cpp
+++ b/Engine/Source/Runtime/UMG/Private/WidgetBlueprintGeneratedClass.cpp
@@ -127,7 +127,7 @@ UWidgetBlueprintGeneratedClass::UWidgetBlueprintGeneratedClass()
 #endif
 }
 
-void UWidgetBlueprintGeneratedClass::InitializeBindingsStatic(UUserWidget* UserWidget, const TArray< FDelegateRuntimeBinding >& InBindings)
+void UWidgetBlueprintGeneratedClass::InitializeBindingsStatic(UUserWidget* UserWidget, const TArrayView<const FDelegateRuntimeBinding> InBindings, const TMap<FName, FObjectPropertyBase*>& InPropertyMap)
 {
 	check(!UserWidget->IsTemplate());
 
@@ -137,41 +137,38 @@ void UWidgetBlueprintGeneratedClass::InitializeBindingsStatic(UUserWidget* UserW
 	// For each property binding that we're given, find the corresponding field, and setup the delegate binding on the widget.
 	for (const FDelegateRuntimeBinding& Binding : InBindings)
 	{
-		// If the binding came from a parent class, this will still find it - FindField() searches the super class hierarchy by default.
-		FObjectProperty* WidgetProperty = FindFProperty<FObjectProperty>(UserWidget->GetClass(), *Binding.ObjectName);
-		if (WidgetProperty == nullptr)
+		if (FObjectPropertyBase*const* PropPtr = InPropertyMap.Find(*Binding.ObjectName))
 		{
-			continue;
-		}
-
-		UWidget* Widget = Cast<UWidget>(WidgetProperty->GetObjectPropertyValue_InContainer(UserWidget));
-
-		if (Widget)
-		{
-			FDelegateProperty* DelegateProperty = FindFProperty<FDelegateProperty>(Widget->GetClass(), FName(*(Binding.PropertyName.ToString() + TEXT("Delegate"))));
-			if (!DelegateProperty)
-			{
-				DelegateProperty = FindFProperty<FDelegateProperty>(Widget->GetClass(), Binding.PropertyName);
-			}
+			const FObjectPropertyBase* WidgetProperty = *PropPtr;
+			check(WidgetProperty);
 
-			if (DelegateProperty)
+			if (UWidget* Widget = Cast<UWidget>(WidgetProperty->GetObjectPropertyValue_InContainer(UserWidget)))
 			{
-				bool bSourcePathBound = false;
-
-				if (Binding.SourcePath.IsValid())
+				FDelegateProperty* DelegateProperty = FindFProperty<FDelegateProperty>(Widget->GetClass(), FName(*(Binding.PropertyName.ToString() + TEXT("Delegate"))));
+				if (!DelegateProperty)
 				{
-					bSourcePathBound = Widget->AddBinding(DelegateProperty, UserWidget, Binding.SourcePath);
+					DelegateProperty = FindFProperty<FDelegateProperty>(Widget->GetClass(), Binding.PropertyName);
 				}
 
-				// If no native binder is found then the only possibility is that the binding is for
-				// a delegate that doesn't match the known native binders available and so we
-				// fallback to just attempting to bind to the function directly.
-				if (bSourcePathBound == false)
+				if (DelegateProperty)
 				{
-					FScriptDelegate* ScriptDelegate = DelegateProperty->GetPropertyValuePtr_InContainer(Widget);
-					if (ScriptDelegate)
+					bool bSourcePathBound = false;
+
+					if (Binding.SourcePath.IsValid())
 					{
-						ScriptDelegate->BindUFunction(UserWidget, Binding.FunctionName);
+						bSourcePathBound = Widget->AddBinding(DelegateProperty, UserWidget, Binding.SourcePath);
+					}
+
+					// If no native binder is found then the only possibility is that the binding is for
+					// a delegate that doesn't match the known native binders available and so we
+					// fallback to just attempting to bind to the function directly.
+					if (bSourcePathBound == false)
+					{
+						FScriptDelegate* ScriptDelegate = DelegateProperty->GetPropertyValuePtr_InContainer(Widget);
+						if (ScriptDelegate)
+						{
+							ScriptDelegate->BindUFunction(UserWidget, Binding.FunctionName);
+						}
 					}
 				}
 			}
@@ -182,8 +179,8 @@ void UWidgetBlueprintGeneratedClass::InitializeBindingsStatic(UUserWidget* UserW
 void UWidgetBlueprintGeneratedClass::InitializeWidgetStatic(UUserWidget* UserWidget
 	, const UClass* InClass
 	, UWidgetTree* InWidgetTree
-	, const TArray< UWidgetAnimation* >& InAnimations
-	, const TArray< FDelegateRuntimeBinding >& InBindings)
+	, const TArrayView<UWidgetAnimation*> InAnimations
+	, const TArrayView<const FDelegateRuntimeBinding> InBindings)
 {
 	check(InClass);
 
@@ -220,10 +217,17 @@ void UWidgetBlueprintGeneratedClass::InitializeWidgetStatic(UUserWidget* UserWid
 
 	if (ClonedTree)
 	{
-		BindAnimations(UserWidget, InAnimations);
-
 		UClass* WidgetBlueprintClass = UserWidget->GetClass();
 
+		TMap<FName, FObjectPropertyBase*> ObjectPropertiesMap;
+		for (TFieldIterator<FObjectPropertyBase>It(WidgetBlueprintClass, EFieldIterationFlags::Default); It; ++It)
+		{
+			check(*It);
+			ObjectPropertiesMap.Add(It->GetFName(), *It);
+		}
+
+		BindAnimationsStatic(UserWidget, InAnimations, ObjectPropertiesMap);
+
 		ClonedTree->ForEachWidget([&](UWidget* Widget) {
 			// Not fatal if NULL, but shouldn't happen
 			if (!ensure(Widget != nullptr))
@@ -240,9 +244,10 @@ void UWidgetBlueprintGeneratedClass::InitializeWidgetStatic(UUserWidget* UserWid
 #endif
 
 			// Find property with the same name as the template and assign the new widget to it.
-			FObjectPropertyBase* Prop = FindFProperty<FObjectPropertyBase>(WidgetBlueprintClass, Widget->GetFName());
-			if (Prop)
+			if (FObjectPropertyBase** PropPtr = ObjectPropertiesMap.Find(Widget->GetFName()))
 			{
+				FObjectPropertyBase* Prop = *PropPtr;
+				check(Prop);
 				Prop->SetObjectPropertyValue_InContainer(UserWidget, Widget);
 				UObject* Value = Prop->GetObjectPropertyValue_InContainer(UserWidget);
 				check(Value == Widget);
@@ -259,7 +264,7 @@ void UWidgetBlueprintGeneratedClass::InitializeWidgetStatic(UUserWidget* UserWid
 #endif
 		});
 
-		InitializeBindingsStatic(UserWidget, InBindings);
+		InitializeBindingsStatic(UserWidget, InBindings, ObjectPropertiesMap);
 
 		// Bind any delegates on widgets
 		if (!UserWidget->IsDesignTime())
@@ -269,7 +274,7 @@ void UWidgetBlueprintGeneratedClass::InitializeWidgetStatic(UUserWidget* UserWid
 	}
 }
 
-void UWidgetBlueprintGeneratedClass::BindAnimations(UUserWidget* Instance, const TArray< UWidgetAnimation* >& InAnimations)
+void UWidgetBlueprintGeneratedClass::BindAnimationsStatic(UUserWidget* Instance, const TArrayView<UWidgetAnimation*> InAnimations, const TMap<FName, FObjectPropertyBase*>& InPropertyMap)
 {
 	// Note: It's not safe to assume here that the UserWidget class type is a UWidgetBlueprintGeneratedClass!
 	// - @see InitializeWidgetStatic()
@@ -279,10 +284,10 @@ void UWidgetBlueprintGeneratedClass::BindAnimations(UUserWidget* Instance, const
 		if (Animation->GetMovieScene())
 		{
 			// Find property with the same name as the animation and assign the animation to it.
-			FObjectPropertyBase* Prop = FindFProperty<FObjectPropertyBase>(Instance->GetClass(), Animation->GetMovieScene()->GetFName());
-			if (Prop)
+			if (FObjectPropertyBase*const* PropPtr = InPropertyMap.Find(Animation->GetMovieScene()->GetFName()))
 			{
-				Prop->SetObjectPropertyValue_InContainer(Instance, Animation);
+				check(*PropPtr);
+				(*PropPtr)->SetObjectPropertyValue_InContainer(Instance, Animation);
 			}
 		}
 	}
@@ -297,8 +302,9 @@ void UWidgetBlueprintGeneratedClass::SetClassRequiresNativeTick(bool InClassRequ
 
 void UWidgetBlueprintGeneratedClass::InitializeWidget(UUserWidget* UserWidget) const
 {
-	TArray<UWidgetAnimation*> AllAnims;
-	TArray<FDelegateRuntimeBinding> AllBindings;
+	FMemMark Mark(FMemStack::Get());
+	TArray<UWidgetAnimation*, TMemStackAllocator<>> AllAnims;
+	TArray<FDelegateRuntimeBinding, TMemStackAllocator<>> AllBindings;
 
 	// Include current class animations.
 	AllAnims.Append(Animations);
diff --git a/Engine/Source/Runtime/UMG/Public/Blueprint/UserWidget.h b/Engine/Source/Runtime/UMG/Public/Blueprint/UserWidget.h
--- a/Engine/Source/Runtime/UMG/Public/Blueprint/UserWidget.h
+++ b/Engine/Source/Runtime/UMG/Public/Blueprint/UserWidget.h
@@ -1256,7 +1256,7 @@ public:
 	UPROPERTY()
 	uint8 bHasScriptImplementedPaint : 1;
 
-protected:
+private:
 
 	/** Has this widget been initialized by its class yet? */
 	uint8 bInitialized : 1;
diff --git a/Engine/Source/Runtime/UMG/Public/Blueprint/WidgetBlueprintGeneratedClass.h b/Engine/Source/Runtime/UMG/Public/Blueprint/WidgetBlueprintGeneratedClass.h
--- a/Engine/Source/Runtime/UMG/Public/Blueprint/WidgetBlueprintGeneratedClass.h
+++ b/Engine/Source/Runtime/UMG/Public/Blueprint/WidgetBlueprintGeneratedClass.h
@@ -129,8 +129,8 @@ public:
 	static void InitializeWidgetStatic(UUserWidget* UserWidget
 		, const UClass* InClass
 		, UWidgetTree* InWidgetTree
-		, const TArray< UWidgetAnimation* >& InAnimations
-		, const TArray< FDelegateRuntimeBinding >& InBindings);
+		, const TArrayView<UWidgetAnimation*> InAnimations
+		, const TArrayView<const FDelegateRuntimeBinding> InBindings);
 
 	bool ClassRequiresNativeTick() const { return bClassRequiresNativeTick; }
 
@@ -162,7 +162,7 @@ public:
 	}
 
 private:
-	static void InitializeBindingsStatic(UUserWidget* UserWidget, const TArray< FDelegateRuntimeBinding >& InBindings);
-	static void BindAnimations(UUserWidget* Instance, const TArray< UWidgetAnimation* >& InAnimations);
+	static void InitializeBindingsStatic(UUserWidget* UserWidget, const TArrayView<const FDelegateRuntimeBinding> InBindings, const TMap<FName, FObjectPropertyBase*>& InPropertyMap);
+	static void BindAnimationsStatic(UUserWidget* Instance, const TArrayView<UWidgetAnimation*> InAnimations, const TMap<FName, FObjectPropertyBase*>& InPropertyMap);
 	void InitializeFieldNotification(const UUserWidget*);
 };
