Optimisation - Add option to not update the position of widgets within the SWorldWidgetScreenLayer if they are not visible. Prevents invalidating all the widgets in the layer to update widgets that aren't visible anyhow.

diff --git a/Engine/Source/Runtime/UMG/Private/Slate/SWorldWidgetScreenLayer.cpp b/Engine/Source/Runtime/UMG/Private/Slate/SWorldWidgetScreenLayer.cpp
--- a/Engine/Source/Runtime/UMG/Private/Slate/SWorldWidgetScreenLayer.cpp
+++ b/Engine/Source/Runtime/UMG/Private/Slate/SWorldWidgetScreenLayer.cpp
@@ -19,6 +19,13 @@ static FAutoConsoleVariableRef CVarSlateWorldWidgetZOrder(
 	ECVF_Default
 	);
 
+static bool GSlateWorldWidgetIgnoreNotVisibleWidgets = false;
+static FAutoConsoleVariableRef CVarSlateWorldWidgetIgnoreNotVisibleWidgets(
+	TEXT("Slate.WorldWidgetIgnoreNotVisibleWidgets"),
+	GSlateWorldWidgetIgnoreNotVisibleWidgets,
+	TEXT("Whether to not update the position of world widgets if they are not visible - to prevent invalidating the whole layer unnecessarily")
+	);
+
 SWorldWidgetScreenLayer::FComponentEntry::FComponentEntry()
 	: Slot(nullptr)
 {
@@ -123,7 +130,23 @@ void SWorldWidgetScreenLayer::Tick(const FGeometry& AllottedGeometry, const doub
 					FVector WorldLocation = SceneComponent->GetComponentLocation();
 
 					FVector2D ScreenPosition2D;
-					const bool bProjected = bHasProjectionData ? FSceneView::ProjectWorldToScreen(WorldLocation, ProjectionData.GetConstrainedViewRect(), ViewProjectionMatrix, ScreenPosition2D) : false;
+					const bool bProjected = [&Entry, bHasProjectionData, &WorldLocation, &ScreenPosition2D, &ProjectionData, &ViewProjectionMatrix]()
+					{
+						if (!bHasProjectionData)
+						{
+							return false;
+						}
+
+						if (GSlateWorldWidgetIgnoreNotVisibleWidgets && 
+							Entry.WidgetComponent && 
+							!Entry.WidgetComponent->IsWidgetVisible())
+						{
+							return false;
+						}
+
+						return FSceneView::ProjectWorldToScreen(WorldLocation, ProjectionData.GetConstrainedViewRect(), ViewProjectionMatrix, ScreenPosition2D);
+					}();
+
 					if (bProjected)
 					{
 						const float ViewportDist = FVector::Dist(ProjectionData.ViewOrigin, WorldLocation);
