UMG: Remove WidgetFieldNotificationExtension from BP type. Rename argument for GetExtension.
#rnx

diff --git a/Engine/Source/Runtime/UMG/Private/Binding/WidgetFieldNotificationExtension.h b/Engine/Source/Runtime/UMG/Private/Binding/WidgetFieldNotificationExtension.h
--- a/Engine/Source/Runtime/UMG/Private/Binding/WidgetFieldNotificationExtension.h
+++ b/Engine/Source/Runtime/UMG/Private/Binding/WidgetFieldNotificationExtension.h
@@ -11,7 +11,7 @@
 
 class UVisual;
 
-UCLASS(MinimalAPI, Transient)
+UCLASS(MinimalAPI, Transient, NotBlueprintType)
 class UWidgetFieldNotificationExtension : public UUserWidgetExtension
 {
 	GENERATED_BODY()
diff --git a/Engine/Source/Runtime/UMG/Public/Blueprint/UserWidget.h b/Engine/Source/Runtime/UMG/Public/Blueprint/UserWidget.h
--- a/Engine/Source/Runtime/UMG/Public/Blueprint/UserWidget.h
+++ b/Engine/Source/Runtime/UMG/Public/Blueprint/UserWidget.h
@@ -1069,11 +1069,11 @@ public:
 
 	/** Find the first extension of the requested type. */
 	UFUNCTION(BlueprintCallable, Category = "User Interface|Extension")
-	UUserWidgetExtension* GetExtension(TSubclassOf<UUserWidgetExtension> InExtensionType) const;
+	UUserWidgetExtension* GetExtension(TSubclassOf<UUserWidgetExtension> ExtensionType) const;
 
 	/** Find the extensions of the requested type. */
 	UFUNCTION(BlueprintCallable, Category = "User Interface|Extension")
-	TArray<UUserWidgetExtension*> GetExtensions(TSubclassOf<UUserWidgetExtension> InExtensionType) const;
+	TArray<UUserWidgetExtension*> GetExtensions(TSubclassOf<UUserWidgetExtension> ExtensionType) const;
 
 	/** Add the extension of the requested type. */
 	template<typename ExtensionType>
