UMG: Fix bad test in FieldId

diff --git a/Engine/Source/Runtime/UMG/Public/FieldNotification/FieldId.h b/Engine/Source/Runtime/UMG/Public/FieldNotification/FieldId.h
--- a/Engine/Source/Runtime/UMG/Public/FieldNotification/FieldId.h
+++ b/Engine/Source/Runtime/UMG/Public/FieldNotification/FieldId.h
@@ -98,7 +98,7 @@ public:
 public:
 	bool IsValid() const
 	{
-		return FieldName.IsNone();
+		return !FieldName.IsNone();
 	}
 
 	FName GetFieldName() const
