Fix include bug after changes to EngineTypes.h

diff --git a/Engine/Source/Runtime/UMG/Public/Components/WidgetInteractionComponent.h b/Engine/Source/Runtime/UMG/Public/Components/WidgetInteractionComponent.h
--- a/Engine/Source/Runtime/UMG/Public/Components/WidgetInteractionComponent.h
+++ b/Engine/Source/Runtime/UMG/Public/Components/WidgetInteractionComponent.h
@@ -6,7 +6,7 @@
 #include "UObject/ObjectMacros.h"
 #include "UObject/UObjectGlobals.h"
 #include "InputCoreTypes.h"
-#include "Engine/EngineTypes.h"
+#include "Engine/HitResult.h"
 #include "Components/SceneComponent.h"
 #include "GenericPlatform/GenericApplication.h"
 #include "Layout/WidgetPath.h"
