UMG: Update the CheckBox widget with accessors.

diff --git a/Engine/Source/Runtime/UMG/Private/Components/CheckBox.cpp b/Engine/Source/Runtime/UMG/Private/Components/CheckBox.cpp
--- a/Engine/Source/Runtime/UMG/Private/Components/CheckBox.cpp
+++ b/Engine/Source/Runtime/UMG/Private/Components/CheckBox.cpp
@@ -12,8 +12,6 @@
 /////////////////////////////////////////////////////
 // UCheckBox
 
-UE_FIELD_NOTIFICATION_IMPLEMENT_CLASS_DESCRIPTOR_OneField(UCheckBox, CheckedState);
-
 static FCheckBoxStyle* DefaultCheckboxStyle = nullptr;
 
 #if WITH_EDITOR
@@ -31,6 +29,7 @@ UCheckBox::UCheckBox(const FObjectInitializer& ObjectInitializer)
 		DefaultCheckboxStyle->UnlinkColors();
 	}
 
+PRAGMA_DISABLE_DEPRECATION_WARNINGS
 	WidgetStyle = *DefaultCheckboxStyle;
 
 #if WITH_EDITOR 
@@ -57,6 +56,8 @@ UCheckBox::UCheckBox(const FObjectInitializer& ObjectInitializer)
 
 	ClickMethod = EButtonClickMethod::DownAndUp;
 	TouchMethod = EButtonTouchMethod::DownAndUp;
+	PressMethod = EButtonPressMethod::DownAndUp;
+PRAGMA_ENABLE_DEPRECATION_WARNINGS
 
 	IsFocusable = true;
 #if WITH_EDITORONLY_DATA
@@ -74,6 +75,7 @@ void UCheckBox::ReleaseSlateResources(bool bReleaseChildren)
 
 TSharedRef<SWidget> UCheckBox::RebuildWidget()
 {
+PRAGMA_DISABLE_DEPRECATION_WARNINGS
 	MyCheckbox = SNew(SCheckBox)
 		.OnCheckStateChanged( BIND_UOBJECT_DELEGATE(FOnCheckStateChanged, SlateOnCheckStateChangedCallback) )
 		.Style(&WidgetStyle)
@@ -83,6 +85,7 @@ TSharedRef<SWidget> UCheckBox::RebuildWidget()
 		.PressMethod(PressMethod)
 		.IsFocusable(IsFocusable)
 		;
+PRAGMA_ENABLE_DEPRECATION_WARNINGS
 
 	if ( GetChildrenCount() > 0 )
 	{
@@ -96,8 +99,14 @@ void UCheckBox::SynchronizeProperties()
 {
 	Super::SynchronizeProperties();
 
+
+PRAGMA_DISABLE_DEPRECATION_WARNINGS
 	MyCheckbox->SetStyle(&WidgetStyle);
 	MyCheckbox->SetIsChecked( PROPERTY_BINDING(ECheckBoxState, CheckedState) );
+	MyCheckbox->SetClickMethod(ClickMethod);
+	MyCheckbox->SetTouchMethod(TouchMethod);
+	MyCheckbox->SetPressMethod(PressMethod);
+PRAGMA_ENABLE_DEPRECATION_WARNINGS
 }
 
 void UCheckBox::OnSlotAdded(UPanelSlot* InSlot)
@@ -128,6 +137,13 @@ bool UCheckBox::IsPressed() const
 	return false;
 }
 
+
+PRAGMA_DISABLE_DEPRECATION_WARNINGS
+EButtonClickMethod::Type UCheckBox::GetClickMethod() const
+{
+	return ClickMethod;
+}
+
 void UCheckBox::SetClickMethod(EButtonClickMethod::Type InClickMethod)
 {
 	ClickMethod = InClickMethod;
@@ -137,6 +153,11 @@ void UCheckBox::SetClickMethod(EButtonClickMethod::Type InClickMethod)
 	}
 }
 
+EButtonTouchMethod::Type UCheckBox::GetTouchMethod() const
+{
+	return TouchMethod;
+}
+
 void UCheckBox::SetTouchMethod(EButtonTouchMethod::Type InTouchMethod)
 {
 	TouchMethod = InTouchMethod;
@@ -155,6 +176,11 @@ void UCheckBox::SetPressMethod(EButtonPressMethod::Type InPressMethod)
 	}
 }
 
+EButtonPressMethod::Type UCheckBox::GetPressMethod() const
+{
+	return PressMethod;
+}
+
 bool UCheckBox::IsChecked() const
 {
 	if ( MyCheckbox.IsValid() )
@@ -204,15 +230,32 @@ void UCheckBox::SetCheckedState(ECheckBoxState InCheckedState)
 	}
 }
 
+const FCheckBoxStyle& UCheckBox::GetWidgetStyle() const
+{
+	return WidgetStyle;
+}
+
+void UCheckBox::SetWidgetStyle(const FCheckBoxStyle& InStyle)
+{
+	WidgetStyle = InStyle;
+
+	if (MyCheckbox)
+	{
+		MyCheckbox->SetStyle(&WidgetStyle);
+	}
+}
+PRAGMA_ENABLE_DEPRECATION_WARNINGS
+
 void UCheckBox::SlateOnCheckStateChangedCallback(ECheckBoxState NewState)
 {
+PRAGMA_DISABLE_DEPRECATION_WARNINGS
 	if (CheckedState != NewState)
 	{
 		CheckedState = NewState;
 		BroadcastFieldValueChanged(FFieldNotificationClassDescriptor::CheckedState);
 	}
+PRAGMA_ENABLE_DEPRECATION_WARNINGS
 
-	//@TODO: Choosing to treat Undetermined as Checked
 	const bool bWantsToBeChecked = NewState != ECheckBoxState::Unchecked;
 	OnCheckStateChanged.Broadcast(bWantsToBeChecked);
 }
diff --git a/Engine/Source/Runtime/UMG/Public/Components/CheckBox.h b/Engine/Source/Runtime/UMG/Public/Components/CheckBox.h
--- a/Engine/Source/Runtime/UMG/Public/Components/CheckBox.h
+++ b/Engine/Source/Runtime/UMG/Public/Components/CheckBox.h
@@ -31,41 +31,42 @@ class UMG_API UCheckBox : public UContentWidget
 {
 	GENERATED_UCLASS_BODY()
 
-public:
-	UE_FIELD_NOTIFICATION_DECLARE_CLASS_DESCRIPTOR_OneField(CheckedState);
-
 public:
 	/** Whether the check box is currently in a checked state */
-	UPROPERTY(EditAnywhere, Category=Appearance)
+	UE_DEPRECATED(5.1, "Direct access to CheckedState is deprecated. Please use the getter or setter.")
+	UPROPERTY(EditAnywhere, BlueprintReadWrite, Getter, Setter, BlueprintGetter="GetCheckedState", BlueprintSetter="SetCheckedState", FieldNotify, Category="Appearance")
 	ECheckBoxState CheckedState;
 
 	/** A bindable delegate for the IsChecked. */
 	UPROPERTY()
 	FGetCheckBoxState CheckedStateDelegate;
 
-public:
 	/** The checkbox bar style */
-	UPROPERTY(EditAnywhere, BlueprintReadWrite, Category="Style", meta=( DisplayName="Style" ))
+	UE_DEPRECATED(5.1, "Direct access to WidgetStyle is deprecated. Please use the getter or setter.")
+	UPROPERTY(EditAnywhere, BlueprintReadWrite, Getter, Setter, Category="Style", meta = (DisplayName="Style"))
 	FCheckBoxStyle WidgetStyle;
 
 	/** How the content of the toggle button should align within the given space */
-	UPROPERTY(EditAnywhere, BlueprintReadOnly, Category=Appearance)
+	UPROPERTY(EditAnywhere, BlueprintReadOnly, Category="Appearance")
 	TEnumAsByte<EHorizontalAlignment> HorizontalAlignment;
 
 	/** The type of mouse action required by the user to trigger the buttons 'Click' */
-	UPROPERTY(EditAnywhere, BlueprintReadOnly, Category = "Interaction", AdvancedDisplay)
+	UE_DEPRECATED(5.1, "Direct access to ClickMethod is deprecated. Please use the getter or setter.")
+	UPROPERTY(EditAnywhere, BlueprintReadWrite, Getter, Setter, BlueprintSetter="SetClickMethod", Category="Interaction", AdvancedDisplay)
 	TEnumAsByte<EButtonClickMethod::Type> ClickMethod;
 
 	/** The type of touch action required by the user to trigger the buttons 'Click' */
-	UPROPERTY(EditAnywhere, BlueprintReadOnly, Category = "Interaction", AdvancedDisplay)
+	UE_DEPRECATED(5.1, "Direct access to TouchMethod is deprecated. Please use the getter or setter.")
+	UPROPERTY(EditAnywhere, BlueprintReadWrite, Getter, Setter, BlueprintSetter="SetTouchMethod", Category="Interaction", AdvancedDisplay)
 	TEnumAsByte<EButtonTouchMethod::Type> TouchMethod;
 
 	/** The type of keyboard/gamepad button press action required by the user to trigger the buttons 'Click' */
-	UPROPERTY(EditAnywhere, BlueprintReadOnly, Category = "Interaction", AdvancedDisplay)
+	UE_DEPRECATED(5.1, "Direct access to PressMethod is deprecated. Please use the getter or setter.")
+	UPROPERTY(EditAnywhere, BlueprintReadWrite, Getter, Setter, BlueprintSetter="SetPressMethod", Category="Interaction", AdvancedDisplay)
 	TEnumAsByte<EButtonPressMethod::Type> PressMethod;
 
 	/** Sometimes a button should only be mouse-clickable and never keyboard focusable. */
-	UPROPERTY(EditAnywhere, BlueprintReadOnly, Category = "Interaction")
+	UPROPERTY(EditAnywhere, BlueprintReadOnly, Category="Interaction")
 	bool IsFocusable;
 
 public:
@@ -96,13 +97,31 @@ public:
 	UFUNCTION(BlueprintCallable, Category="Widget")
 	void SetCheckedState(ECheckBoxState InCheckedState);
 
-	UFUNCTION(BlueprintCallable, Category = "Button")
+	/** Returns the local style. */
+	const FCheckBoxStyle& GetWidgetStyle() const;
+
+	/** Sets the style. */
+	void SetWidgetStyle(const FCheckBoxStyle& InStyle);
+
+	/** Returns the click method. */
+	EButtonClickMethod::Type GetClickMethod() const;
+
+	/** Sets the click method. */
+	UFUNCTION(BlueprintCallable, Category="Button")
 	void SetClickMethod(EButtonClickMethod::Type InClickMethod);
 
-	UFUNCTION(BlueprintCallable, Category = "Button")
+	/** Returns the touch method. */
+	EButtonTouchMethod::Type GetTouchMethod() const;
+
+	/** Sets the touch method. */
+	UFUNCTION(BlueprintCallable, Category="Button")
 	void SetTouchMethod(EButtonTouchMethod::Type InTouchMethod);
 
-	UFUNCTION(BlueprintCallable, Category = "Button")
+	/** Returns the press method. */
+	EButtonPressMethod::Type GetPressMethod() const;
+
+	/** Sets the press method. */
+	UFUNCTION(BlueprintCallable, Category="Button")
 	void SetPressMethod(EButtonPressMethod::Type InPressMethod);
 
 public:
@@ -130,7 +149,10 @@ protected:
 	//~ Begin UWidget Interface
 	virtual TSharedRef<SWidget> RebuildWidget() override;
 #if WITH_EDITOR
-	virtual TSharedRef<SWidget> RebuildDesignWidget(TSharedRef<SWidget> Content) override { return Content; }
+	virtual TSharedRef<SWidget> RebuildDesignWidget(TSharedRef<SWidget> Content) override
+	{
+		return Content;
+	}
 #endif
 	//~ End UWidget Interface
 
