UMG: Fix matching macro in TextBlock
#rnx

diff --git a/Engine/Source/Runtime/UMG/Private/Components/TextBlock.cpp b/Engine/Source/Runtime/UMG/Private/Components/TextBlock.cpp
--- a/Engine/Source/Runtime/UMG/Private/Components/TextBlock.cpp
+++ b/Engine/Source/Runtime/UMG/Private/Components/TextBlock.cpp
@@ -196,7 +196,6 @@ void UTextBlock::SetTextOverflowPolicy(ETextOverflowPolicy InOverflowPolicy)
 	TextOverflowPolicy = InOverflowPolicy;
 	SynchronizeProperties();
 }
-PRAGMA_DISABLE_DEPRECATION_WARNINGS
 
 UMaterialInstanceDynamic* UTextBlock::GetDynamicFontMaterial()
 {
@@ -241,6 +240,7 @@ UMaterialInstanceDynamic* UTextBlock::GetDynamicOutlineMaterial()
 
 	return nullptr;
 }
+PRAGMA_ENABLE_DEPRECATION_WARNINGS
 
 
 TSharedRef<SWidget> UTextBlock::RebuildWidget()
@@ -284,6 +284,7 @@ TSharedRef<SWidget> UTextBlock::RebuildWidget()
 	}
 }
 
+PRAGMA_DISABLE_DEPRECATION_WARNINGS
 EVisibility UTextBlock::GetTextWarningImageVisibility() const
 {
 	return Text.IsCultureInvariant() ? EVisibility::Visible : EVisibility::Collapsed;
@@ -376,6 +377,7 @@ TAttribute<FText> UTextBlock::GetDisplayText()
 {
 	return PROPERTY_BINDING(FText, Text);
 }
+PRAGMA_ENABLE_DEPRECATION_WARNINGS
 
 #if WITH_EDITOR
 
@@ -383,7 +385,7 @@ FString UTextBlock::GetLabelMetadata() const
 {
 	const int32 MaxSampleLength = 15;
 
-	FString TextStr = Text.ToString().Replace(TEXT("\n"), TEXT(" "));
+	FString TextStr = GetText().ToString().Replace(TEXT("\n"), TEXT(" "));
 	TextStr = TextStr.Len() <= MaxSampleLength ? TextStr : TextStr.Left(MaxSampleLength - 2) + TEXT("..");
 	return TEXT(" \"") + TextStr + TEXT("\"");
 }
@@ -402,11 +404,12 @@ const FText UTextBlock::GetPaletteCategory()
 
 void UTextBlock::OnCreationFromPalette()
 {
-	Text = LOCTEXT("TextBlockDefaultValue", "Text Block");
+	SetText(LOCTEXT("TextBlockDefaultValue", "Text Block"));
 }
 
 bool UTextBlock::CanEditChange(const FProperty* InProperty) const
 {
+PRAGMA_DISABLE_DEPRECATION_WARNINGS
 	if (bSimpleTextMode && InProperty)
 	{
 		static TArray<FName> InvalidPropertiesInSimpleMode =
@@ -425,6 +428,7 @@ bool UTextBlock::CanEditChange(const FProperty* InProperty) const
 	}
 
 	return Super::CanEditChange(InProperty);
+PRAGMA_ENABLE_DEPRECATION_WARNINGS
 }
 
 #endif //if WITH_EDITOR
