UMG: Update the EditableText widget with accessors.

diff --git a/Engine/Source/Runtime/Slate/Public/Widgets/Input/SEditableTextBox.h b/Engine/Source/Runtime/Slate/Public/Widgets/Input/SEditableTextBox.h
--- a/Engine/Source/Runtime/Slate/Public/Widgets/Input/SEditableTextBox.h
+++ b/Engine/Source/Runtime/Slate/Public/Widgets/Input/SEditableTextBox.h
@@ -198,6 +198,16 @@ public:
 	 */
 	void SetText( const TAttribute< FText >& InNewText );
 	
+	/**
+	 * Returns the Hint text string
+	 *
+	 * @return  Hint Text string
+	 */
+	FText GetHintText() const
+	{
+		return EditableText->GetHintText();
+	}
+
 	/** See the HintText attribute */
 	void SetHintText( const TAttribute< FText >& InHintText );
 	
diff --git a/Engine/Source/Runtime/UMG/Private/Components/EditableText.cpp b/Engine/Source/Runtime/UMG/Private/Components/EditableText.cpp
--- a/Engine/Source/Runtime/UMG/Private/Components/EditableText.cpp
+++ b/Engine/Source/Runtime/UMG/Private/Components/EditableText.cpp
@@ -50,6 +50,7 @@ UEditableText::UEditableText(const FObjectInitializer& ObjectInitializer)
 	}
 #endif // WITH_EDITOR
 
+	PRAGMA_DISABLE_DEPRECATION_WARNINGS
 	IsReadOnly = false;
 	IsPassword = false;
 	MinimumDesiredWidth = 0.0f;
@@ -63,6 +64,7 @@ UEditableText::UEditableText(const FObjectInitializer& ObjectInitializer)
 	VirtualKeyboardDismissAction = EVirtualKeyboardDismissAction::TextChangeOnDismiss;
 	Clipping = EWidgetClipping::ClipToBounds;
 	OverflowPolicy = ETextOverflowPolicy::Clip;
+	PRAGMA_ENABLE_DEPRECATION_WARNINGS
 
 #if WITH_EDITORONLY_DATA
 	AccessibleBehavior = ESlateAccessibleBehavior::Auto;
@@ -79,22 +81,24 @@ void UEditableText::ReleaseSlateResources(bool bReleaseChildren)
 
 TSharedRef<SWidget> UEditableText::RebuildWidget()
 {
-	MyEditableText = SNew( SEditableText )
-		.Style( &WidgetStyle )
-		.MinDesiredWidth( MinimumDesiredWidth )
-		.IsCaretMovedWhenGainFocus( IsCaretMovedWhenGainFocus )
-		.SelectAllTextWhenFocused( SelectAllTextWhenFocused )
-		.RevertTextOnEscape( RevertTextOnEscape )
-		.ClearKeyboardFocusOnCommit( ClearKeyboardFocusOnCommit )
-		.SelectAllTextOnCommit( SelectAllTextOnCommit )
-		.OnTextChanged( BIND_UOBJECT_DELEGATE( FOnTextChanged, HandleOnTextChanged ) )
-		.OnTextCommitted( BIND_UOBJECT_DELEGATE( FOnTextCommitted, HandleOnTextCommitted ) )
-		.VirtualKeyboardType( EVirtualKeyboardType::AsKeyboardType( KeyboardType.GetValue() ) )
+	PRAGMA_DISABLE_DEPRECATION_WARNINGS
+	MyEditableText = SNew(SEditableText)
+		.Style(&WidgetStyle)
+		.MinDesiredWidth(MinimumDesiredWidth)
+		.IsCaretMovedWhenGainFocus(IsCaretMovedWhenGainFocus)
+		.SelectAllTextWhenFocused(SelectAllTextWhenFocused)
+		.RevertTextOnEscape(RevertTextOnEscape)
+		.ClearKeyboardFocusOnCommit(ClearKeyboardFocusOnCommit)
+		.SelectAllTextOnCommit(SelectAllTextOnCommit)
+		.OnTextChanged(BIND_UOBJECT_DELEGATE(FOnTextChanged, HandleOnTextChanged))
+		.OnTextCommitted(BIND_UOBJECT_DELEGATE(FOnTextCommitted, HandleOnTextCommitted))
+		.VirtualKeyboardType(EVirtualKeyboardType::AsKeyboardType(KeyboardType.GetValue()))
 		.VirtualKeyboardOptions(VirtualKeyboardOptions)
 		.VirtualKeyboardTrigger(VirtualKeyboardTrigger)
 		.VirtualKeyboardDismissAction(VirtualKeyboardDismissAction)
 		.Justification(Justification)
 		.OverflowPolicy(OverflowPolicy);
+	PRAGMA_ENABLE_DEPRECATION_WARNINGS
 	
 	return MyEditableText.ToSharedRef();
 }
@@ -103,6 +107,9 @@ void UEditableText::SynchronizeProperties()
 {
 	Super::SynchronizeProperties();
 
+
+	PRAGMA_DISABLE_DEPRECATION_WARNINGS
+
 	TAttribute<FText> TextBinding = PROPERTY_BINDING(FText, Text);
 	TAttribute<FText> HintTextBinding = PROPERTY_BINDING(FText, HintText);
 
@@ -115,11 +122,16 @@ void UEditableText::SynchronizeProperties()
 	MyEditableText->SetJustification(Justification);
 	MyEditableText->SetOverflowPolicy(OverflowPolicy);
 
+	PRAGMA_ENABLE_DEPRECATION_WARNINGS
+
 	// TODO UMG Complete making all properties settable on SEditableText
 
 	ShapedTextOptions.SynchronizeShapedTextProperties(*MyEditableText);
 }
 
+
+PRAGMA_DISABLE_DEPRECATION_WARNINGS
+
 FText UEditableText::GetText() const
 {
 	if ( MyEditableText.IsValid() )
@@ -132,7 +144,16 @@ FText UEditableText::GetText() const
 
 void UEditableText::SetText(FText InText)
 {
+	// We detect if the Text is internal pointing to the same thing if so, nothing to do.
+	if (Text.IdenticalTo(InText))
+	{
+		return;
+	}
+
 	Text = InText;
+
+	BroadcastFieldValueChanged(FFieldNotificationClassDescriptor::Text);
+
 	if ( MyEditableText.IsValid() )
 	{
 		MyEditableText->SetText(Text);
@@ -148,6 +169,16 @@ void UEditableText::SetIsPassword(bool InbIsPassword)
 	}
 }
 
+FText UEditableText::GetHintText() const
+{
+	if (MyEditableText.IsValid())
+	{
+		return MyEditableText->GetHintText();
+	}
+
+	return HintText;
+}
+
 void UEditableText::SetHintText(FText InHintText)
 {
 	HintText = InHintText;
@@ -157,6 +188,25 @@ void UEditableText::SetHintText(FText InHintText)
 	}
 }
 
+float UEditableText::GetMinimumDesiredWidth() const
+{
+	return MinimumDesiredWidth;
+}
+
+void UEditableText::SetMinimumDesiredWidth(float InMinDesiredWidth)
+{
+	MinimumDesiredWidth = InMinDesiredWidth;
+	if (MyEditableText.IsValid())
+	{
+		MyEditableText->SetMinDesiredWidth(MinimumDesiredWidth);
+	}
+}
+
+bool UEditableText::GetIsReadOnly() const
+{
+	return IsReadOnly;
+}
+
 void UEditableText::SetIsReadOnly(bool InbIsReadyOnly)
 {
 	IsReadOnly = InbIsReadyOnly;
@@ -166,6 +216,18 @@ void UEditableText::SetIsReadOnly(bool InbIsReadyOnly)
 	}
 }
 
+
+bool UEditableText::GetIsPassword() const
+{
+	return IsPassword;
+}
+
+
+ETextJustify::Type UEditableText::GetJustification() const
+{
+	return Justification;
+}
+
 void UEditableText::SetJustification(ETextJustify::Type InJustification)
 {
 	Justification = InJustification;
@@ -175,6 +237,11 @@ void UEditableText::SetJustification(ETextJustify::Type InJustification)
 	}
 }
 
+ETextOverflowPolicy UEditableText::GetTextOverflowPolicy() const
+{
+	return OverflowPolicy;
+}
+
 void UEditableText::SetTextOverflowPolicy(ETextOverflowPolicy InOverflowPolicy)
 {
 	OverflowPolicy = InOverflowPolicy;
@@ -184,6 +251,8 @@ void UEditableText::SetTextOverflowPolicy(ETextOverflowPolicy InOverflowPolicy)
 	}
 }
 
+PRAGMA_ENABLE_DEPRECATION_WARNINGS
+
 void UEditableText::SetClearKeyboardFocusOnCommit(bool bInClearKeyboardFocusOnCommit)
 {
 	ClearKeyboardFocusOnCommit = bInClearKeyboardFocusOnCommit;
diff --git a/Engine/Source/Runtime/UMG/Public/Components/EditableText.h b/Engine/Source/Runtime/UMG/Public/Components/EditableText.h
--- a/Engine/Source/Runtime/UMG/Public/Components/EditableText.h
+++ b/Engine/Source/Runtime/UMG/Public/Components/EditableText.h
@@ -32,7 +32,8 @@ public:
 
 public:
 	/** The text content for this editable text box widget */
-	UPROPERTY(EditAnywhere, Category=Content)
+	UE_DEPRECATED(5.1, "Direct access to Text is deprecated. Please use the getter or setter.")
+	UPROPERTY(EditAnywhere, BlueprintReadWrite, Getter, Setter, BlueprintGetter = "GetText", BlueprintSetter = "SetText", FieldNotify, Category = "Content")
 	FText Text;
 
 	/** A bindable delegate to allow logic to drive the text of the widget */
@@ -42,7 +43,8 @@ public:
 public:
 
 	/** Hint text that appears when there is no text in the text box */
-	UPROPERTY(EditAnywhere, BlueprintReadOnly, Category=Content)
+	UE_DEPRECATED(5.1, "Direct access to Hint Text is deprecated. Please use the getter or setter.")
+	UPROPERTY(EditAnywhere, BlueprintReadWrite, Getter, Setter, BlueprintGetter = "GetHintText", BlueprintSetter = "SetHintText", Category = Content, meta = (MultiLine = "true"))
 	FText HintText;
 
 	/** A bindable delegate to allow logic to drive the hint text of the widget */
@@ -54,15 +56,18 @@ public:
 	FEditableTextStyle WidgetStyle;
 
 	/** Sets whether this text box can actually be modified interactively by the user */
-	UPROPERTY(EditAnywhere, Category=Appearance)
+	UE_DEPRECATED(5.1, "Direct access to IsReadOnly is deprecated. Please use the getter or setter.")
+	UPROPERTY(EditAnywhere, BlueprintReadWrite, Getter, Setter, Category = Appearance)
 	bool IsReadOnly;
 
 	/** Sets whether this text box is for storing a password */
-	UPROPERTY(EditAnywhere, Category=Appearance)
+	UE_DEPRECATED(5.1, "Direct access to IsPassword is deprecated. Please use the getter or setter.")
+	UPROPERTY(EditAnywhere, BlueprintReadWrite, Getter, Setter, BlueprintSetter = "SetIsPassword", Category = Appearance)
 	bool IsPassword;
 
-	/** Minimum width that a text block should be */
-	UPROPERTY(EditAnywhere, Category=Appearance)
+	/** The minimum desired Width for the text */
+	UE_DEPRECATED(5.1, "Direct access to MinimumDesiredWidth is deprecated. Please use the getter or setter.")
+	UPROPERTY(EditAnywhere, BlueprintReadWrite, Getter, Setter, BlueprintSetter = "SetMinimumDesiredWidth", Category = "Appearance")
 	float MinimumDesiredWidth;
 
 	/** Workaround as we lose focus when the auto completion closes. */
@@ -105,11 +110,13 @@ public:
 	EVirtualKeyboardDismissAction VirtualKeyboardDismissAction;
 	
 	/** How the text should be aligned with the margin. */
-	UPROPERTY(EditAnywhere, BlueprintReadWrite, BlueprintSetter=SetJustification, Category=Appearance)
+	UE_DEPRECATED(5.1, "Direct access to Justification is deprecated. Please use the getter or setter.")
+	UPROPERTY(EditAnywhere, BlueprintReadWrite, Getter, Setter, Blueprintgetter = GetJustification, BlueprintSetter = SetJustification, Category = Appearance)
 	TEnumAsByte<ETextJustify::Type> Justification;
 
-	/** Sets what should happen when text is clipped because the block does not have enough space */
-	UPROPERTY(EditAnywhere, BlueprintReadWrite, BlueprintSetter=SetTextOverflowPolicy, Category=Appearance)
+	/** Sets what happens to text that is clipped and doesn't fit within the clip rect for this widget */
+	UE_DEPRECATED(5.1, "Direct access to OverflowPolicy is deprecated. Please use the getter or setter.")
+	UPROPERTY(EditAnywhere, BlueprintReadWrite, Getter = "GetTextOverflowPolicy", Setter = "SetTextOverflowPolicy", BlueprintSetter = "SetTextOverflowPolicy", Category = "Clipping", AdvancedDisplay, meta = (DisplayName = "Overflow Policy"))
 	ETextOverflowPolicy OverflowPolicy;
 
 	/** Controls how the text within this widget should be shaped. */
@@ -141,18 +148,43 @@ public:
 	UFUNCTION(BlueprintCallable, Category = "Widget", meta=(DisplayName="SetText (Editable Text)"))
 	void SetText(FText InText);
 
+	bool GetIsPassword() const;
+
 	UFUNCTION(BlueprintCallable, Category = "Widget")
 	void SetIsPassword(UPARAM(DisplayName="IsPassword") bool InbIsPassword);
 
+	/** Gets the Hint text that appears when there is no text in the text box */
+	UFUNCTION(BlueprintCallable, Category = "Widget", meta = (DisplayName = "GetHintText (Editable Text)"))
+	FText GetHintText() const;
+
 	UFUNCTION(BlueprintCallable, Category = "Widget")
 	void SetHintText(FText InHintText);
 
+	/** @return the minimum desired width for this text box */
+	float GetMinimumDesiredWidth() const;
+
+	/**
+	 *  Set the minimum desired width for this text box
+	 *
+	 *  @param InMinDesiredWidth new minimum desired width
+	*/
+	UFUNCTION(BlueprintCallable, Category = "Appearance")
+	void SetMinimumDesiredWidth(float InMinDesiredWidth);
+
+	bool GetIsReadOnly() const;
+
 	UFUNCTION(BlueprintCallable, Category = "Widget", meta=(DisplayName="SetIsReadOnly (Editable Text)"))
 	void SetIsReadOnly(UPARAM(DisplayName="ReadyOnly") bool InbIsReadyOnly);
 
+	UFUNCTION(BlueprintCallable, Category = "Widget")
+	ETextJustify::Type GetJustification() const;
+
 	UFUNCTION(BlueprintSetter)
 	void SetJustification(ETextJustify::Type InJustification);
 
+	/** @return the text overflow policy for this text block. */
+	ETextOverflowPolicy GetTextOverflowPolicy() const;
+
 	UFUNCTION(BlueprintSetter)
 	void SetTextOverflowPolicy(ETextOverflowPolicy InOverflowPolicy);
 
