UMG: Update the Border widget with accessors.

diff --git a/Engine/Plugins/Experimental/CommonUI/Source/CommonUI/Private/CommonBorder.cpp b/Engine/Plugins/Experimental/CommonUI/Source/CommonUI/Private/CommonBorder.cpp
--- a/Engine/Plugins/Experimental/CommonUI/Source/CommonUI/Private/CommonBorder.cpp
+++ b/Engine/Plugins/Experimental/CommonUI/Source/CommonUI/Private/CommonBorder.cpp
@@ -125,10 +125,10 @@ void UCommonBorder::SafeAreaUpdated()
 		}
 
 		FMargin NewMargin;
-		NewMargin.Left = FMath::Max(MinimumPadding.Left, Padding.Left - SafeMargin.Left);
-		NewMargin.Right = FMath::Max(MinimumPadding.Right, Padding.Right - SafeMargin.Right);
-		NewMargin.Top = FMath::Max(MinimumPadding.Top, Padding.Top - SafeMargin.Top);
-		NewMargin.Bottom = FMath::Max(MinimumPadding.Bottom, Padding.Bottom - SafeMargin.Bottom);
+		NewMargin.Left = FMath::Max(MinimumPadding.Left, GetPadding().Left - SafeMargin.Left);
+		NewMargin.Right = FMath::Max(MinimumPadding.Right, GetPadding().Right - SafeMargin.Right);
+		NewMargin.Top = FMath::Max(MinimumPadding.Top, GetPadding().Top - SafeMargin.Top);
+		NewMargin.Bottom = FMath::Max(MinimumPadding.Bottom, GetPadding().Bottom - SafeMargin.Bottom);
 		
 		MyBorder->SetPadding(NewMargin);
 	}
diff --git a/Engine/Plugins/Experimental/CommonUI/Source/CommonUI/Private/CommonCustomNavigation.cpp b/Engine/Plugins/Experimental/CommonUI/Source/CommonUI/Private/CommonCustomNavigation.cpp
--- a/Engine/Plugins/Experimental/CommonUI/Source/CommonUI/Private/CommonCustomNavigation.cpp
+++ b/Engine/Plugins/Experimental/CommonUI/Source/CommonUI/Private/CommonCustomNavigation.cpp
@@ -66,7 +66,7 @@ private:
 UCommonCustomNavigation::UCommonCustomNavigation(const FObjectInitializer& ObjectInitializer)
 	: Super(ObjectInitializer)
 {
-	Padding = FMargin(0);
+	SetPadding(FMargin(0.f));
 }
 
 TSharedRef<SWidget> UCommonCustomNavigation::RebuildWidget()
diff --git a/Engine/Plugins/Experimental/CommonUI/Source/CommonUI/Private/CommonHardwareVisibilityBorder.cpp b/Engine/Plugins/Experimental/CommonUI/Source/CommonUI/Private/CommonHardwareVisibilityBorder.cpp
--- a/Engine/Plugins/Experimental/CommonUI/Source/CommonUI/Private/CommonHardwareVisibilityBorder.cpp
+++ b/Engine/Plugins/Experimental/CommonUI/Source/CommonUI/Private/CommonHardwareVisibilityBorder.cpp
@@ -15,7 +15,7 @@ UCommonHardwareVisibilityBorder::UCommonHardwareVisibilityBorder(const FObjectIn
 	, VisibleType(ESlateVisibility::SelfHitTestInvisible)
 	, HiddenType(ESlateVisibility::Collapsed)
 {
-	Padding = FMargin(0,0,0,0);
+	SetPadding(FMargin(0.f,0.f,0.f,0.f));
 }
 
 void UCommonHardwareVisibilityBorder::OnWidgetRebuilt()
diff --git a/Engine/Source/Runtime/UMG/Private/Components/Border.cpp b/Engine/Source/Runtime/UMG/Private/Components/Border.cpp
--- a/Engine/Source/Runtime/UMG/Private/Components/Border.cpp
+++ b/Engine/Source/Runtime/UMG/Private/Components/Border.cpp
@@ -21,6 +21,7 @@ UBorder::UBorder(const FObjectInitializer& ObjectInitializer)
 {
 	bIsVariable = false;
 
+PRAGMA_DISABLE_DEPRECATION_WARNINGS
 	ContentColorAndOpacity = FLinearColor::White;
 	BrushColor = FLinearColor::White;
 
@@ -32,6 +33,7 @@ UBorder::UBorder(const FObjectInitializer& ObjectInitializer)
 	DesiredSizeScale = FVector2D(1, 1);
 
 	bShowEffectWhenDisabled = true;
+PRAGMA_ENABLE_DEPRECATION_WARNINGS
 }
 
 void UBorder::ReleaseSlateResources(bool bReleaseChildren)
@@ -58,6 +60,7 @@ void UBorder::SynchronizeProperties()
 {
 	Super::SynchronizeProperties();
 	
+PRAGMA_DISABLE_DEPRECATION_WARNINGS
 	TAttribute<FLinearColor> ContentColorAndOpacityBinding = PROPERTY_BINDING(FLinearColor, ContentColorAndOpacity);
 	TAttribute<FSlateColor> BrushColorBinding = OPTIONAL_BINDING_CONVERT(FLinearColor, BrushColor, FSlateColor, ConvertLinearColorToSlateColor);
 	TAttribute<const FSlateBrush*> ImageBinding = OPTIONAL_BINDING_CONVERT(FSlateBrush, Background, const FSlateBrush*, ConvertImage);
@@ -70,7 +73,8 @@ void UBorder::SynchronizeProperties()
 	
 	MyBorder->SetDesiredSizeScale(DesiredSizeScale);
 	MyBorder->SetShowEffectWhenDisabled(bShowEffectWhenDisabled != 0);
-	
+PRAGMA_ENABLE_DEPRECATION_WARNINGS
+
 	MyBorder->SetOnMouseButtonDown(BIND_UOBJECT_DELEGATE(FPointerEventHandler, HandleMouseButtonDown));
 	MyBorder->SetOnMouseButtonUp(BIND_UOBJECT_DELEGATE(FPointerEventHandler, HandleMouseButtonUp));
 	MyBorder->SetOnMouseMove(BIND_UOBJECT_DELEGATE(FPointerEventHandler, HandleMouseMove));
@@ -84,12 +88,14 @@ UClass* UBorder::GetSlotClass() const
 
 void UBorder::OnSlotAdded(UPanelSlot* InSlot)
 {
+PRAGMA_DISABLE_DEPRECATION_WARNINGS
 	// Copy the content properties into the new slot so that it matches what has been setup
 	// so far by the user.
 	UBorderSlot* BorderSlot = CastChecked<UBorderSlot>(InSlot);
 	BorderSlot->Padding = Padding;
 	BorderSlot->HorizontalAlignment = HorizontalAlignment;
 	BorderSlot->VerticalAlignment = VerticalAlignment;
+PRAGMA_ENABLE_DEPRECATION_WARNINGS
 
 	// Add the child to the live slot if it already exists
 	if ( MyBorder.IsValid() )
@@ -108,6 +114,16 @@ void UBorder::OnSlotRemoved(UPanelSlot* InSlot)
 	}
 }
 
+PRAGMA_DISABLE_DEPRECATION_WARNINGS
+FLinearColor UBorder::GetContentColorAndOpacity() const
+{
+	if (MyBorder.IsValid())
+	{
+		return MyBorder->GetColorAndOpacity();
+	}
+	return ContentColorAndOpacity;
+}
+
 void UBorder::SetContentColorAndOpacity(FLinearColor Color)
 {
 	ContentColorAndOpacity = Color;
@@ -117,6 +133,11 @@ void UBorder::SetContentColorAndOpacity(FLinearColor Color)
 	}
 }
 
+FMargin UBorder::GetPadding() const
+{
+	return Padding;
+}
+
 void UBorder::SetPadding(FMargin InPadding)
 {
 	Padding = InPadding;
@@ -126,15 +147,25 @@ void UBorder::SetPadding(FMargin InPadding)
 	}
 }
 
+EHorizontalAlignment UBorder::GetHorizontalAlignment() const
+{
+	return HorizontalAlignment;
+}
+
 void UBorder::SetHorizontalAlignment(EHorizontalAlignment InHorizontalAlignment)
 {
 	HorizontalAlignment = InHorizontalAlignment;
-	if ( MyBorder.IsValid() )
+	if (MyBorder.IsValid())
 	{
 		MyBorder->SetHAlign(InHorizontalAlignment);
 	}
 }
 
+EVerticalAlignment UBorder::GetVerticalAlignment() const
+{
+	return VerticalAlignment;
+}
+
 void UBorder::SetVerticalAlignment(EVerticalAlignment InVerticalAlignment)
 {
 	VerticalAlignment = InVerticalAlignment;
@@ -144,6 +175,11 @@ void UBorder::SetVerticalAlignment(EVerticalAlignment InVerticalAlignment)
 	}
 }
 
+FLinearColor UBorder::GetBrushColor() const
+{
+	return BrushColor;
+}
+
 void UBorder::SetBrushColor(FLinearColor Color)
 {
 	BrushColor = Color;
@@ -152,6 +188,7 @@ void UBorder::SetBrushColor(FLinearColor Color)
 		MyBorder->SetBorderBackgroundColor(Color);
 	}
 }
+PRAGMA_ENABLE_DEPRECATION_WARNINGS
 
 FReply UBorder::HandleMouseButtonDown(const FGeometry& Geometry, const FPointerEvent& MouseEvent)
 {
@@ -240,6 +277,13 @@ void UBorder::SetBrushFromMaterial(UMaterialInterface* Material)
 	}
 }
 
+
+PRAGMA_DISABLE_DEPRECATION_WARNINGS
+bool UBorder::GetShowEffectWhenDisabled() const
+{
+	return bShowEffectWhenDisabled;
+}
+
 void UBorder::SetShowEffectWhenDisabled(bool bInShowEffectWhenDisabled)
 {
 	bShowEffectWhenDisabled = bInShowEffectWhenDisabled;
@@ -248,6 +292,7 @@ void UBorder::SetShowEffectWhenDisabled(bool bInShowEffectWhenDisabled)
 		MyBorder->SetShowEffectWhenDisabled(bShowEffectWhenDisabled != 0);
 	}
 }
+PRAGMA_ENABLE_DEPRECATION_WARNINGS
 
 UMaterialInstanceDynamic* UBorder::GetDynamicMaterial()
 {
@@ -278,6 +323,12 @@ UMaterialInstanceDynamic* UBorder::GetDynamicMaterial()
 	return nullptr;
 }
 
+PRAGMA_DISABLE_DEPRECATION_WARNINGS
+FVector2D UBorder::GetDesiredSizeScale() const
+{
+	return DesiredSizeScale;
+}
+
 void UBorder::SetDesiredSizeScale(FVector2D InScale)
 {
 	DesiredSizeScale = InScale;
@@ -286,6 +337,7 @@ void UBorder::SetDesiredSizeScale(FVector2D InScale)
 		MyBorder->SetDesiredSizeScale(InScale);
 	}
 }
+PRAGMA_ENABLE_DEPRECATION_WARNINGS
 
 const FSlateBrush* UBorder::ConvertImage(TAttribute<FSlateBrush> InImageAsset) const
 {
diff --git a/Engine/Source/Runtime/UMG/Public/Components/Border.h b/Engine/Source/Runtime/UMG/Public/Components/Border.h
--- a/Engine/Source/Runtime/UMG/Public/Components/Border.h
+++ b/Engine/Source/Runtime/UMG/Public/Components/Border.h
@@ -31,21 +31,24 @@ class UMG_API UBorder : public UContentWidget
 	GENERATED_UCLASS_BODY()
 
 public:
-
 	/** The alignment of the content horizontally. */
-	UPROPERTY(EditAnywhere, BlueprintReadOnly, Category="Content")
+	UE_DEPRECATED(5.1, "Direct access to HorizontalAlignment is deprecated. Please use the getter or setter.")
+	UPROPERTY(EditAnywhere, BlueprintReadWrite, Getter, Setter, BlueprintSetter="SetHorizontalAlignment", Category="Content")
 	TEnumAsByte<EHorizontalAlignment> HorizontalAlignment;
 
 	/** The alignment of the content vertically. */
-	UPROPERTY(EditAnywhere, BlueprintReadOnly, Category="Content")
+	UE_DEPRECATED(5.1, "Direct access to VerticalAlignment is deprecated. Please use the getter or setter.")
+	UPROPERTY(EditAnywhere, BlueprintReadWrite, Getter, Setter, BlueprintSetter="SetVerticalAlignment", Category="Content")
 	TEnumAsByte<EVerticalAlignment> VerticalAlignment;
 
 	/** Whether or not to show the disabled effect when this border is disabled */
-	UPROPERTY(EditAnywhere, BlueprintReadOnly, Category=Appearance, AdvancedDisplay)
+	UE_DEPRECATED(5.1, "Direct access to bShowEffectWhenDisabled is deprecated. Please use the getter or setter.")
+	UPROPERTY(EditAnywhere, BlueprintReadWrite, Getter="GetShowEffectWhenDisabled", Setter="SetShowEffectWhenDisabled", BlueprintSetter="SetShowEffectWhenDisabled", Category="Appearance", AdvancedDisplay)
 	uint8 bShowEffectWhenDisabled:1;
 
 	/** Color and opacity multiplier of content in the border */
-	UPROPERTY(EditAnywhere, BlueprintReadOnly, Category="Content", meta=( sRGB="true" ))
+	UE_DEPRECATED(5.1, "Direct access to ContentColorAndOpacity is deprecated. Please use the getter or setter.")
+	UPROPERTY(EditAnywhere, BlueprintReadWrite, Getter, Setter, BlueprintSetter="SetContentColorAndOpacity", Category="Content", meta=( sRGB="true" ))
 	FLinearColor ContentColorAndOpacity;
 
 	/** A bindable delegate for the ContentColorAndOpacity. */
@@ -53,7 +56,8 @@ public:
 	FGetLinearColor ContentColorAndOpacityDelegate;
 
 	/** The padding area between the slot and the content it contains. */
-	UPROPERTY(EditAnywhere, BlueprintReadOnly, Category="Content")
+	UE_DEPRECATED(5.1, "Direct access to Padding is deprecated. Please use the getter or setter.")
+	UPROPERTY(EditAnywhere, BlueprintReadWrite, Getter, Setter, BlueprintSetter="SetPadding", Category="Content")
 	FMargin Padding;
 
 	/** Brush to drag as the background */
@@ -65,7 +69,8 @@ public:
 	FGetSlateBrush BackgroundDelegate;
 
 	/** Color and opacity of the actual border image */
-	UPROPERTY(EditAnywhere, BlueprintReadOnly, Category=Appearance, meta=( sRGB="true" ))
+	UE_DEPRECATED(5.1, "Direct access to BrushColor is deprecated. Please use the getter or setter.")
+	UPROPERTY(EditAnywhere, BlueprintReadWrite, Getter, Setter, BlueprintSetter="SetBrushColor", Category="Appearance", meta = (sRGB="true"))
 	FLinearColor BrushColor;
 
 	/** A bindable delegate for the BrushColor. */
@@ -78,11 +83,12 @@ public:
 	 * Note: if the parent widget is set up to ignore this widget's desired size,
 	 * then changing this value will have no effect.
 	 */
-	UPROPERTY(EditAnywhere, BlueprintReadOnly, Category = Appearance)
+	UE_DEPRECATED(5.1, "Direct access to DesiredSizeScale is deprecated. Please use the getter or setter.")
+	UPROPERTY(EditAnywhere, BlueprintReadWrite, Getter, Setter, BlueprintSetter="SetDesiredSizeScale", Category="Appearance")
 	FVector2D DesiredSizeScale;
 
 	/** Flips the background image if the localization's flow direction is RightToLeft */
-	UPROPERTY(EditAnywhere, Category = "Localization")
+	UPROPERTY(EditAnywhere, Category="Localization")
 	bool bFlipForRightToLeftFlowDirection;
 
 public:
@@ -100,20 +106,31 @@ public:
 	FOnPointerEvent OnMouseDoubleClickEvent;
 
 public:
+	/**  */
+	FLinearColor GetContentColorAndOpacity() const;
 
 	/**  */
 	UFUNCTION(BlueprintCallable, Category="Appearance")
 	void SetContentColorAndOpacity(FLinearColor InContentColorAndOpacity);
 
+	FMargin GetPadding() const;
+
 	UFUNCTION(BlueprintCallable, Category="Appearance")
 	void SetPadding(FMargin InPadding);
 
+	EHorizontalAlignment GetHorizontalAlignment() const;
+
 	UFUNCTION(BlueprintCallable, Category="Appearance")
 	void SetHorizontalAlignment(EHorizontalAlignment InHorizontalAlignment);
 
+	EVerticalAlignment GetVerticalAlignment() const;
+
 	UFUNCTION(BlueprintCallable, Category="Appearance")
 	void SetVerticalAlignment(EVerticalAlignment InVerticalAlignment);
 
+	/**  */
+	FLinearColor GetBrushColor() const;
+
 	/**  */
 	UFUNCTION(BlueprintCallable, Category="Appearance")
 	void SetBrushColor(FLinearColor InBrushColor);
@@ -134,6 +151,9 @@ public:
 	UFUNCTION(BlueprintCallable, Category="Appearance")
 	void SetBrushFromMaterial(UMaterialInterface* Material);
 
+	/**  */
+	bool GetShowEffectWhenDisabled() const;
+
 	/**  */
 	UFUNCTION(BlueprintCallable, Category="Appearance")
 	void SetShowEffectWhenDisabled(bool bInShowEffectWhenDisabled);
@@ -143,11 +163,14 @@ public:
 	UMaterialInstanceDynamic* GetDynamicMaterial();
 
 public:
+	/** Gets the DesireSizeScale of this border. */
+	FVector2D GetDesiredSizeScale() const;
+
 	/**
-	* Sets the DesireSizeScale of this border.
-	*
-	* @param InScale    The X and Y multipliers for the desired size
-	*/
+	 * Sets the DesireSizeScale of this border.
+	 *
+	 * @param InScale    The X and Y multipliers for the desired size
+	 */
 	UFUNCTION(BlueprintCallable, Category = "Appearance")
 	void SetDesiredSizeScale(FVector2D InScale);
 
