UMG: Update the Widget with accessors.

diff --git a/Engine/Plugins/Experimental/CommonUI/Source/CommonUI/Private/CommonBorder.cpp b/Engine/Plugins/Experimental/CommonUI/Source/CommonUI/Private/CommonBorder.cpp
--- a/Engine/Plugins/Experimental/CommonUI/Source/CommonUI/Private/CommonBorder.cpp
+++ b/Engine/Plugins/Experimental/CommonUI/Source/CommonUI/Private/CommonBorder.cpp
@@ -26,7 +26,7 @@ UCommonBorder::UCommonBorder(const FObjectInitializer& ObjectInitializer)
 {
 	Background.DrawAs = ESlateBrushDrawType::NoDrawType;
 
-	Visibility = ESlateVisibility::SelfHitTestInvisible;
+	SetVisibility(ESlateVisibility::SelfHitTestInvisible);
 }
 
 void UCommonBorder::PostLoad()
diff --git a/Engine/Plugins/Experimental/CommonUI/Source/CommonUI/Private/CommonButtonBase.cpp b/Engine/Plugins/Experimental/CommonUI/Source/CommonUI/Private/CommonButtonBase.cpp
--- a/Engine/Plugins/Experimental/CommonUI/Source/CommonUI/Private/CommonButtonBase.cpp
+++ b/Engine/Plugins/Experimental/CommonUI/Source/CommonUI/Private/CommonButtonBase.cpp
@@ -434,7 +434,7 @@ void UCommonButtonBase::SetIsEnabled(bool bInIsEnabled)
 	else
 	{
 		// Change the underlying enabled bool but do not call the case because we don't want to propogate it to the underlying SWidget
-		bIsEnabled = bInIsEnabled;
+		SetIsEnabled(bInIsEnabled);
 		DisableButton();
 	}
 }
diff --git a/Engine/Plugins/Experimental/CommonUI/Source/CommonUI/Private/CommonLazyWidget.cpp b/Engine/Plugins/Experimental/CommonUI/Source/CommonUI/Private/CommonLazyWidget.cpp
--- a/Engine/Plugins/Experimental/CommonUI/Source/CommonUI/Private/CommonLazyWidget.cpp
+++ b/Engine/Plugins/Experimental/CommonUI/Source/CommonUI/Private/CommonLazyWidget.cpp
@@ -9,7 +9,7 @@ UCommonLazyWidget::UCommonLazyWidget(const FObjectInitializer& Initializer)
 	: Super(Initializer)
 {
 	LoadingBackgroundBrush.DrawAs = ESlateBrushDrawType::NoDrawType;
-	Visibility = ESlateVisibility::SelfHitTestInvisible;
+	SetVisibility(ESlateVisibility::SelfHitTestInvisible);
 }
 
 TSharedRef<SWidget> UCommonLazyWidget::RebuildWidget()
diff --git a/Engine/Plugins/Experimental/CommonUI/Source/CommonUI/Private/CommonListView.cpp b/Engine/Plugins/Experimental/CommonUI/Source/CommonUI/Private/CommonListView.cpp
--- a/Engine/Plugins/Experimental/CommonUI/Source/CommonUI/Private/CommonListView.cpp
+++ b/Engine/Plugins/Experimental/CommonUI/Source/CommonUI/Private/CommonListView.cpp
@@ -12,7 +12,7 @@
 UCommonListView::UCommonListView(const FObjectInitializer& ObjectInitializer)
 	: Super(ObjectInitializer)
 {
-	Visibility = ESlateVisibility::Visible;
+	SetVisibility(ESlateVisibility::Visible);
 	bEnableScrollAnimation = true;
 }
 
diff --git a/Engine/Plugins/Experimental/CommonUI/Source/CommonUI/Private/CommonLoadGuard.cpp b/Engine/Plugins/Experimental/CommonUI/Source/CommonUI/Private/CommonLoadGuard.cpp
--- a/Engine/Plugins/Experimental/CommonUI/Source/CommonUI/Private/CommonLoadGuard.cpp
+++ b/Engine/Plugins/Experimental/CommonUI/Source/CommonUI/Private/CommonLoadGuard.cpp
@@ -271,7 +271,7 @@ UCommonLoadGuard::UCommonLoadGuard(const FObjectInitializer& Initializer)
 	, ThrobberAlignment(EHorizontalAlignment::HAlign_Center)
 	, LoadingText(NSLOCTEXT("LoadGuard", "Loading", "Loading..."))
 {
-	Visibility = ESlateVisibility::SelfHitTestInvisible;
+	SetVisibility(ESlateVisibility::SelfHitTestInvisible);
 
 	// Default to not showing the loading BG brush - it's an opt-in kind of thing
 	LoadingBackgroundBrush.DrawAs = ESlateBrushDrawType::NoDrawType;
diff --git a/Engine/Plugins/Experimental/CommonUI/Source/CommonUI/Private/CommonRichTextBlock.cpp b/Engine/Plugins/Experimental/CommonUI/Source/CommonUI/Private/CommonRichTextBlock.cpp
--- a/Engine/Plugins/Experimental/CommonUI/Source/CommonUI/Private/CommonRichTextBlock.cpp
+++ b/Engine/Plugins/Experimental/CommonUI/Source/CommonUI/Private/CommonRichTextBlock.cpp
@@ -565,9 +565,9 @@ TSharedRef<SWidget> UCommonRichTextBlock::RebuildWidget()
 
 	// If the clipping mode is the default, but we're using a scrolling style,
 	// we need to switch over to a clip to bounds style.
-	if (Clipping == EWidgetClipping::Inherit)
+	if (GetClipping() == EWidgetClipping::Inherit)
 	{
-		Clipping = EWidgetClipping::OnDemand;
+		SetClipping(EWidgetClipping::OnDemand);
 	}
 
 	MyTextScroller =
diff --git a/Engine/Plugins/Experimental/CommonUI/Source/CommonUI/Private/CommonTextBlock.cpp b/Engine/Plugins/Experimental/CommonUI/Source/CommonUI/Private/CommonTextBlock.cpp
--- a/Engine/Plugins/Experimental/CommonUI/Source/CommonUI/Private/CommonTextBlock.cpp
+++ b/Engine/Plugins/Experimental/CommonUI/Source/CommonUI/Private/CommonTextBlock.cpp
@@ -276,8 +276,8 @@ void UCommonTextStyle::ApplyToTextBlock(const TSharedRef<STextBlock>& TextBlock)
 UCommonTextBlock::UCommonTextBlock(const FObjectInitializer& ObjectInitializer)
 	: Super(ObjectInitializer)
 {
-	Visibility = ESlateVisibility::SelfHitTestInvisible;
-	Clipping = EWidgetClipping::Inherit;
+	SetVisibility(ESlateVisibility::SelfHitTestInvisible);
+	SetClipping(EWidgetClipping::Inherit);
 }
 
 void UCommonTextBlock::PostInitProperties()
@@ -639,9 +639,9 @@ bool UCommonTextBlock::CanEditChange(const FProperty* InProperty) const
 {
 	if (Super::CanEditChange(InProperty))
 	{
+PRAGMA_DISABLE_DEPRECATION_WARNINGS
 		if (const UCommonTextStyle* TextStyle = GetStyleCDO())
 		{
-PRAGMA_DISABLE_DEPRECATION_WARNINGS
 			static TArray<FName> InvalidPropertiesWithStyle =
 			{
 				GET_MEMBER_NAME_CHECKED(UCommonTextBlock, Font),
@@ -652,15 +652,14 @@ PRAGMA_DISABLE_DEPRECATION_WARNINGS
 				GET_MEMBER_NAME_CHECKED(UCommonTextBlock, ShadowOffset),
 				GET_MEMBER_NAME_CHECKED(UCommonTextBlock, ShadowColorAndOpacity)
 			};
-PRAGMA_ENABLE_DEPRECATION_WARNINGS
+
 			return !InvalidPropertiesWithStyle.Contains(InProperty->GetFName());
 		}
-
 		if (bAutoCollapseWithEmptyText && InProperty->GetFName() == GET_MEMBER_NAME_CHECKED(UCommonTextBlock, Visibility))
 		{
 			return false;
 		}
-
+PRAGMA_ENABLE_DEPRECATION_WARNINGS
 		return true;
 	}
 
@@ -679,9 +678,9 @@ TSharedRef<SWidget> UCommonTextBlock::RebuildWidget()
 
 	// If the clipping mode is the default, but we're using a scrolling style,
 	// we need to switch over to a clip to bounds style.
-	if (Clipping == EWidgetClipping::Inherit)
+	if (GetClipping() == EWidgetClipping::Inherit)
 	{
-		Clipping = EWidgetClipping::OnDemand;
+		SetClipping(EWidgetClipping::OnDemand);
 	}
 
 	// clang-format off
diff --git a/Engine/Plugins/Experimental/CommonUI/Source/CommonUI/Private/CommonWidgetCarousel.cpp b/Engine/Plugins/Experimental/CommonUI/Source/CommonUI/Private/CommonWidgetCarousel.cpp
--- a/Engine/Plugins/Experimental/CommonUI/Source/CommonUI/Private/CommonWidgetCarousel.cpp
+++ b/Engine/Plugins/Experimental/CommonUI/Source/CommonUI/Private/CommonWidgetCarousel.cpp
@@ -9,7 +9,7 @@ UCommonWidgetCarousel::UCommonWidgetCarousel(const FObjectInitializer& ObjectIni
 	: Super(ObjectInitializer)
 {
 	bIsVariable = true;
-	Clipping = EWidgetClipping::ClipToBounds;
+	SetClipping(EWidgetClipping::ClipToBounds);
 }
 
 void UCommonWidgetCarousel::ReleaseSlateResources(bool bReleaseChildren)
diff --git a/Engine/Plugins/Experimental/CommonUI/Source/CommonUI/Private/Widgets/CommonActivatableWidgetContainer.cpp b/Engine/Plugins/Experimental/CommonUI/Source/CommonUI/Private/Widgets/CommonActivatableWidgetContainer.cpp
--- a/Engine/Plugins/Experimental/CommonUI/Source/CommonUI/Private/Widgets/CommonActivatableWidgetContainer.cpp
+++ b/Engine/Plugins/Experimental/CommonUI/Source/CommonUI/Private/Widgets/CommonActivatableWidgetContainer.cpp
@@ -32,7 +32,7 @@ UCommonActivatableWidgetContainerBase::UCommonActivatableWidgetContainerBase(con
 	: Super(Initializer)
 	, GeneratedWidgetsPool(*this)
 {
-	Visibility = ESlateVisibility::Collapsed;
+	SetVisibility(ESlateVisibility::Collapsed);
 }
 
 void UCommonActivatableWidgetContainerBase::AddWidgetInstance(UCommonActivatableWidget& ActivatableWidget)
diff --git a/Engine/Source/Editor/UMGEditor/Private/Designer/SDesignerView.cpp b/Engine/Source/Editor/UMGEditor/Private/Designer/SDesignerView.cpp
--- a/Engine/Source/Editor/UMGEditor/Private/Designer/SDesignerView.cpp
+++ b/Engine/Source/Editor/UMGEditor/Private/Designer/SDesignerView.cpp
@@ -1778,7 +1778,7 @@ FReply SDesignerView::OnMouseMove(const FGeometry& MyGeometry, const FPointerEve
 						{
 							const FSlateRenderTransform& AbsoluteToLocalTransform = Inverse(ParentGeometry.GetAccumulatedRenderTransform());
 
-							FWidgetTransform WidgetRenderTransform = WidgetPreview->RenderTransform;
+							FWidgetTransform WidgetRenderTransform = WidgetPreview->GetRenderTransform();
 							WidgetRenderTransform.Translation += AbsoluteToLocalTransform.TransformVector(MouseEvent.GetCursorDelta());
 
 							static const FName RenderTransformName(TEXT("RenderTransform"));
diff --git a/Engine/Source/Editor/UMGEditor/Private/Designer/STransformHandle.cpp b/Engine/Source/Editor/UMGEditor/Private/Designer/STransformHandle.cpp
--- a/Engine/Source/Editor/UMGEditor/Private/Designer/STransformHandle.cpp
+++ b/Engine/Source/Editor/UMGEditor/Private/Designer/STransformHandle.cpp
@@ -138,7 +138,7 @@ FReply STransformHandle::OnMouseMove(const FGeometry& MyGeometry, const FPointer
 		ETransformMode::Type TransformMode = Designer->GetTransformMode();
 		if ( TransformMode == ETransformMode::Render )
 		{
-			FWidgetTransform PreviewRenderTransform = Preview->RenderTransform;
+			FWidgetTransform PreviewRenderTransform = Preview->GetRenderTransform();
 
 			static const FName RenderTransformName(TEXT("RenderTransform"));
 
diff --git a/Engine/Source/Editor/UMGEditor/Private/Hierarchy/SHierarchyViewItem.cpp b/Engine/Source/Editor/UMGEditor/Private/Hierarchy/SHierarchyViewItem.cpp
--- a/Engine/Source/Editor/UMGEditor/Private/Hierarchy/SHierarchyViewItem.cpp
+++ b/Engine/Source/Editor/UMGEditor/Private/Hierarchy/SHierarchyViewItem.cpp
@@ -744,7 +744,7 @@ bool FHierarchyRoot::DoesWidgetOverrideFlowDirection() const
 	TSharedPtr<FWidgetBlueprintEditor> BPEd = BlueprintEditor.Pin();
 	if (UWidget* Default = BPEd->GetWidgetBlueprintObj()->GeneratedClass->GetDefaultObject<UWidget>())
 	{
-		return Default->FlowDirectionPreference != EFlowDirectionPreference::Inherit;
+		return Default->GetFlowDirectionPreference() != EFlowDirectionPreference::Inherit;
 	}
 
 	return false;
diff --git a/Engine/Source/Editor/UMGEditor/Private/Hierarchy/SHierarchyViewItem.h b/Engine/Source/Editor/UMGEditor/Private/Hierarchy/SHierarchyViewItem.h
--- a/Engine/Source/Editor/UMGEditor/Private/Hierarchy/SHierarchyViewItem.h
+++ b/Engine/Source/Editor/UMGEditor/Private/Hierarchy/SHierarchyViewItem.h
@@ -282,7 +282,7 @@ public:
 		UWidget* TemplateWidget = Item.GetTemplate();
 		if (TemplateWidget)
 		{
-			return TemplateWidget->FlowDirectionPreference != EFlowDirectionPreference::Inherit;
+			return TemplateWidget->GetFlowDirectionPreference() != EFlowDirectionPreference::Inherit;
 		}
 
 		return false;
diff --git a/Engine/Source/Runtime/UMG/Private/Animation/MovieSceneUMGComponentTypes.cpp b/Engine/Source/Runtime/UMG/Private/Animation/MovieSceneUMGComponentTypes.cpp
--- a/Engine/Source/Runtime/UMG/Private/Animation/MovieSceneUMGComponentTypes.cpp
+++ b/Engine/Source/Runtime/UMG/Private/Animation/MovieSceneUMGComponentTypes.cpp
@@ -54,7 +54,7 @@ static void SetRenderOpacity(UObject* Object, bool bIsDouble, float InRenderOpac
 
 static FIntermediateWidgetTransform GetRenderTransform(const UObject* Object)
 {
-	FWidgetTransform Transform = CastChecked<const UWidget>(Object)->RenderTransform;
+	FWidgetTransform Transform = CastChecked<const UWidget>(Object)->GetRenderTransform();
 
 	FIntermediateWidgetTransform IntermediateTransform{};
 	ConvertOperationalProperty(Transform, IntermediateTransform);
diff --git a/Engine/Source/Runtime/UMG/Private/Components/BackgroundBlur.cpp b/Engine/Source/Runtime/UMG/Private/Components/BackgroundBlur.cpp
--- a/Engine/Source/Runtime/UMG/Private/Components/BackgroundBlur.cpp
+++ b/Engine/Source/Runtime/UMG/Private/Components/BackgroundBlur.cpp
@@ -25,7 +25,7 @@ UBackgroundBlur::UBackgroundBlur(const FObjectInitializer& ObjectInitializer)
 	, LowQualityFallbackBrush(FSlateNoResource())
 {
 	bIsVariable = false;
-	Visibility = ESlateVisibility::SelfHitTestInvisible;
+	SetVisibility(ESlateVisibility::SelfHitTestInvisible);
 }
 
 void UBackgroundBlur::ReleaseSlateResources(bool bReleaseChildren)
diff --git a/Engine/Source/Runtime/UMG/Private/Components/CanvasPanel.cpp b/Engine/Source/Runtime/UMG/Private/Components/CanvasPanel.cpp
--- a/Engine/Source/Runtime/UMG/Private/Components/CanvasPanel.cpp
+++ b/Engine/Source/Runtime/UMG/Private/Components/CanvasPanel.cpp
@@ -13,7 +13,7 @@ UCanvasPanel::UCanvasPanel(const FObjectInitializer& ObjectInitializer)
 	: Super(ObjectInitializer)
 {
 	bIsVariable = false;
-	Visibility = ESlateVisibility::SelfHitTestInvisible;
+	SetVisibility(ESlateVisibility::SelfHitTestInvisible);
 }
 
 void UCanvasPanel::ReleaseSlateResources(bool bReleaseChildren)
diff --git a/Engine/Source/Runtime/UMG/Private/Components/DynamicEntryBoxBase.cpp b/Engine/Source/Runtime/UMG/Private/Components/DynamicEntryBoxBase.cpp
--- a/Engine/Source/Runtime/UMG/Private/Components/DynamicEntryBoxBase.cpp
+++ b/Engine/Source/Runtime/UMG/Private/Components/DynamicEntryBoxBase.cpp
@@ -15,7 +15,7 @@ UDynamicEntryBoxBase::UDynamicEntryBoxBase(const FObjectInitializer& Initializer
 	: Super(Initializer)
 	, EntryWidgetPool(*this)
 {
-	Visibility = ESlateVisibility::SelfHitTestInvisible;
+	SetVisibility(ESlateVisibility::SelfHitTestInvisible);
 	EntrySizeRule.SizeRule = ESlateSizeRule::Automatic;
 }
 
diff --git a/Engine/Source/Runtime/UMG/Private/Components/EditableText.cpp b/Engine/Source/Runtime/UMG/Private/Components/EditableText.cpp
--- a/Engine/Source/Runtime/UMG/Private/Components/EditableText.cpp
+++ b/Engine/Source/Runtime/UMG/Private/Components/EditableText.cpp
@@ -62,7 +62,7 @@ UEditableText::UEditableText(const FObjectInitializer& ObjectInitializer)
 	AllowContextMenu = true;
 	VirtualKeyboardTrigger = EVirtualKeyboardTrigger::OnFocusByPointer;
 	VirtualKeyboardDismissAction = EVirtualKeyboardDismissAction::TextChangeOnDismiss;
-	Clipping = EWidgetClipping::ClipToBounds;
+	SetClipping(EWidgetClipping::ClipToBounds);
 	OverflowPolicy = ETextOverflowPolicy::Clip;
 	PRAGMA_ENABLE_DEPRECATION_WARNINGS
 
diff --git a/Engine/Source/Runtime/UMG/Private/Components/GridPanel.cpp b/Engine/Source/Runtime/UMG/Private/Components/GridPanel.cpp
--- a/Engine/Source/Runtime/UMG/Private/Components/GridPanel.cpp
+++ b/Engine/Source/Runtime/UMG/Private/Components/GridPanel.cpp
@@ -14,7 +14,7 @@ UGridPanel::UGridPanel(const FObjectInitializer& ObjectInitializer)
 	: Super(ObjectInitializer)
 {
 	bIsVariable = false;
-	Visibility = ESlateVisibility::SelfHitTestInvisible;
+	SetVisibility(ESlateVisibility::SelfHitTestInvisible);
 }
 
 void UGridPanel::ReleaseSlateResources(bool bReleaseChildren)
diff --git a/Engine/Source/Runtime/UMG/Private/Components/HorizontalBox.cpp b/Engine/Source/Runtime/UMG/Private/Components/HorizontalBox.cpp
--- a/Engine/Source/Runtime/UMG/Private/Components/HorizontalBox.cpp
+++ b/Engine/Source/Runtime/UMG/Private/Components/HorizontalBox.cpp
@@ -12,7 +12,7 @@ UHorizontalBox::UHorizontalBox(const FObjectInitializer& ObjectInitializer)
 	: Super(ObjectInitializer)
 {
 	bIsVariable = false;
-	Visibility = ESlateVisibility::SelfHitTestInvisible;
+	SetVisibility(ESlateVisibility::SelfHitTestInvisible);
 }
 
 void UHorizontalBox::ReleaseSlateResources(bool bReleaseChildren)
diff --git a/Engine/Source/Runtime/UMG/Private/Components/InvalidationBox.cpp b/Engine/Source/Runtime/UMG/Private/Components/InvalidationBox.cpp
--- a/Engine/Source/Runtime/UMG/Private/Components/InvalidationBox.cpp
+++ b/Engine/Source/Runtime/UMG/Private/Components/InvalidationBox.cpp
@@ -14,7 +14,7 @@ UInvalidationBox::UInvalidationBox(const FObjectInitializer& ObjectInitializer)
 	: Super(ObjectInitializer)
 {
 	bCanCache = true;
-	Visibility = ESlateVisibility::SelfHitTestInvisible;
+	SetVisibility(ESlateVisibility::SelfHitTestInvisible);
 }
 
 void UInvalidationBox::ReleaseSlateResources(bool bReleaseChildren)
diff --git a/Engine/Source/Runtime/UMG/Private/Components/ListViewBase.cpp b/Engine/Source/Runtime/UMG/Private/Components/ListViewBase.cpp
--- a/Engine/Source/Runtime/UMG/Private/Components/ListViewBase.cpp
+++ b/Engine/Source/Runtime/UMG/Private/Components/ListViewBase.cpp
@@ -14,7 +14,7 @@ UListViewBase::UListViewBase(const FObjectInitializer& ObjectInitializer)
 	, EntryWidgetPool(*this)
 {
 	bIsVariable = true;
-	Clipping = EWidgetClipping::ClipToBounds;
+	SetClipping(EWidgetClipping::ClipToBounds);
 }
 
 #if WITH_EDITOR
diff --git a/Engine/Source/Runtime/UMG/Private/Components/MultiLineEditableText.cpp b/Engine/Source/Runtime/UMG/Private/Components/MultiLineEditableText.cpp
--- a/Engine/Source/Runtime/UMG/Private/Components/MultiLineEditableText.cpp
+++ b/Engine/Source/Runtime/UMG/Private/Components/MultiLineEditableText.cpp
@@ -55,7 +55,7 @@ UMultiLineEditableText::UMultiLineEditableText(const FObjectInitializer& ObjectI
 	RevertTextOnEscape = false;
 	ClearKeyboardFocusOnCommit = true;
 	AllowContextMenu = true;
-	Clipping = EWidgetClipping::ClipToBounds;
+	SetClipping(EWidgetClipping::ClipToBounds);
 	VirtualKeyboardDismissAction = EVirtualKeyboardDismissAction::TextChangeOnDismiss;
 	AutoWrapText = true;
 
diff --git a/Engine/Source/Runtime/UMG/Private/Components/NamedSlot.cpp b/Engine/Source/Runtime/UMG/Private/Components/NamedSlot.cpp
--- a/Engine/Source/Runtime/UMG/Private/Components/NamedSlot.cpp
+++ b/Engine/Source/Runtime/UMG/Private/Components/NamedSlot.cpp
@@ -15,7 +15,7 @@ UNamedSlot::UNamedSlot(const FObjectInitializer& ObjectInitializer)
 	: Super(ObjectInitializer)
 {
 	bIsVariable = true;
-	Visibility = ESlateVisibility::SelfHitTestInvisible;
+	SetVisibility(ESlateVisibility::SelfHitTestInvisible);
 }
 
 void UNamedSlot::ReleaseSlateResources(bool bReleaseChildren)
diff --git a/Engine/Source/Runtime/UMG/Private/Components/Overlay.cpp b/Engine/Source/Runtime/UMG/Private/Components/Overlay.cpp
--- a/Engine/Source/Runtime/UMG/Private/Components/Overlay.cpp
+++ b/Engine/Source/Runtime/UMG/Private/Components/Overlay.cpp
@@ -12,7 +12,7 @@ UOverlay::UOverlay(const FObjectInitializer& ObjectInitializer)
 	: Super(ObjectInitializer)
 {
 	bIsVariable = false;
-	Visibility = ESlateVisibility::SelfHitTestInvisible;
+	SetVisibility(ESlateVisibility::SelfHitTestInvisible);
 }
 
 void UOverlay::ReleaseSlateResources(bool bReleaseChildren)
diff --git a/Engine/Source/Runtime/UMG/Private/Components/RetainerBox.cpp b/Engine/Source/Runtime/UMG/Private/Components/RetainerBox.cpp
--- a/Engine/Source/Runtime/UMG/Private/Components/RetainerBox.cpp
+++ b/Engine/Source/Runtime/UMG/Private/Components/RetainerBox.cpp
@@ -16,7 +16,7 @@ static FName DefaultTextureParameterName("Texture");
 URetainerBox::URetainerBox(const FObjectInitializer& ObjectInitializer)
 	: Super(ObjectInitializer)
 {
-	Visibility = ESlateVisibility::Visible;
+	SetVisibility(ESlateVisibility::Visible);
 	Phase = 0;
 	PhaseCount = 1;
 	RenderOnPhase = true;
diff --git a/Engine/Source/Runtime/UMG/Private/Components/RichTextBlock.cpp b/Engine/Source/Runtime/UMG/Private/Components/RichTextBlock.cpp
--- a/Engine/Source/Runtime/UMG/Private/Components/RichTextBlock.cpp
+++ b/Engine/Source/Runtime/UMG/Private/Components/RichTextBlock.cpp
@@ -47,7 +47,7 @@ FORCEINLINE TSharedPtr< ObjectType > MakeShareableDeferredCleanup(ObjectType* In
 URichTextBlock::URichTextBlock(const FObjectInitializer& ObjectInitializer)
 	: Super(ObjectInitializer)
 {
-	Visibility = ESlateVisibility::SelfHitTestInvisible;
+	SetVisibility(ESlateVisibility::SelfHitTestInvisible);
 	TextTransformPolicy = ETextTransformPolicy::None;
 	TextOverflowPolicy = ETextOverflowPolicy::Clip;
 }
diff --git a/Engine/Source/Runtime/UMG/Private/Components/SafeZone.cpp b/Engine/Source/Runtime/UMG/Private/Components/SafeZone.cpp
--- a/Engine/Source/Runtime/UMG/Private/Components/SafeZone.cpp
+++ b/Engine/Source/Runtime/UMG/Private/Components/SafeZone.cpp
@@ -14,7 +14,7 @@ USafeZone::USafeZone()
 	, PadBottom(true)
 {
 	bCanHaveMultipleChildren = false;
-	Visibility = ESlateVisibility::SelfHitTestInvisible;
+	SetVisibility(ESlateVisibility::SelfHitTestInvisible);
 }
 
 #if WITH_EDITOR
diff --git a/Engine/Source/Runtime/UMG/Private/Components/ScaleBox.cpp b/Engine/Source/Runtime/UMG/Private/Components/ScaleBox.cpp
--- a/Engine/Source/Runtime/UMG/Private/Components/ScaleBox.cpp
+++ b/Engine/Source/Runtime/UMG/Private/Components/ScaleBox.cpp
@@ -13,7 +13,7 @@ UScaleBox::UScaleBox(const FObjectInitializer& ObjectInitializer)
 	: Super(ObjectInitializer)
 {
 	bIsVariable = false;
-	Visibility = ESlateVisibility::SelfHitTestInvisible;
+	SetVisibility(ESlateVisibility::SelfHitTestInvisible);
 
 	StretchDirection = EStretchDirection::Both;
 	Stretch = EStretch::ScaleToFit;
diff --git a/Engine/Source/Runtime/UMG/Private/Components/ScrollBox.cpp b/Engine/Source/Runtime/UMG/Private/Components/ScrollBox.cpp
--- a/Engine/Source/Runtime/UMG/Private/Components/ScrollBox.cpp
+++ b/Engine/Source/Runtime/UMG/Private/Components/ScrollBox.cpp
@@ -37,8 +37,8 @@ UScrollBox::UScrollBox(const FObjectInitializer& ObjectInitializer)
 {
 	bIsVariable = false;
 
-	Visibility = ESlateVisibility::Visible;
-	Clipping = EWidgetClipping::ClipToBounds;
+	SetVisibility(ESlateVisibility::Visible);
+	SetClipping(EWidgetClipping::ClipToBounds);
 
 	if (DefaultScrollBoxStyle == nullptr)
 	{
diff --git a/Engine/Source/Runtime/UMG/Private/Components/SizeBox.cpp b/Engine/Source/Runtime/UMG/Private/Components/SizeBox.cpp
--- a/Engine/Source/Runtime/UMG/Private/Components/SizeBox.cpp
+++ b/Engine/Source/Runtime/UMG/Private/Components/SizeBox.cpp
@@ -15,7 +15,7 @@ USizeBox::USizeBox(const FObjectInitializer& ObjectInitializer)
 	: Super(ObjectInitializer)
 {
 	bIsVariable = false;
-	Visibility = ESlateVisibility::SelfHitTestInvisible;
+	SetVisibility(ESlateVisibility::SelfHitTestInvisible);
 	MinAspectRatio = 1;
 	MaxAspectRatio = 1;
 }
diff --git a/Engine/Source/Runtime/UMG/Private/Components/Spacer.cpp b/Engine/Source/Runtime/UMG/Private/Components/Spacer.cpp
--- a/Engine/Source/Runtime/UMG/Private/Components/Spacer.cpp
+++ b/Engine/Source/Runtime/UMG/Private/Components/Spacer.cpp
@@ -14,7 +14,7 @@ USpacer::USpacer(const FObjectInitializer& ObjectInitializer)
 	, Size(1.0f, 1.0f)
 {
 	bIsVariable = false;
-	Visibility = ESlateVisibility::SelfHitTestInvisible;
+	SetVisibility(ESlateVisibility::SelfHitTestInvisible);
 }
 
 void USpacer::ReleaseSlateResources(bool bReleaseChildren)
diff --git a/Engine/Source/Runtime/UMG/Private/Components/UniformGridPanel.cpp b/Engine/Source/Runtime/UMG/Private/Components/UniformGridPanel.cpp
--- a/Engine/Source/Runtime/UMG/Private/Components/UniformGridPanel.cpp
+++ b/Engine/Source/Runtime/UMG/Private/Components/UniformGridPanel.cpp
@@ -14,7 +14,7 @@ UUniformGridPanel::UUniformGridPanel(const FObjectInitializer& ObjectInitializer
 	: Super(ObjectInitializer)
 {
 	bIsVariable = false;
-	Visibility = ESlateVisibility::SelfHitTestInvisible;
+	SetVisibility(ESlateVisibility::SelfHitTestInvisible);
 }
 
 void UUniformGridPanel::ReleaseSlateResources(bool bReleaseChildren)
diff --git a/Engine/Source/Runtime/UMG/Private/Components/VerticalBox.cpp b/Engine/Source/Runtime/UMG/Private/Components/VerticalBox.cpp
--- a/Engine/Source/Runtime/UMG/Private/Components/VerticalBox.cpp
+++ b/Engine/Source/Runtime/UMG/Private/Components/VerticalBox.cpp
@@ -12,7 +12,7 @@ UVerticalBox::UVerticalBox(const FObjectInitializer& ObjectInitializer)
 	: Super(ObjectInitializer)
 {
 	bIsVariable = false;
-	Visibility = ESlateVisibility::SelfHitTestInvisible;
+	SetVisibility(ESlateVisibility::SelfHitTestInvisible);
 }
 
 void UVerticalBox::ReleaseSlateResources(bool bReleaseChildren)
diff --git a/Engine/Source/Runtime/UMG/Private/Components/Widget.cpp b/Engine/Source/Runtime/UMG/Private/Components/Widget.cpp
--- a/Engine/Source/Runtime/UMG/Private/Components/Widget.cpp
+++ b/Engine/Source/Runtime/UMG/Private/Components/Widget.cpp
@@ -176,6 +176,7 @@ UE_FIELD_NOTIFICATION_IMPLEMENT_CLASS_DESCRIPTOR_ThreeFields(UWidget, ToolTipTex
 UWidget::UWidget(const FObjectInitializer& ObjectInitializer)
 	: Super(ObjectInitializer)
 {
+PRAGMA_DISABLE_DEPRECATION_WARNINGS
 	bIsEnabled = true;
 	bIsVariable = true;
 #if WITH_EDITOR
@@ -194,6 +195,8 @@ UWidget::UWidget(const FObjectInitializer& ObjectInitializer)
 #endif
 	AccessibleWidgetData = nullptr;
 
+PRAGMA_ENABLE_DEPRECATION_WARNINGS
+
 #if WITH_EDITORONLY_DATA
 	{ static const FAutoRegisterLocalizationDataGatheringCallback AutomaticRegistrationOfLocalizationGatherer(UWidget::StaticClass(), &GatherWidgetForLocalization); }
 #endif
@@ -201,6 +204,12 @@ UWidget::UWidget(const FObjectInitializer& ObjectInitializer)
 	INC_DWORD_STAT(STAT_SlateUTotalWidgets);
 }
 
+PRAGMA_DISABLE_DEPRECATION_WARNINGS
+FWidgetTransform UWidget::GetRenderTransform() const
+{
+	return RenderTransform;
+}
+
 void UWidget::SetRenderTransform(FWidgetTransform Transform)
 {
 	RenderTransform = Transform;
@@ -252,6 +261,11 @@ void UWidget::UpdateRenderTransform()
 	}
 }
 
+FVector2D UWidget::GetRenderTransformPivot() const
+{
+	return RenderTransformPivot;
+}
+
 void UWidget::SetRenderTransformPivot(FVector2D Pivot)
 {
 	RenderTransformPivot = Pivot;
@@ -263,6 +277,21 @@ void UWidget::SetRenderTransformPivot(FVector2D Pivot)
 	}
 }
 
+EFlowDirectionPreference UWidget::GetFlowDirectionPreference() const
+{
+	return FlowDirectionPreference;
+}
+
+void UWidget::SetFlowDirectionPreference(EFlowDirectionPreference FlowDirection)
+{
+	FlowDirectionPreference = FlowDirection;
+	
+	if (TSharedPtr<SWidget> SafeWidget = GetCachedWidget())
+	{
+		SafeWidget->SetFlowDirectionPreference(FlowDirectionPreference);
+	}
+}
+
 bool UWidget::GetIsEnabled() const
 {
 	TSharedPtr<SWidget> SafeWidget = GetCachedWidget();
@@ -284,6 +313,11 @@ void UWidget::SetIsEnabled(bool bInIsEnabled)
 	}
 }
 
+EMouseCursor::Type UWidget::GetCursor() const
+{
+	return Cursor;
+}
+
 void UWidget::SetCursor(EMouseCursor::Type InCursor)
 {
 	bOverride_Cursor = true;
@@ -399,6 +433,7 @@ void UWidget::SetClipping(EWidgetClipping InClipping)
 		SafeWidget->SetClipping(InClipping);
 	}
 }
+PRAGMA_ENABLE_DEPRECATION_WARNINGS
 
 void UWidget::ForceVolatile(bool bForce)
 {
@@ -410,6 +445,12 @@ void UWidget::ForceVolatile(bool bForce)
 	}
 }
 
+PRAGMA_DISABLE_DEPRECATION_WARNINGS
+FText UWidget::GetToolTipText() const
+{
+	return ToolTipText;
+}
+
 void UWidget::SetToolTipText(const FText& InToolTipText)
 {
 	ToolTipText = InToolTipText;
@@ -422,6 +463,11 @@ void UWidget::SetToolTipText(const FText& InToolTipText)
 	}
 }
 
+UWidget* UWidget::GetToolTip() const
+{
+	return ToolTipWidget;
+}
+
 void UWidget::SetToolTip(UWidget* InToolTipWidget)
 {
 	ToolTipWidget = InToolTipWidget;
@@ -446,6 +492,7 @@ void UWidget::SetToolTip(UWidget* InToolTipWidget)
 		}
 	}
 }
+PRAGMA_ENABLE_DEPRECATION_WARNINGS
 
 bool UWidget::IsHovered() const
 {
@@ -1272,6 +1319,7 @@ void UWidget::SynchronizeProperties()
 	TSharedPtr<SWidget> SafeContentWidget = MyGCWidget.IsValid() ? MyGCWidget.Pin() : MyWidget.Pin();
 #endif
 
+PRAGMA_DISABLE_DEPRECATION_WARNINGS
 #if WITH_EDITOR
 	// Always use an enabled and visible state in the designer.
 	if ( IsDesignTime() )
@@ -1344,6 +1392,7 @@ void UWidget::SynchronizeProperties()
 		}
 	}
 #endif
+PRAGMA_ENABLE_DEPRECATION_WARNINGS
 }
 
 
diff --git a/Engine/Source/Runtime/UMG/Private/Components/WidgetSwitcher.cpp b/Engine/Source/Runtime/UMG/Private/Components/WidgetSwitcher.cpp
--- a/Engine/Source/Runtime/UMG/Private/Components/WidgetSwitcher.cpp
+++ b/Engine/Source/Runtime/UMG/Private/Components/WidgetSwitcher.cpp
@@ -13,7 +13,7 @@ UWidgetSwitcher::UWidgetSwitcher(const FObjectInitializer& ObjectInitializer)
 	: Super(ObjectInitializer)
 {
 	bIsVariable = true;
-	Visibility = ESlateVisibility::SelfHitTestInvisible;
+	SetVisibility(ESlateVisibility::SelfHitTestInvisible);
 }
 
 void UWidgetSwitcher::ReleaseSlateResources(bool bReleaseChildren)
diff --git a/Engine/Source/Runtime/UMG/Private/Components/WindowTitleBarArea.cpp b/Engine/Source/Runtime/UMG/Private/Components/WindowTitleBarArea.cpp
--- a/Engine/Source/Runtime/UMG/Private/Components/WindowTitleBarArea.cpp
+++ b/Engine/Source/Runtime/UMG/Private/Components/WindowTitleBarArea.cpp
@@ -19,7 +19,7 @@ UWindowTitleBarArea::UWindowTitleBarArea(const FObjectInitializer& ObjectInitial
 	: Super(ObjectInitializer)
 {
 	bIsVariable = false;
-	Visibility = ESlateVisibility::Visible;
+	SetVisibility(ESlateVisibility::Visible);
 
 	bDoubleClickTogglesFullscreen = false;
 }
diff --git a/Engine/Source/Runtime/UMG/Private/Components/WrapBox.cpp b/Engine/Source/Runtime/UMG/Private/Components/WrapBox.cpp
--- a/Engine/Source/Runtime/UMG/Private/Components/WrapBox.cpp
+++ b/Engine/Source/Runtime/UMG/Private/Components/WrapBox.cpp
@@ -12,7 +12,7 @@ UWrapBox::UWrapBox(const FObjectInitializer& ObjectInitializer)
 	: Super(ObjectInitializer)
 {
 	bIsVariable = false;
-	Visibility = ESlateVisibility::SelfHitTestInvisible;
+	SetVisibility(ESlateVisibility::SelfHitTestInvisible);
 	WrapSize = 500;
 	bExplicitWrapSize = false;
 	HorizontalAlignment = HAlign_Left;
diff --git a/Engine/Source/Runtime/UMG/Private/UserWidget.cpp b/Engine/Source/Runtime/UMG/Private/UserWidget.cpp
--- a/Engine/Source/Runtime/UMG/Private/UserWidget.cpp
+++ b/Engine/Source/Runtime/UMG/Private/UserWidget.cpp
@@ -77,7 +77,7 @@ UUserWidget::UUserWidget(const FObjectInitializer& ObjectInitializer)
 	, TickFrequency(EWidgetTickFrequency::Auto)
 {
 	ViewportAnchors = FAnchors(0, 0, 1, 1);
-	Visibility = ESlateVisibility::SelfHitTestInvisible;
+	SetVisibility(ESlateVisibility::SelfHitTestInvisible);
 
 	bIsFocusable = false;
 	ColorAndOpacity = FLinearColor::White;
@@ -1882,7 +1882,7 @@ FReply UUserWidget::NativeOnTouchForceChanged(const FGeometry& InGeometry, const
 FCursorReply UUserWidget::NativeOnCursorQuery( const FGeometry& InGeometry, const FPointerEvent& InCursorEvent )
 {
 	return (bOverride_Cursor)
-		? FCursorReply::Cursor(Cursor)
+		? FCursorReply::Cursor(GetCursor())
 		: FCursorReply::Unhandled();
 }
 
diff --git a/Engine/Source/Runtime/UMG/Public/Components/Border.h b/Engine/Source/Runtime/UMG/Public/Components/Border.h
--- a/Engine/Source/Runtime/UMG/Public/Components/Border.h
+++ b/Engine/Source/Runtime/UMG/Public/Components/Border.h
@@ -163,11 +163,11 @@ public:
 	UMaterialInstanceDynamic* GetDynamicMaterial();
 
 public:
-	/** Gets the DesireSizeScale of this border. */
+	/** Gets the DesiredSizeScale of this border. */
 	FVector2D GetDesiredSizeScale() const;
 
 	/**
-	 * Sets the DesireSizeScale of this border.
+	 * Sets the DesiredSizeScale of this border.
 	 *
 	 * @param InScale    The X and Y multipliers for the desired size
 	 */
diff --git a/Engine/Source/Runtime/UMG/Public/Components/Widget.h b/Engine/Source/Runtime/UMG/Public/Components/Widget.h
--- a/Engine/Source/Runtime/UMG/Public/Components/Widget.h
+++ b/Engine/Source/Runtime/UMG/Public/Components/Widget.h
@@ -207,7 +207,7 @@ public:
 /**
  * This is the base class for all wrapped Slate controls that are exposed to UObjects.
  */
-UCLASS(Abstract, BlueprintType, Blueprintable)
+UCLASS(Abstract, BlueprintType, Blueprintable, CustomFieldNotify)
 class UMG_API UWidget : public UVisual, public INotifyFieldValueChanged
 {
 	GENERATED_UCLASS_BODY()
@@ -259,7 +259,8 @@ public:
 	FGetBool bIsEnabledDelegate;
 
 	/** Tooltip text to show when the user hovers over the widget with the mouse */
-	UPROPERTY(EditAnywhere, BlueprintReadWrite, BlueprintSetter="SetToolTipText", Category="Behavior", meta=(MultiLine=true))
+	UE_DEPRECATED(5.1, "Direct access to ToolTipText is deprecated. Please use the getter or setter.")
+	UPROPERTY(EditAnywhere, BlueprintReadWrite, Getter, Setter, BlueprintSetter="SetToolTipText", Category="Behavior", meta=(MultiLine=true))
 	FText ToolTipText;
 
 	/** A bindable delegate for ToolTipText */
@@ -267,7 +268,8 @@ public:
 	FGetText ToolTipTextDelegate;
 
 	/** Tooltip widget to show when the user hovers over the widget with the mouse */
-	UPROPERTY(VisibleAnywhere, BlueprintReadOnly, Category="Behavior", AdvancedDisplay)
+	UE_DEPRECATED(5.1, "Direct access to ToolTipWidget is deprecated. Please use the getter or setter.")
+	UPROPERTY(VisibleAnywhere, BlueprintReadWrite, Getter="GetToolTip", Setter="SetToolTip", BlueprintSetter="SetToolTip", Category="Behavior", AdvancedDisplay)
 	TObjectPtr<UWidget> ToolTipWidget;
 
 	/** A bindable delegate for ToolTipWidget */
@@ -277,27 +279,26 @@ public:
 
 	/** A bindable delegate for Visibility */
 	UPROPERTY()
-	FGetSlateVisibility VisibilityDelegate;
-
-	/** A bindable delegate for Cursor */
-	//UPROPERTY()
-	//FGetMouseCursor CursorDelegate;
+	FGetSlateVisibility VisibilityDelegate;;
 
 public:
 
 	/** The render transform of the widget allows for arbitrary 2D transforms to be applied to the widget. */
-	UPROPERTY(EditAnywhere, BlueprintReadOnly, Category="Render Transform", meta=( DisplayName="Transform" ))
+	UE_DEPRECATED(5.1, "Direct access to RenderTransform is deprecated. Please use the getter or setter.")
+	UPROPERTY(EditAnywhere, BlueprintReadWrite, Getter, Setter, BlueprintSetter="SetRenderTransform", Category="Render Transform", meta = (DisplayName = "Transform"))
 	FWidgetTransform RenderTransform;
 
 	/**
 	 * The render transform pivot controls the location about which transforms are applied.  
 	 * This value is a normalized coordinate about which things like rotations will occur.
 	 */
-	UPROPERTY(EditAnywhere, BlueprintReadOnly, Category="Render Transform", meta=( DisplayName="Pivot" ))
+	UE_DEPRECATED(5.1, "Direct access to RenderTransformPivot is deprecated. Please use the getter or setter.")
+	UPROPERTY(EditAnywhere, BlueprintReadWrite, Getter, Setter, BlueprintSetter="SetRenderTransformPivot", Category="Render Transform", meta=( DisplayName="Pivot" ))
 	FVector2D RenderTransformPivot;
 
 	/** Allows you to set a new flow direction */
-	UPROPERTY(EditAnywhere, Category = "Localization")
+	UE_DEPRECATED(5.1, "Direct access to FlowDirectionPreference is deprecated. Please use the getter or setter.")
+	UPROPERTY(EditAnywhere, Getter, Setter, Category="Localization")
 	EFlowDirectionPreference FlowDirectionPreference;
 
 	/**
@@ -312,7 +313,8 @@ public:
 	uint8 bCreatedByConstructionScript:1;
 
 	/** Sets whether this widget can be modified interactively by the user */
-	UPROPERTY(EditAnywhere, BlueprintReadOnly, Category="Behavior")
+	UE_DEPRECATED(5.1, "Direct access to bIsEnabled is deprecated. Please use the getter or setter.")
+	UPROPERTY(EditAnywhere, BlueprintReadWrite, FieldNotify, Getter="GetIsEnabled", Setter="SetIsEnabled", BlueprintGetter="GetIsEnabled", BlueprintSetter="SetIsEnabled", Category="Behavior")
 	uint8 bIsEnabled:1;
 
 	/**  */
@@ -392,7 +394,8 @@ public:
 #endif
 
 	/** The cursor to show when the mouse is over the widget */
-	UPROPERTY(EditAnywhere, BlueprintReadOnly, Category="Behavior", AdvancedDisplay, meta=( editcondition="bOverride_Cursor" ))
+	UE_DEPRECATED(5.1, "Direct access to Cursor is deprecated. Please use the getter or setter.")
+	UPROPERTY(EditAnywhere, BlueprintReadWrite, Getter, Setter, BlueprintSetter="SetCursor", Category="Behavior", AdvancedDisplay, meta = (editcondition = "bOverride_Cursor"))
 	TEnumAsByte<EMouseCursor::Type> Cursor;
 
 	/**
@@ -404,15 +407,18 @@ public:
 	 * performance cost to clipping.  Do not enable clipping unless a panel actually needs to prevent
 	 * content from showing up outside its bounds.
 	 */
-	UPROPERTY(EditAnywhere, Category = "Clipping")
+	UE_DEPRECATED(5.1, "Direct access to Clipping is deprecated. Please use the getter or setter.")
+	UPROPERTY(EditAnywhere, BlueprintReadWrite, Getter, Setter, Category="Clipping")
 	EWidgetClipping Clipping;
 
 	/** The visibility of the widget */
-	UPROPERTY(EditAnywhere, Category="Behavior")
+	UE_DEPRECATED(5.1, "Direct access to Visibility is deprecated. Please use the getter or setter.")
+	UPROPERTY(EditAnywhere, BlueprintReadWrite, FieldNotify, Getter, Setter, BlueprintGetter="GetVisibility", BlueprintSetter="SetVisibility", Category="Behavior")
 	ESlateVisibility Visibility;
 
 	/** The opacity of the widget */
-	UPROPERTY(EditAnywhere, Category="Behavior")
+	UE_DEPRECATED(5.1, "Direct access to RenderOpacity is deprecated. Please use the getter or setter.")
+	UPROPERTY(EditAnywhere, BlueprintReadWrite, Getter, Setter, BlueprintGetter="GetRenderOpacity", BlueprintSetter="SetRenderOpacity", Category="Behavior")
 	float RenderOpacity;
 
 private:
@@ -460,6 +466,8 @@ public:
 #endif
 
 public:
+	/** */
+	FWidgetTransform GetRenderTransform() const;
 
 	/** */
 	UFUNCTION(BlueprintCallable, Category="Widget|Transform")
@@ -485,10 +493,19 @@ public:
 	UFUNCTION(BlueprintCallable, Category="Widget|Transform")
 	void SetRenderTranslation(FVector2D Translation);
 
+	/** */
+	FVector2D GetRenderTransformPivot() const;
+
 	/** */
 	UFUNCTION(BlueprintCallable, Category="Widget|Transform")
 	void SetRenderTransformPivot(FVector2D Pivot);
 
+	/** */
+	EFlowDirectionPreference GetFlowDirectionPreference() const;
+
+	/** */
+	void SetFlowDirectionPreference(EFlowDirectionPreference FlowDirection);
+
 	/** Gets the current enabled status of the widget */
 	UFUNCTION(BlueprintCallable, Category="Widget")
 	bool GetIsEnabled() const;
@@ -497,14 +514,23 @@ public:
 	UFUNCTION(BlueprintCallable, Category="Widget")
 	virtual void SetIsEnabled(bool bInIsEnabled);
 
+	/** @return the tooltip text for the widget. */
+	FText GetToolTipText() const;
+
 	/** Sets the tooltip text for the widget. */
 	UFUNCTION(BlueprintCallable, Category="Widget")
 	void SetToolTipText(const FText& InToolTipText);
 
+	/** @return the custom widget as the tooltip of the widget. */
+	UWidget* GetToolTip() const;
+
 	/** Sets a custom widget as the tooltip of the widget. */
 	UFUNCTION(BlueprintCallable, Category="Widget")
 	void SetToolTip(UWidget* Widget);
 
+	/** Sets the cursor to show over the widget. */
+	EMouseCursor::Type GetCursor() const;
+
 	/** Sets the cursor to show over the widget. */
 	UFUNCTION(BlueprintCallable, Category="Widget")
 	void SetCursor(EMouseCursor::Type InCursor);
