UMG: Update ProgressBar with accessors.

diff --git a/Engine/Source/Runtime/UMG/Private/Components/ProgressBar.cpp b/Engine/Source/Runtime/UMG/Private/Components/ProgressBar.cpp
--- a/Engine/Source/Runtime/UMG/Private/Components/ProgressBar.cpp
+++ b/Engine/Source/Runtime/UMG/Private/Components/ProgressBar.cpp
@@ -10,9 +10,6 @@
 // UProgressBar
 
 
-UE_FIELD_NOTIFICATION_IMPLEMENT_CLASS_DESCRIPTOR_ThreeFields(UProgressBar, Percent, FillColorAndOpacity, bIsMarquee);
-
-
 static FProgressBarStyle* DefaultProgressBarStyle = nullptr;
 
 #if WITH_EDITOR
@@ -30,6 +27,7 @@ UProgressBar::UProgressBar(const FObjectInitializer& ObjectInitializer)
 		DefaultProgressBarStyle->UnlinkColors();
 	}
 
+PRAGMA_DISABLE_DEPRECATION_WARNINGS
 	WidgetStyle = *DefaultProgressBarStyle;
 
 #if WITH_EDITOR 
@@ -58,6 +56,7 @@ UProgressBar::UProgressBar(const FObjectInitializer& ObjectInitializer)
 	Percent = 0;
 	FillColorAndOpacity = FLinearColor::White;
 	BorderPadding = FVector2D(0, 0);
+PRAGMA_ENABLE_DEPRECATION_WARNINGS
 }
 
 void UProgressBar::ReleaseSlateResources(bool bReleaseChildren)
@@ -74,6 +73,7 @@ TSharedRef<SWidget> UProgressBar::RebuildWidget()
 	return MyProgressBar.ToSharedRef();
 }
 
+PRAGMA_DISABLE_DEPRECATION_WARNINGS
 void UProgressBar::SynchronizeProperties()
 {
 	Super::SynchronizeProperties();
@@ -90,47 +90,104 @@ void UProgressBar::SynchronizeProperties()
 	MyProgressBar->SetBorderPadding(BorderPadding);
 }
 
-void UProgressBar::SetIsMarquee(bool InbIsMarquee)
+const FProgressBarStyle& UProgressBar::GetWidgetStyle() const
 {
-	if (bIsMarquee != InbIsMarquee)
+	return WidgetStyle;
+}
+
+void UProgressBar::SetWidgetStyle(const FProgressBarStyle& InStyle)
+{
+	WidgetStyle = InStyle;
+	if (MyProgressBar.IsValid())
 	{
-		bIsMarquee = InbIsMarquee;
-		BroadcastFieldValueChanged(FFieldNotificationClassDescriptor::bIsMarquee);
+		MyProgressBar->SetStyle(&WidgetStyle);
 	}
+}
+
+float UProgressBar::GetPercent() const
+{
+	return Percent;
+}
 
+void UProgressBar::SetPercent(float InPercent)
+{
+	Percent = InPercent;
 	if (MyProgressBar.IsValid())
 	{
-		MyProgressBar->SetPercent(bIsMarquee ? TOptional<float>() : Percent);
+		MyProgressBar->SetPercent(InPercent);
 	}
 }
 
-void UProgressBar::SetFillColorAndOpacity(FLinearColor Color)
+EProgressBarFillType::Type UProgressBar::GetBarFillType() const
+{
+	return BarFillType;
+}
+
+void UProgressBar::SetBarFillType(EProgressBarFillType::Type InBarFillType)
+{
+	BarFillType = InBarFillType;
+	if (MyProgressBar.IsValid())
+	{
+		MyProgressBar->SetBarFillType(BarFillType);
+	}
+}
+
+EProgressBarFillStyle::Type UProgressBar::GetBarFillStyle() const
 {
-	if (FillColorAndOpacity != Color)
+	return BarFillStyle;
+}
+
+void UProgressBar::SetBarFillStyle(EProgressBarFillStyle::Type InBarFillStyle)
+{
+	BarFillStyle = InBarFillStyle;
+	if (MyProgressBar.IsValid())
 	{
-		FillColorAndOpacity = Color;
-		BroadcastFieldValueChanged(FFieldNotificationClassDescriptor::FillColorAndOpacity);
+		MyProgressBar->SetBarFillStyle(BarFillStyle);
 	}
+}
 
+bool UProgressBar::UseMarquee() const
+{
+	return bIsMarquee;
+}
+
+void UProgressBar::SetIsMarquee(bool InbIsMarquee)
+{
+	bIsMarquee = InbIsMarquee;
 	if (MyProgressBar.IsValid())
 	{
-		MyProgressBar->SetFillColorAndOpacity(FillColorAndOpacity);
+		MyProgressBar->SetPercent(bIsMarquee ? TOptional<float>() : Percent);
 	}
 }
 
-void UProgressBar::SetPercent(float InPercent)
+FVector2D UProgressBar::GetBorderPadding() const
 {
-	if (Percent != InPercent)
+	return BorderPadding;
+}
+
+void UProgressBar::SetBorderPadding(FVector2D InBorderPadding)
+{
+	BorderPadding = InBorderPadding;
+	if (MyProgressBar.IsValid())
 	{
-		Percent = InPercent;
-		BroadcastFieldValueChanged(FFieldNotificationClassDescriptor::Percent);
+		MyProgressBar->SetBorderPadding(BorderPadding);
 	}
+}
+
+FLinearColor UProgressBar::GetFillColorAndOpacity() const
+{
+	return FillColorAndOpacity;
+}
 
+void UProgressBar::SetFillColorAndOpacity(FLinearColor Color)
+{
+	FillColorAndOpacity = Color;
 	if (MyProgressBar.IsValid())
 	{
-		MyProgressBar->SetPercent(InPercent);
+		MyProgressBar->SetFillColorAndOpacity(FillColorAndOpacity);
 	}
 }
+PRAGMA_ENABLE_DEPRECATION_WARNINGS
 
 #if WITH_EDITOR
 
@@ -141,7 +198,9 @@ const FText UProgressBar::GetPaletteCategory()
 
 void UProgressBar::OnCreationFromPalette()
 {
-	FillColorAndOpacity = FLinearColor(0, 0.5f, 1.0f);
+PRAGMA_DISABLE_DEPRECATION_WARNINGS
+	FillColorAndOpacity = FLinearColor(0.f, 0.5f, 1.0f);
+PRAGMA_ENABLE_DEPRECATION_WARNINGS
 }
 
 #endif
diff --git a/Engine/Source/Runtime/UMG/Public/Components/ProgressBar.h b/Engine/Source/Runtime/UMG/Public/Components/ProgressBar.h
--- a/Engine/Source/Runtime/UMG/Public/Components/ProgressBar.h
+++ b/Engine/Source/Runtime/UMG/Public/Components/ProgressBar.h
@@ -21,31 +21,36 @@ UCLASS()
 class UMG_API UProgressBar : public UWidget
 {
 	GENERATED_UCLASS_BODY()
-	
-public:
-	UE_FIELD_NOTIFICATION_DECLARE_CLASS_DESCRIPTOR_ThreeFields(Percent, FillColorAndOpacity, bIsMarquee);
 
 public:
 	/** The progress bar style */
-	UPROPERTY(EditAnywhere, BlueprintReadWrite, Category="Style", meta=( DisplayName="Style" ))
+	UE_DEPRECATED(5.1, "Direct access to WidgetStyle is deprecated. Please use the getter or setter.")
+	UPROPERTY(EditAnywhere, BlueprintReadWrite, Getter, Setter, Category="Style", meta=( DisplayName="Style" ))
 	FProgressBarStyle WidgetStyle;
 
 	/** Used to determine the fill position of the progress bar ranging 0..1 */
-	UPROPERTY(EditAnywhere, BlueprintReadOnly, Category=Progress, meta=( UIMin = "0", UIMax = "1" ))
+	UE_DEPRECATED(5.1, "Direct access to Percent is deprecated. Please use the getter or setter.")
+	UPROPERTY(EditAnywhere, BlueprintReadWrite, Getter, Setter, BlueprintSetter="SetPercent", Category="Progress", meta = (UIMin = "0", UIMax = "1"))
 	float Percent;
 
 	/** Defines the direction in which the progress bar fills */
-	UPROPERTY(EditAnywhere, BlueprintReadOnly, Category=Progress)
+	UE_DEPRECATED(5.1, "Direct access to BarFillType is deprecated. Please use the getter or setter.")
+	UPROPERTY(EditAnywhere, BlueprintReadWrite, Getter, Setter, Category="Progress")
 	TEnumAsByte<EProgressBarFillType::Type> BarFillType;
 
 	/** Defines the visual style of the progress bar fill - scale or mask */
-	UPROPERTY(EditAnywhere, BlueprintReadOnly, Category=Progress)
+	UE_DEPRECATED(5.1, "Direct access to BarFillStyle is deprecated. Please use the getter or setter.")
+	UPROPERTY(EditAnywhere, BlueprintReadWrite, Getter, Setter, Category="Progress")
 	TEnumAsByte<EProgressBarFillStyle::Type> BarFillStyle;
 	
-	UPROPERTY(EditAnywhere, BlueprintReadOnly, Category=Progress)
+	/** */
+	UE_DEPRECATED(5.1, "Direct access to bIsMarquee is deprecated. Please use the getter or setter.")
+	UPROPERTY(EditAnywhere, BlueprintReadWrite, Getter="UseMarquee", Setter="SetIsMarquee", BlueprintSetter="SetIsMarquee", Category="Progress")
 	bool bIsMarquee;
 
-	UPROPERTY(EditAnywhere, BlueprintReadOnly, Category=Progress)
+	/** */
+	UE_DEPRECATED(5.1, "Direct access to BorderPadding is deprecated. Please use the getter or setter.")
+	UPROPERTY(EditAnywhere, BlueprintReadWrite, Getter, Setter, Category="Progress")
 	FVector2D BorderPadding;
 
 	/** A bindable delegate to allow logic to drive the text of the widget */
@@ -53,7 +58,8 @@ public:
 	FGetFloat PercentDelegate;
 
 	/** Fill Color and Opacity */
-	UPROPERTY(EditAnywhere, BlueprintReadOnly, Category=Appearance)
+	UE_DEPRECATED(5.1, "Direct access to FillColorAndOpacity is deprecated. Please use the getter or setter.")
+	UPROPERTY(EditAnywhere, BlueprintReadWrite, FieldNotify, Getter, Setter, BlueprintSetter="SetFillColorAndOpacity", Category="Appearance")
 	FLinearColor FillColorAndOpacity;
 
 	/** */
@@ -62,19 +68,50 @@ public:
 
 public:
 	
+	/** */
+	const FProgressBarStyle& GetWidgetStyle() const;
+
+	/**/
+	void SetWidgetStyle(const FProgressBarStyle& InStyle);
+
+	/** */
+	float GetPercent() const;
+
 	/** Sets the current value of the ProgressBar. */
 	UFUNCTION(BlueprintCallable, Category="Progress")
 	void SetPercent(float InPercent);
 
-	/** Sets the fill color of the progress bar. */
-	UFUNCTION(BlueprintCallable, Category="Progress")
-	void SetFillColorAndOpacity(FLinearColor InColor);
+	/** */
+	EProgressBarFillType::Type GetBarFillType() const;
+
+	/** */
+	void SetBarFillType(EProgressBarFillType::Type InBarFillType);
+
+	/** */
+	EProgressBarFillStyle::Type GetBarFillStyle() const;
+
+	/** */
+	void SetBarFillStyle(EProgressBarFillStyle::Type InBarFillStyle);
+
+	/** */
+	bool UseMarquee() const;
 
 	/** Sets the progress bar to show as a marquee. */
 	UFUNCTION(BlueprintCallable, Category="Progress")
 	void SetIsMarquee(bool InbIsMarquee);
 
-	//TODO UMG Add Set BarFillType.
+	/** */
+	FVector2D GetBorderPadding() const;
+
+	/** */
+	void SetBorderPadding(FVector2D InBorderPadding);
+
+	/** */
+	FLinearColor GetFillColorAndOpacity() const;
+
+	/** Sets the fill color of the progress bar. */
+	UFUNCTION(BlueprintCallable, Category="Progress")
+	void SetFillColorAndOpacity(FLinearColor InColor);
 
 public:
 	
