Fix to DrawLine Blueprint API
#prefllight 627fd8343339713c4a9d9228

diff --git a/Engine/Source/Runtime/UMG/Private/WidgetBlueprintLibrary.cpp b/Engine/Source/Runtime/UMG/Private/WidgetBlueprintLibrary.cpp
--- a/Engine/Source/Runtime/UMG/Private/WidgetBlueprintLibrary.cpp
+++ b/Engine/Source/Runtime/UMG/Private/WidgetBlueprintLibrary.cpp
@@ -181,26 +181,71 @@ void UWidgetBlueprintLibrary::DrawLine(FPaintContext& Context, FVector2D Positio
 	Points.Add(PositionA);
 	Points.Add(PositionB);
 
-	FSlateDrawElement::MakeLines(
-		Context.OutDrawElements,
-		Context.MaxLayer,
-		Context.AllottedGeometry.ToPaintGeometry(),
-		Points,
-		ESlateDrawEffect::None,
-		Tint,
-		bAntiAlias,
-		Thickness);
+	if ((PositionA - PositionB).SquaredLength() > KINDA_SMALL_NUMBER)
+	{
+		FSlateDrawElement::MakeLines(
+			Context.OutDrawElements,
+			Context.MaxLayer,
+			Context.AllottedGeometry.ToPaintGeometry(),
+			Points,
+			ESlateDrawEffect::None,
+			Tint,
+			bAntiAlias,
+			Thickness);
+	}
 }
 
 void UWidgetBlueprintLibrary::DrawLines(FPaintContext& Context, const TArray<FVector2D>& Points, FLinearColor Tint, bool bAntiAlias, float Thickness)
 {
+	if (Points.Num() < 2)
+	{
+		return;
+	}
+
+	// We need to trim points that might be overlapping
+	TArray<FVector2D> ValidatedPoints;
+	const TArray<FVector2D>* LinePoints = &Points;
+	FVector2D LastPoint = Points[0];
+	for (int32 Index = 1; Index < Points.Num(); Index++)
+	{
+		FVector2D CurrentPoint = Points[Index];
+		// If a the distance between two point is very small, remove it.
+		if ((CurrentPoint - LastPoint).SquaredLength() < KINDA_SMALL_NUMBER)
+		{
+			// The validated list doesn't exist. Allocate it.
+			if (LinePoints != &ValidatedPoints)
+			{
+				// We will reserve one less since we are removing this one.
+				ValidatedPoints.Reserve(Points.Num()-1);
+				for (int32 AddIndex = 0; AddIndex < Index ; AddIndex++)
+				{
+					ValidatedPoints.Push(Points[AddIndex]);
+				}
+				LinePoints = &ValidatedPoints;
+			}
+		}
+		// When the point is valid we must add it to the list only if the array already contains points.
+		else
+		{
+			if (!ValidatedPoints.IsEmpty())
+			{
+				ValidatedPoints.Push(CurrentPoint);
+			}
+			LastPoint = CurrentPoint;
+		}
+	}
+	if (LinePoints->Num() < 2)
+	{
+		return;
+	}
+
 	Context.MaxLayer++;
 
 	FSlateDrawElement::MakeLines(
 		Context.OutDrawElements,
 		Context.MaxLayer,
 		Context.AllottedGeometry.ToPaintGeometry(),
-		Points,
+		*LinePoints,
 		ESlateDrawEffect::None,
 		Tint,
 		bAntiAlias,
