UMG: Update Spacer, SizeBox, UniformGridPanel, WrapBox with accessors.

diff --git a/Engine/Source/Runtime/UMG/Private/Components/SizeBox.cpp b/Engine/Source/Runtime/UMG/Private/Components/SizeBox.cpp
--- a/Engine/Source/Runtime/UMG/Private/Components/SizeBox.cpp
+++ b/Engine/Source/Runtime/UMG/Private/Components/SizeBox.cpp
@@ -16,8 +16,10 @@ USizeBox::USizeBox(const FObjectInitializer& ObjectInitializer)
 {
 	bIsVariable = false;
 	SetVisibilityInternal(ESlateVisibility::SelfHitTestInvisible);
-	MinAspectRatio = 1;
-	MaxAspectRatio = 1;
+PRAGMA_DISABLE_DEPRECATION_WARNINGS
+	MinAspectRatio = 1.f;
+	MaxAspectRatio = 1.f;
+PRAGMA_ENABLE_DEPRECATION_WARNINGS
 }
 
 void USizeBox::ReleaseSlateResources(bool bReleaseChildren)
@@ -39,6 +41,12 @@ TSharedRef<SWidget> USizeBox::RebuildWidget()
 	return MySizeBox.ToSharedRef();
 }
 
+PRAGMA_DISABLE_DEPRECATION_WARNINGS
+float USizeBox::GetWidthOverride() const
+{
+	return WidthOverride;
+}
+
 void USizeBox::SetWidthOverride(float InWidthOverride)
 {
 	bOverride_WidthOverride = true;
@@ -58,6 +66,11 @@ void USizeBox::ClearWidthOverride()
 	}
 }
 
+float USizeBox::GetHeightOverride() const
+{
+	return HeightOverride;
+}
+
 void USizeBox::SetHeightOverride(float InHeightOverride)
 {
 	bOverride_HeightOverride = true;
@@ -77,6 +90,11 @@ void USizeBox::ClearHeightOverride()
 	}
 }
 
+float USizeBox::GetMinDesiredWidth() const
+{
+	return MinDesiredWidth;
+}
+
 void USizeBox::SetMinDesiredWidth(float InMinDesiredWidth)
 {
 	bOverride_MinDesiredWidth = true;
@@ -96,6 +114,11 @@ void USizeBox::ClearMinDesiredWidth()
 	}
 }
 
+float USizeBox::GetMinDesiredHeight() const
+{
+	return MinDesiredHeight;
+}
+
 void USizeBox::SetMinDesiredHeight(float InMinDesiredHeight)
 {
 	bOverride_MinDesiredHeight = true;
@@ -115,6 +138,11 @@ void USizeBox::ClearMinDesiredHeight()
 	}
 }
 
+float USizeBox::GetMaxDesiredWidth() const
+{
+	return MaxDesiredWidth;
+}
+
 void USizeBox::SetMaxDesiredWidth(float InMaxDesiredWidth)
 {
 	bOverride_MaxDesiredWidth = true;
@@ -134,6 +162,11 @@ void USizeBox::ClearMaxDesiredWidth()
 	}
 }
 
+float USizeBox::GetMaxDesiredHeight() const
+{
+	return MaxDesiredHeight;
+}
+
 void USizeBox::SetMaxDesiredHeight(float InMaxDesiredHeight)
 {
 	bOverride_MaxDesiredHeight = true;
@@ -153,6 +186,11 @@ void USizeBox::ClearMaxDesiredHeight()
 	}
 }
 
+float USizeBox::GetMinAspectRatio() const
+{
+	return MinAspectRatio;
+}
+
 void USizeBox::SetMinAspectRatio(float InMinAspectRatio)
 {
 	bOverride_MinAspectRatio = true;
@@ -172,6 +210,11 @@ void USizeBox::ClearMinAspectRatio()
 	}
 }
 
+float USizeBox::GetMaxAspectRatio() const
+{
+	return MaxAspectRatio;
+}
+
 void USizeBox::SetMaxAspectRatio(float InMaxAspectRatio)
 {
 	bOverride_MaxAspectRatio = true;
@@ -267,6 +310,7 @@ void USizeBox::SynchronizeProperties()
 		ClearMaxAspectRatio();
 	}
 }
+PRAGMA_ENABLE_DEPRECATION_WARNINGS
 
 UClass* USizeBox::GetSlotClass() const
 {
diff --git a/Engine/Source/Runtime/UMG/Private/Components/Spacer.cpp b/Engine/Source/Runtime/UMG/Private/Components/Spacer.cpp
--- a/Engine/Source/Runtime/UMG/Private/Components/Spacer.cpp
+++ b/Engine/Source/Runtime/UMG/Private/Components/Spacer.cpp
@@ -9,6 +9,7 @@
 /////////////////////////////////////////////////////
 // USpacer
 
+PRAGMA_DISABLE_DEPRECATION_WARNINGS
 USpacer::USpacer(const FObjectInitializer& ObjectInitializer)
 	: Super(ObjectInitializer)
 	, Size(1.0f, 1.0f)
@@ -16,6 +17,7 @@ USpacer::USpacer(const FObjectInitializer& ObjectInitializer)
 	bIsVariable = false;
 	SetVisibilityInternal(ESlateVisibility::SelfHitTestInvisible);
 }
+PRAGMA_ENABLE_DEPRECATION_WARNINGS
 
 void USpacer::ReleaseSlateResources(bool bReleaseChildren)
 {
@@ -24,6 +26,17 @@ void USpacer::ReleaseSlateResources(bool bReleaseChildren)
 	MySpacer.Reset();
 }
 
+PRAGMA_DISABLE_DEPRECATION_WARNINGS
+FVector2D USpacer::GetSize() const
+{
+	if (MySpacer.IsValid())
+	{
+		return MySpacer->GetSize();
+	}
+	
+	return Size;
+}
+
 void USpacer::SetSize(FVector2D InSize)
 {
 	Size = InSize;
@@ -36,10 +49,8 @@ void USpacer::SetSize(FVector2D InSize)
 
 TSharedRef<SWidget> USpacer::RebuildWidget()
 {
-	MySpacer = SNew(SSpacer);
-
-	//TODO UMG Consider using a design time wrapper for spacer to show expandy arrows or some other
-	// indicator that there's a widget at work here.
+	MySpacer = SNew(SSpacer)
+		.Size(Size);
 	
 	return MySpacer.ToSharedRef();
 }
@@ -50,6 +61,7 @@ void USpacer::SynchronizeProperties()
 	
 	MySpacer->SetSize(Size);
 }
+PRAGMA_ENABLE_DEPRECATION_WARNINGS
 
 #if WITH_EDITOR
 
diff --git a/Engine/Source/Runtime/UMG/Private/Components/UniformGridPanel.cpp b/Engine/Source/Runtime/UMG/Private/Components/UniformGridPanel.cpp
--- a/Engine/Source/Runtime/UMG/Private/Components/UniformGridPanel.cpp
+++ b/Engine/Source/Runtime/UMG/Private/Components/UniformGridPanel.cpp
@@ -80,6 +80,12 @@ UUniformGridSlot* UUniformGridPanel::AddChildToUniformGrid(UWidget* Content, int
 	return GridSlot;
 }
 
+PRAGMA_DISABLE_DEPRECATION_WARNINGS
+FMargin UUniformGridPanel::GetSlotPadding() const
+{
+	return SlotPadding;
+}
+
 void UUniformGridPanel::SetSlotPadding(FMargin InSlotPadding)
 {
 	SlotPadding = InSlotPadding;
@@ -89,6 +95,11 @@ void UUniformGridPanel::SetSlotPadding(FMargin InSlotPadding)
 	}
 }
 
+float UUniformGridPanel::GetMinDesiredSlotWidth() const
+{
+	return MinDesiredSlotWidth;
+}
+
 void UUniformGridPanel::SetMinDesiredSlotWidth(float InMinDesiredSlotWidth)
 {
 	MinDesiredSlotWidth = InMinDesiredSlotWidth;
@@ -98,6 +109,11 @@ void UUniformGridPanel::SetMinDesiredSlotWidth(float InMinDesiredSlotWidth)
 	}
 }
 
+float UUniformGridPanel::GetMinDesiredSlotHeight() const
+{
+	return MinDesiredSlotHeight;
+}
+
 void UUniformGridPanel::SetMinDesiredSlotHeight(float InMinDesiredSlotHeight)
 {
 	MinDesiredSlotHeight = InMinDesiredSlotHeight;
@@ -115,6 +131,7 @@ void UUniformGridPanel::SynchronizeProperties()
 	MyUniformGridPanel->SetMinDesiredSlotWidth(MinDesiredSlotWidth);
 	MyUniformGridPanel->SetMinDesiredSlotHeight(MinDesiredSlotHeight);
 }
+PRAGMA_ENABLE_DEPRECATION_WARNINGS
 
 #if WITH_EDITOR
 
diff --git a/Engine/Source/Runtime/UMG/Private/Components/WrapBox.cpp b/Engine/Source/Runtime/UMG/Private/Components/WrapBox.cpp
--- a/Engine/Source/Runtime/UMG/Private/Components/WrapBox.cpp
+++ b/Engine/Source/Runtime/UMG/Private/Components/WrapBox.cpp
@@ -11,12 +11,14 @@
 UWrapBox::UWrapBox(const FObjectInitializer& ObjectInitializer)
 	: Super(ObjectInitializer)
 {
+PRAGMA_DISABLE_DEPRECATION_WARNINGS
 	bIsVariable = false;
 	SetVisibilityInternal(ESlateVisibility::SelfHitTestInvisible);
 	WrapSize = 500;
 	bExplicitWrapSize = false;
 	HorizontalAlignment = HAlign_Left;
 	Orientation = EOrientation::Orient_Horizontal;
+PRAGMA_ENABLE_DEPRECATION_WARNINGS
 }
 
 void UWrapBox::ReleaseSlateResources(bool bReleaseChildren)
@@ -60,11 +62,13 @@ UWrapBoxSlot* UWrapBox::AddChildToWrapBox(UWidget* Content)
 
 TSharedRef<SWidget> UWrapBox::RebuildWidget()
 {
+PRAGMA_DISABLE_DEPRECATION_WARNINGS
 	MyWrapBox = SNew(SWrapBox)
 		.UseAllottedSize(!bExplicitWrapSize)
 		.PreferredSize(WrapSize)
 		.HAlign(HorizontalAlignment)
 		.Orientation(Orientation);
+PRAGMA_ENABLE_DEPRECATION_WARNINGS
 
 	for ( UPanelSlot* PanelSlot : Slots )
 	{
@@ -78,6 +82,7 @@ TSharedRef<SWidget> UWrapBox::RebuildWidget()
 	return MyWrapBox.ToSharedRef();
 }
 
+PRAGMA_DISABLE_DEPRECATION_WARNINGS
 void UWrapBox::SynchronizeProperties()
 {
 	Super::SynchronizeProperties();
@@ -89,26 +94,77 @@ void UWrapBox::SynchronizeProperties()
 	MyWrapBox->SetOrientation(Orientation);
 }
 
+FVector2D UWrapBox::GetInnerSlotPadding() const
+{
+	return InnerSlotPadding;
+}
+
 void UWrapBox::SetInnerSlotPadding(FVector2D InPadding)
 {
 	InnerSlotPadding = InPadding;
-
-	if ( MyWrapBox.IsValid() )
+	if (MyWrapBox.IsValid())
 	{
 		MyWrapBox->SetInnerSlotPadding(InPadding);
 	}
 }
 
-void UWrapBox::SetHorizontalAlignment(EHorizontalAlignment InHorizontalAlignment)
+float UWrapBox::GetWrapSize() const
 {
-	HorizontalAlignment = InHorizontalAlignment;
+	return WrapSize;
+}
 
+void UWrapBox::SetWrapSize(float InWrapSize)
+{
+	WrapSize = InWrapSize;
+	if (MyWrapBox.IsValid())
+	{
+		MyWrapBox->SetWrapSize(InWrapSize);
+	}
+}
+
+bool UWrapBox::UseExplicitWrapSize() const
+{
+	return bExplicitWrapSize;
+}
+
+void UWrapBox::SetExplicitWrapSize(bool bInExplicitWrapSize)
+{
+	bExplicitWrapSize = bInExplicitWrapSize;
 	if (MyWrapBox.IsValid())
+	{
+		MyWrapBox->SetUseAllottedSize(!bExplicitWrapSize);
+	}
+}
+
+EHorizontalAlignment UWrapBox::GetHorizontalAlignment() const
+{
+	return HorizontalAlignment;
+}
+
+void UWrapBox::SetHorizontalAlignment(EHorizontalAlignment InHorizontalAlignment)
+{
+	HorizontalAlignment = InHorizontalAlignment;
+	if (MyWrapBox)
 	{
 		MyWrapBox->SetHorizontalAlignment(InHorizontalAlignment);
 	}
 }
 
+EOrientation UWrapBox::GetOrientation() const
+{
+	return Orientation;
+}
+
+void UWrapBox::SetOrientation(EOrientation InOrientation)
+{
+	Orientation = InOrientation;
+	if (MyWrapBox)
+	{
+		MyWrapBox->SetOrientation(InOrientation);
+	}
+}
+PRAGMA_ENABLE_DEPRECATION_WARNINGS
+
 #if WITH_EDITOR
 
 const FText UWrapBox::GetPaletteCategory()
diff --git a/Engine/Source/Runtime/UMG/Public/Components/SizeBox.h b/Engine/Source/Runtime/UMG/Public/Components/SizeBox.h
--- a/Engine/Source/Runtime/UMG/Public/Components/SizeBox.h
+++ b/Engine/Source/Runtime/UMG/Public/Components/SizeBox.h
@@ -28,70 +28,91 @@ protected:
 public:
 
 	/** When specified, ignore the content's desired size and report the WidthOverride as the Box's desired width. */
-	UPROPERTY(EditAnywhere, BlueprintReadOnly, Category="Child Layout", meta=( editcondition="bOverride_WidthOverride" ))
+	UE_DEPRECATED(5.1, "Direct access to WidthOverride is deprecated. Please use the getter or setter.")
+	UPROPERTY(EditAnywhere, BlueprintReadWrite, Getter, Setter, BlueprintSetter="SetWidthOverride", Category="Child Layout", meta=( editcondition="bOverride_WidthOverride" ))
 	float WidthOverride;
 
 	/** When specified, ignore the content's desired size and report the HeightOverride as the Box's desired height. */
-	UPROPERTY(EditAnywhere, BlueprintReadOnly, Category="Child Layout", meta=( editcondition="bOverride_HeightOverride" ))
+	UE_DEPRECATED(5.1, "Direct access to HeightOverride is deprecated. Please use the getter or setter.")
+	UPROPERTY(EditAnywhere, BlueprintReadWrite, Getter, Setter, BlueprintSetter="SetHeightOverride", Category="Child Layout", meta=( editcondition="bOverride_HeightOverride" ))
 	float HeightOverride;
 
 	/** When specified, will report the MinDesiredWidth if larger than the content's desired width. */
-	UPROPERTY(EditAnywhere, BlueprintReadOnly, Category="Child Layout", meta=( editcondition="bOverride_MinDesiredWidth" ))
+	UE_DEPRECATED(5.1, "Direct access to MinDesiredWidth is deprecated. Please use the getter or setter.")
+	UPROPERTY(EditAnywhere, BlueprintReadWrite, Getter, Setter, BlueprintSetter="SetMinDesiredWidth", Category="Child Layout", meta=( editcondition="bOverride_MinDesiredWidth" ))
 	float MinDesiredWidth;
 
 	/** When specified, will report the MinDesiredHeight if larger than the content's desired height. */
-	UPROPERTY(EditAnywhere, BlueprintReadOnly, Category="Child Layout", meta=( editcondition="bOverride_MinDesiredHeight" ))
+	UE_DEPRECATED(5.1, "Direct access to MinDesiredHeight is deprecated. Please use the getter or setter.")
+	UPROPERTY(EditAnywhere, BlueprintReadWrite, Getter, Setter, BlueprintSetter="SetMinDesiredHeight", Category="Child Layout", meta=( editcondition="bOverride_MinDesiredHeight" ))
 	float MinDesiredHeight;
 
 	/** When specified, will report the MaxDesiredWidth if smaller than the content's desired width. */
-	UPROPERTY(EditAnywhere, BlueprintReadOnly, Category="Child Layout", meta=( editcondition="bOverride_MaxDesiredWidth" ))
+	UE_DEPRECATED(5.1, "Direct access to MaxDesiredWidth is deprecated. Please use the getter or setter.")
+	UPROPERTY(EditAnywhere, BlueprintReadWrite, Getter, Setter, BlueprintSetter="SetMaxDesiredWidth", Category="Child Layout", meta=( editcondition="bOverride_MaxDesiredWidth" ))
 	float MaxDesiredWidth;
 
 	/** When specified, will report the MaxDesiredHeight if smaller than the content's desired height. */
-	UPROPERTY(EditAnywhere, BlueprintReadOnly, Category="Child Layout", meta=( editcondition="bOverride_MaxDesiredHeight" ))
+	UE_DEPRECATED(5.1, "Direct access to MaxDesiredHeight is deprecated. Please use the getter or setter.")
+	UPROPERTY(EditAnywhere, BlueprintReadWrite, Getter, Setter, BlueprintSetter="SetMaxDesiredHeight", Category="Child Layout", meta=( editcondition="bOverride_MaxDesiredHeight" ))
 	float MaxDesiredHeight;
 
 	/** */
-	UPROPERTY(EditAnywhere, BlueprintReadOnly, Category = "Child Layout", meta = (editcondition = "bOverride_MinAspectRatio"))
+	UE_DEPRECATED(5.1, "Direct access to MinAspectRatio is deprecated. Please use the getter or setter.")
+	UPROPERTY(EditAnywhere, BlueprintReadWrite, Getter, Setter, BlueprintSetter="SetMinAspectRatio", Category = "Child Layout", meta = (editcondition = "bOverride_MinAspectRatio"))
 	float MinAspectRatio;
 
 	/** */
-	UPROPERTY(EditAnywhere, BlueprintReadOnly, Category="Child Layout", meta=( editcondition="bOverride_MaxAspectRatio" ))
+	UE_DEPRECATED(5.1, "Direct access to MaxAspectRatio is deprecated. Please use the getter or setter.")
+	UPROPERTY(EditAnywhere, BlueprintReadWrite, Getter, Setter, BlueprintSetter="SetMaxAspectRatio", Category="Child Layout", meta=( editcondition="bOverride_MaxAspectRatio" ))
 	float MaxAspectRatio;
 
 	/**  */
+	UE_DEPRECATED(5.1, "Direct access to bOverride_WidthOverride is deprecated. Please use the SetWidthOverride or ClearWidthOverride.")
 	UPROPERTY(EditAnywhere, Category="Child Layout", meta=(InlineEditConditionToggle))
 	uint8 bOverride_WidthOverride : 1;
 
 	/**  */
+	UE_DEPRECATED(5.1, "Direct access to bOverride_HeightOverride is deprecated. Please use the SetHeightOverride or ClearHeightOverride.")
 	UPROPERTY(EditAnywhere, Category="Child Layout", meta=(InlineEditConditionToggle))
 	uint8 bOverride_HeightOverride : 1;
 
 	/**  */
+	UE_DEPRECATED(5.1, "Direct access to bOverride_MinDesiredWidth is deprecated. Please use the SetMinDesiredWidth or ClearMinDesiredWidth.")
 	UPROPERTY(EditAnywhere, Category="Child Layout", meta=(InlineEditConditionToggle))
 	uint8 bOverride_MinDesiredWidth : 1;
 
 	/**  */
+	UE_DEPRECATED(5.1, "Direct access to bOverride_MinDesiredHeight is deprecated. Please use the SetMinDesiredHeight or ClearMinDesiredHeight.")
 	UPROPERTY(EditAnywhere, Category="Child Layout", meta=(InlineEditConditionToggle))
 	uint8 bOverride_MinDesiredHeight : 1;
 
 	/**  */
+	UE_DEPRECATED(5.1, "Direct access to bOverride_MaxDesiredWidth is deprecated. Please use the SetMaxDesiredWidth or ClearMaxDesiredWidth.")
 	UPROPERTY(EditAnywhere, Category="Child Layout", meta=(InlineEditConditionToggle))
 	uint8 bOverride_MaxDesiredWidth : 1;
 
 	/**  */
+	UE_DEPRECATED(5.1, "Direct access to bOverride_MaxDesiredHeight is deprecated. Please use the SetMaxDesiredHeight or ClearMaxDesiredHeight.")
 	UPROPERTY(EditAnywhere, Category="Child Layout", meta=(InlineEditConditionToggle))
 	uint8 bOverride_MaxDesiredHeight : 1;
 
 	/**  */
+	UE_DEPRECATED(5.1, "Direct access to bOverride_MinAspectRatio is deprecated. Please use the SetMinAspectRatio or ClearMinAspectRatio.")
 	UPROPERTY(EditAnywhere, Category="Child Layout", meta=(InlineEditConditionToggle))
 	uint8 bOverride_MinAspectRatio : 1;
 
 	/**  */
+	UE_DEPRECATED(5.1, "Direct access to bOverride_MaxAspectRatio is deprecated. Please use the SetMaxAspectRatio  or ClearMaxAspectRatio .")
 	UPROPERTY(EditAnywhere, Category = "Child Layout", meta=(InlineEditConditionToggle))
 	uint8 bOverride_MaxAspectRatio : 1;
 
-		
+	
+public:
+
+	/** */
+	float GetWidthOverride() const;
+
 	/** When specified, ignore the content's desired size and report the WidthOverride as the Box's desired width. */
 	UFUNCTION(BlueprintCallable, Category="Layout|Size Box")
 	void SetWidthOverride(float InWidthOverride);
@@ -99,6 +120,9 @@ public:
 	UFUNCTION(BlueprintCallable, Category="Layout|Size Box")
 	void ClearWidthOverride();
 
+	/** */
+	float GetHeightOverride() const;
+
 	/** When specified, ignore the content's desired size and report the HeightOverride as the Box's desired height. */
 	UFUNCTION(BlueprintCallable, Category="Layout|Size Box")
 	void SetHeightOverride(float InHeightOverride);
@@ -106,6 +130,9 @@ public:
 	UFUNCTION(BlueprintCallable, Category="Layout|Size Box")
 	void ClearHeightOverride();
 
+	/** */
+	float GetMinDesiredWidth() const;
+
 	/** When specified, will report the MinDesiredWidth if larger than the content's desired width. */
 	UFUNCTION(BlueprintCallable, Category="Layout|Size Box")
 	void SetMinDesiredWidth(float InMinDesiredWidth);
@@ -113,6 +140,9 @@ public:
 	UFUNCTION(BlueprintCallable, Category="Layout|Size Box")
 	void ClearMinDesiredWidth();
 
+	/** */
+	float GetMinDesiredHeight() const;
+
 	/** When specified, will report the MinDesiredHeight if larger than the content's desired height. */
 	UFUNCTION(BlueprintCallable, Category="Layout|Size Box")
 	void SetMinDesiredHeight(float InMinDesiredHeight);
@@ -120,6 +150,9 @@ public:
 	UFUNCTION(BlueprintCallable, Category="Layout|Size Box")
 	void ClearMinDesiredHeight();
 
+	/** */
+	float GetMaxDesiredWidth() const;
+
 	/** When specified, will report the MaxDesiredWidth if smaller than the content's desired width. */
 	UFUNCTION(BlueprintCallable, Category="Layout|Size Box")
 	void SetMaxDesiredWidth(float InMaxDesiredWidth);
@@ -127,6 +160,9 @@ public:
 	UFUNCTION(BlueprintCallable, Category="Layout|Size Box")
 	void ClearMaxDesiredWidth();
 
+	/** */
+	float GetMaxDesiredHeight() const;
+
 	/** When specified, will report the MaxDesiredHeight if smaller than the content's desired height. */
 	UFUNCTION(BlueprintCallable, Category="Layout|Size Box")
 	void SetMaxDesiredHeight(float InMaxDesiredHeight);
@@ -134,16 +170,22 @@ public:
 	UFUNCTION(BlueprintCallable, Category="Layout|Size Box")
 	void ClearMaxDesiredHeight();
 
-	UFUNCTION(BlueprintCallable, Category="Layout|Size Box")
-	void SetMinAspectRatio(float InMinAspectRatio);
+	/** */
+	float GetMinAspectRatio() const;
 
 	UFUNCTION(BlueprintCallable, Category="Layout|Size Box")
-	void SetMaxAspectRatio(float InMaxAspectRatio);
+	void SetMinAspectRatio(float InMinAspectRatio);
 
 	UFUNCTION(BlueprintCallable, Category="Layout|Size Box")
 	void ClearMinAspectRatio();
 
+	/** */
+	float GetMaxAspectRatio() const;
+
 	UFUNCTION(BlueprintCallable, Category="Layout|Size Box")
+	void SetMaxAspectRatio(float InMaxAspectRatio);
+
+	UFUNCTION(BlueprintCallable, Category = "Layout|Size Box")
 	void ClearMaxAspectRatio();
 
 
diff --git a/Engine/Source/Runtime/UMG/Public/Components/Spacer.h b/Engine/Source/Runtime/UMG/Public/Components/Spacer.h
--- a/Engine/Source/Runtime/UMG/Public/Components/Spacer.h
+++ b/Engine/Source/Runtime/UMG/Public/Components/Spacer.h
@@ -23,31 +23,32 @@ class UMG_API USpacer : public UWidget
 public:
 
 	/** The size of the spacer */
-	UPROPERTY(EditAnywhere, BlueprintReadOnly, Category=Appearance)
+	UE_DEPRECATED(5.1, "Direct access to Size is deprecated. Please use the getter or setter.")
+	UPROPERTY(EditAnywhere, BlueprintReadWrite, Getter, Setter, BlueprintSetter="SetSize", Category="Appearance")
 	FVector2D Size;
 
 public:
 
+	/** */
+	FVector2D GetSize() const;
+
 	/** Sets the size of the spacer */
 	UFUNCTION(BlueprintCallable, Category="Widget")
 	void SetSize(FVector2D InSize);
 	
-	// UWidget interface
+	//~ UWidget interface
 	virtual void SynchronizeProperties() override;
-	// End of UWidget interface
-
-	// UVisual interface
 	virtual void ReleaseSlateResources(bool bReleaseChildren) override;
-	// End of UVisual interface
+	//~ End of UVisual interface
 
 #if WITH_EDITOR
 	virtual const FText GetPaletteCategory() override;
 #endif
 
 protected:
-	// UWidget interface
+	//~ UWidget interface
 	virtual TSharedRef<SWidget> RebuildWidget() override;
-	// End of UWidget interface
+	//~ End of UWidget interface
 
 protected:
 	TSharedPtr<SSpacer> MySpacer;
diff --git a/Engine/Source/Runtime/UMG/Public/Components/UniformGridPanel.h b/Engine/Source/Runtime/UMG/Public/Components/UniformGridPanel.h
--- a/Engine/Source/Runtime/UMG/Public/Components/UniformGridPanel.h
+++ b/Engine/Source/Runtime/UMG/Public/Components/UniformGridPanel.h
@@ -23,27 +23,39 @@ class UMG_API UUniformGridPanel : public UPanelWidget
 public:
 
 	/** Padding given to each slot */
-	UPROPERTY(EditAnywhere, BlueprintReadOnly, Category="Child Layout")
+	UE_DEPRECATED(5.1, "Direct access to SlotPadding is deprecated. Please use the getter or setter.")
+	UPROPERTY(EditAnywhere, BlueprintReadWrite, Getter, Setter, BlueprintSetter="SetSlotPadding", Category="Child Layout")
 	FMargin SlotPadding;
 
 	/** The minimum desired width of the slots */
-	UPROPERTY(EditAnywhere, BlueprintReadOnly, Category="Child Layout")
+	UE_DEPRECATED(5.1, "Direct access to MinDesiredSlotWidth is deprecated. Please use the getter or setter.")
+	UPROPERTY(EditAnywhere, BlueprintReadWrite, Getter, Setter, BlueprintSetter="SetMinDesiredSlotWidth", Category="Child Layout")
 	float MinDesiredSlotWidth;
 
 	/** The minimum desired height of the slots */
-	UPROPERTY(EditAnywhere, BlueprintReadOnly, Category="Child Layout")
+	UE_DEPRECATED(5.1, "Direct access to MinDesiredSlotHeight is deprecated. Please use the getter or setter.")
+	UPROPERTY(EditAnywhere, BlueprintReadWrite, Getter, Setter, BlueprintSetter="SetMinDesiredSlotHeight", Category="Child Layout")
 	float MinDesiredSlotHeight;
 
 public:
 
+	/** */
+	FMargin GetSlotPadding() const;
+
 	/**  */
 	UFUNCTION(BlueprintCallable, Category="Child Layout")
 	void SetSlotPadding(FMargin InSlotPadding);
 
+	/** */
+	float GetMinDesiredSlotWidth() const;
+
 	/**  */
 	UFUNCTION(BlueprintCallable, Category="Child Layout")
 	void SetMinDesiredSlotWidth(float InMinDesiredSlotWidth);
 
+	/** */
+	float GetMinDesiredSlotHeight() const;
+
 	/**  */
 	UFUNCTION(BlueprintCallable, Category="Child Layout")
 	void SetMinDesiredSlotHeight(float InMinDesiredSlotHeight);
@@ -56,15 +68,14 @@ public:
 
 public:
 
-	// UWidget interface
+	//~ UWidget interface
 	virtual void SynchronizeProperties() override;
-	// End of UWidget interface
-
 	virtual void ReleaseSlateResources(bool bReleaseChildren) override;
 
 #if WITH_EDITOR
 	virtual const FText GetPaletteCategory() override;
 #endif
+	//~ End of UWidget interface
 
 protected:
 
@@ -79,7 +90,7 @@ protected:
 	TSharedPtr<SUniformGridPanel> MyUniformGridPanel;
 
 protected:
-	// UWidget interface
+	//~ UWidget interface
 	virtual TSharedRef<SWidget> RebuildWidget() override;
-	// End of UWidget interface
+	//~ End of UWidget interface
 };
diff --git a/Engine/Source/Runtime/UMG/Public/Components/WrapBox.h b/Engine/Source/Runtime/UMG/Public/Components/WrapBox.h
--- a/Engine/Source/Runtime/UMG/Public/Components/WrapBox.h
+++ b/Engine/Source/Runtime/UMG/Public/Components/WrapBox.h
@@ -25,64 +25,95 @@ class UMG_API UWrapBox : public UPanelWidget
 
 public:
 	/** The inner slot padding goes between slots sharing borders */
-	UPROPERTY(EditAnywhere, BlueprintReadOnly, Category="Content Layout")
+	UE_DEPRECATED(5.1, "Direct access to InnerSlotPadding is deprecated. Please use the getter or setter.")
+	UPROPERTY(EditAnywhere, BlueprintReadWrite, Getter, Setter, BlueprintSetter = "SetInnerSlotPadding", Category = "Content Layout")
 	FVector2D InnerSlotPadding;
 
 	/** When this size is exceeded, elements will start appearing on the next line. */
-	UPROPERTY(EditAnywhere, BlueprintReadOnly, Category = "Content Layout", meta=(EditCondition = "bExplicitWrapSize"))
+	UE_DEPRECATED(5.1, "Direct access to WrapSize is deprecated. Please use the getter or setter.")
+	UPROPERTY(EditAnywhere, BlueprintReadWrite, Getter, Setter, Category = "Content Layout", meta=(EditCondition = "bExplicitWrapSize"))
 	float WrapSize;
 
 	/** Use explicit wrap size whenever possible. It greatly simplifies layout calculations and reduces likelihood of "wiggling UI" */
-	UPROPERTY(EditAnywhere, BlueprintReadOnly, Category = "Content Layout")
+	UE_DEPRECATED(5.1, "Direct access to bExplicitWrapSize is deprecated. Please use the getter or setter.")
+	UPROPERTY(EditAnywhere, BlueprintReadWrite, Getter="UseExplicitWrapSize", Setter = "SetExplicitWrapSize", Category = "Content Layout")
 	bool bExplicitWrapSize;
 
 	/** The alignment of each line of wrapped content. */
-	UPROPERTY(EditAnywhere, BlueprintReadOnly, Category = "Content Layout")
+	UE_DEPRECATED(5.1, "Direct access to HorizontalAlignment is deprecated. Please use the getter or setter.")
+	UPROPERTY(EditAnywhere, BlueprintReadWrite, Getter, Setter, BlueprintSetter = "SetHorizontalAlignment", Category = "Content Layout")
 	TEnumAsByte<EHorizontalAlignment> HorizontalAlignment;
 
 	/** Determines if the Wrap Box should arranges the widgets left-to-right or top-to-bottom */
-	UPROPERTY(EditAnywhere, BlueprintReadOnly, Category = "Content Layout")
+	UE_DEPRECATED(5.1, "Direct access to Orientation is deprecated. Please use the getter or setter.")
+	UPROPERTY(EditAnywhere, BlueprintReadWrite, Getter, Setter, Category = "Content Layout")
 	TEnumAsByte<EOrientation> Orientation = EOrientation::Orient_Horizontal;
 
+public:
+	/** Gets the inner slot padding goes between slots sharing borders */
+	FVector2D GetInnerSlotPadding() const;
+
 	/** Sets the inner slot padding goes between slots sharing borders */
 	UFUNCTION(BlueprintCallable, Category="Content Layout")
 	void SetInnerSlotPadding(FVector2D InPadding);
 
+	/** */
+	float GetWrapSize() const;
+
+	/** */
+	void SetWrapSize(float InWrapSize);
+
+	/** */
+	bool UseExplicitWrapSize() const;
+
+	/** */
+	void SetExplicitWrapSize(bool bInExplicitWrapSize);
+
+	/** */
+	EHorizontalAlignment GetHorizontalAlignment() const;
+
+	/** */
 	UFUNCTION(BlueprintCallable, Category="Content Layout")
 	void SetHorizontalAlignment(EHorizontalAlignment InHorizontalAlignment);
 
+	/** */
+	EOrientation GetOrientation() const;
+
+	/** */
+	void SetOrientation(EOrientation InOrientation);
+
 public:
 	UFUNCTION(BlueprintCallable, Category="Panel")
 	UWrapBoxSlot* AddChildToWrapBox(UWidget* Content);
 
 #if WITH_EDITOR
-	// UWidget interface
+	//~ UWidget interface
 	virtual const FText GetPaletteCategory() override;
-	// End UWidget interface
+	//~ End UWidget interface
 #endif
 
 protected:
 
-	// UPanelWidget
+	//~ UPanelWidget
 	virtual UClass* GetSlotClass() const override;
 	virtual void OnSlotAdded(UPanelSlot* Slot) override;
 	virtual void OnSlotRemoved(UPanelSlot* Slot) override;
-	// End UPanelWidget
+	//~ End UPanelWidget
 
-	// UWidget interface
+	//~ UWidget interface
 	virtual void SynchronizeProperties() override;
-	// End of UWidget interface
+	//~ End of UWidget interface
 
-	// UVisual interface
+	//~ UVisual interface
 	virtual void ReleaseSlateResources(bool bReleaseChildren) override;
-	// End of UVisual interface
+	//~ End of UVisual interface
 
 protected:
 
 	TSharedPtr<class SWrapBox> MyWrapBox;
 
 protected:
-	// UWidget interface
+	//~ UWidget interface
 	virtual TSharedRef<SWidget> RebuildWidget() override;
-	// End of UWidget interface
+	//~ End of UWidget interface
 };
