Fix to DrawLine Blueprint API and conversion to float for LWC changes

diff --git a/Engine/Source/Runtime/UMG/Private/WidgetBlueprintLibrary.cpp b/Engine/Source/Runtime/UMG/Private/WidgetBlueprintLibrary.cpp
--- a/Engine/Source/Runtime/UMG/Private/WidgetBlueprintLibrary.cpp
+++ b/Engine/Source/Runtime/UMG/Private/WidgetBlueprintLibrary.cpp
@@ -177,9 +177,9 @@ void UWidgetBlueprintLibrary::DrawLine(FPaintContext& Context, FVector2D Positio
 {
 	Context.MaxLayer++;
 
-	TArray<FVector2D> Points;
-	Points.Add(PositionA);
-	Points.Add(PositionB);
+	TArray<FVector2f> Points;
+	Points.Add(UE::Slate::CastToVector2f(PositionA));
+	Points.Add(UE::Slate::CastToVector2f(PositionB));
 
 	if ((PositionA - PositionB).SquaredLength() > KINDA_SMALL_NUMBER)
 	{
@@ -202,39 +202,22 @@ void UWidgetBlueprintLibrary::DrawLines(FPaintContext& Context, const TArray<FVe
 		return;
 	}
 
-	// We need to trim points that might be overlapping
-	TArray<FVector2D> ValidatedPoints;
-	const TArray<FVector2D>* LinePoints = &Points;
+	// We need to trim points that might be overlapping. We convert to Float at the same time
+	TArray<FVector2f> ValidatedPoints;
+	ValidatedPoints.Reserve(Points.Num());
+	ValidatedPoints.Push(UE::Slate::CastToVector2f(Points[0]));
 	FVector2D LastPoint = Points[0];
 	for (int32 Index = 1; Index < Points.Num(); Index++)
 	{
 		FVector2D CurrentPoint = Points[Index];
-		// If a the distance between two point is very small, remove it.
-		if ((CurrentPoint - LastPoint).SquaredLength() < KINDA_SMALL_NUMBER)
+		// If a the distance between two point is very small, do not add it.
+		if ((CurrentPoint - LastPoint).SquaredLength() > KINDA_SMALL_NUMBER)
 		{
-			// The validated list doesn't exist. Allocate it.
-			if (LinePoints != &ValidatedPoints)
-			{
-				// We will reserve one less since we are removing this one.
-				ValidatedPoints.Reserve(Points.Num()-1);
-				for (int32 AddIndex = 0; AddIndex < Index ; AddIndex++)
-				{
-					ValidatedPoints.Push(Points[AddIndex]);
-				}
-				LinePoints = &ValidatedPoints;
-			}
-		}
-		// When the point is valid we must add it to the list only if the array already contains points.
-		else
-		{
-			if (!ValidatedPoints.IsEmpty())
-			{
-				ValidatedPoints.Push(CurrentPoint);
-			}
+			ValidatedPoints.Push(UE::Slate::CastToVector2f(Points[Index]));
 			LastPoint = CurrentPoint;
 		}
 	}
-	if (LinePoints->Num() < 2)
+	if (ValidatedPoints.Num() < 2)
 	{
 		return;
 	}
@@ -245,7 +228,7 @@ void UWidgetBlueprintLibrary::DrawLines(FPaintContext& Context, const TArray<FVe
 		Context.OutDrawElements,
 		Context.MaxLayer,
 		Context.AllottedGeometry.ToPaintGeometry(),
-		*LinePoints,
+		ValidatedPoints,
 		ESlateDrawEffect::None,
 		Tint,
 		bAntiAlias,
