Adding MoveTemp to previous CL.

diff --git a/Engine/Source/Runtime/UMG/Private/WidgetBlueprintLibrary.cpp b/Engine/Source/Runtime/UMG/Private/WidgetBlueprintLibrary.cpp
--- a/Engine/Source/Runtime/UMG/Private/WidgetBlueprintLibrary.cpp
+++ b/Engine/Source/Runtime/UMG/Private/WidgetBlueprintLibrary.cpp
@@ -228,7 +228,7 @@ void UWidgetBlueprintLibrary::DrawLines(FPaintContext& Context, const TArray<FVe
 		Context.OutDrawElements,
 		Context.MaxLayer,
 		Context.AllottedGeometry.ToPaintGeometry(),
-		ValidatedPoints,
+		MoveTemp(ValidatedPoints),
 		ESlateDrawEffect::None,
 		Tint,
 		bAntiAlias,
