UMG - Adding comments to the NamedSlotInterface.

#trivial

diff --git a/Engine/Source/Runtime/UMG/Public/Components/NamedSlotInterface.h b/Engine/Source/Runtime/UMG/Public/Components/NamedSlotInterface.h
--- a/Engine/Source/Runtime/UMG/Public/Components/NamedSlotInterface.h
+++ b/Engine/Source/Runtime/UMG/Public/Components/NamedSlotInterface.h
@@ -25,22 +25,23 @@ class UMG_API INamedSlotInterface
 
 public:
 
-	/**  */
+	/** Gets the names for slots that we can store widgets into. */
 	virtual void GetSlotNames(TArray<FName>& SlotNames) const = 0;
 
-	/**  */
+	/** Gets the widget for a given slot by name, will return nullptr if no widget is in the slot. */
 	virtual UWidget* GetContentForSlot(FName SlotName) const = 0;
 
-	/**  */
+	/** Sets the widget for a given slot by name. */
 	virtual void SetContentForSlot(FName SlotName, UWidget* Content) = 0;
 
-	/**  */
+	/** Determines if any slot holds the given widget. */
 	bool ContainsContent(UWidget* Content) const;
 
-	/**  */
+	/** Releases named slot related resources. */
 	void ReleaseNamedSlotSlateResources(bool bReleaseChildren);
 
 #if WITH_EDITOR
+	/** Applies the designer flags to the all of the content in all the slots. */
 	void SetNamedSlotDesignerFlags(EWidgetDesignFlags NewFlags);
 #endif
 };
