UMG: GetWidgetTreeOwningClass use the widget class instead of the WidgetTree's owner.

diff --git a/Engine/Source/Runtime/UMG/Private/UserWidget.cpp b/Engine/Source/Runtime/UMG/Private/UserWidget.cpp
--- a/Engine/Source/Runtime/UMG/Private/UserWidget.cpp
+++ b/Engine/Source/Runtime/UMG/Private/UserWidget.cpp
@@ -1350,9 +1350,9 @@
 		UInputDelegateBinding::BindInputDelegates(GetClass(), PC->InputComponent, this);		
 	}
 	
-	if (UWidgetBlueprintGeneratedClass* Class = GetWidgetTreeOwningClass())
+	if (UWidgetBlueprintGeneratedClass* BPClass = Cast<UWidgetBlueprintGeneratedClass>(GetClass()))
 	{
-		Class->ForEachExtension([this](UWidgetBlueprintGeneratedClassExtension* Extension)
+		BPClass->ForEachExtension([this](UWidgetBlueprintGeneratedClassExtension* Extension)
 			{
 				Extension->Initialize(this);
 			});
@@ -1369,9 +1369,9 @@
 void UUserWidget::NativePreConstruct()
 {
 	const bool bIsDesignTime = IsDesignTime();
-	if (UWidgetBlueprintGeneratedClass* Class = GetWidgetTreeOwningClass())
+	if (UWidgetBlueprintGeneratedClass* BPClass = Cast<UWidgetBlueprintGeneratedClass>(GetClass()))
 	{
-		Class->ForEachExtension([this, bIsDesignTime](UWidgetBlueprintGeneratedClassExtension* Extension)
+		BPClass->ForEachExtension([this, bIsDesignTime](UWidgetBlueprintGeneratedClassExtension* Extension)
 			{
 				Extension->PreConstruct(this, bIsDesignTime);
 			});
@@ -1382,9 +1382,9 @@
 
 void UUserWidget::NativeConstruct()
 {
-	if (UWidgetBlueprintGeneratedClass* Class = GetWidgetTreeOwningClass())
+	if (UWidgetBlueprintGeneratedClass* BPClass = Cast<UWidgetBlueprintGeneratedClass>(GetClass()))
 	{
-		Class->ForEachExtension([this](UWidgetBlueprintGeneratedClassExtension* Extension)
+		BPClass->ForEachExtension([this](UWidgetBlueprintGeneratedClassExtension* Extension)
 			{
 				Extension->Construct(this);
 			});
@@ -1404,9 +1404,9 @@
 		Extension->Destruct();
 	}
 
-	if (UWidgetBlueprintGeneratedClass* Class = GetWidgetTreeOwningClass())
+	if (UWidgetBlueprintGeneratedClass* BPClass = Cast<UWidgetBlueprintGeneratedClass>(GetClass()))
 	{
-		Class->ForEachExtension([this](UWidgetBlueprintGeneratedClassExtension* Extension)
+		BPClass->ForEachExtension([this](UWidgetBlueprintGeneratedClassExtension* Extension)
 			{
 				Extension->Destruct(this);
 			});
diff --git a/Engine/Source/Runtime/UMG/Public/Blueprint/UserWidget.h b/Engine/Source/Runtime/UMG/Public/Blueprint/UserWidget.h
--- a/Engine/Source/Runtime/UMG/Public/Blueprint/UserWidget.h
+++ b/Engine/Source/Runtime/UMG/Public/Blueprint/UserWidget.h
@@ -212,7 +212,7 @@
 public:
 	UUserWidget(const FObjectInitializer& ObjectInitializer);
 
-	//UObject interface
+	//~ Begin UObject interface
 	virtual class UWorld* GetWorld() const override;
 	virtual void PostDuplicate(bool bDuplicateForPIE) override;
 	virtual void BeginDestroy() override;
@@ -225,6 +225,12 @@
 
 	EWidgetTickFrequency GetDesiredTickFrequency() const { return TickFrequency; }
 
+	/**
+	 * Returns the BlueprintGeneratedClass that generated the WidgetTree.
+	 * A child UserWidget that extends a parent UserWidget will not have a new WidgetTree.
+	 * The child UserWidget will have the same WidgetTree as the parent UserWidget.
+	 * This function returns the parent UserWidget's BlueprintClass.
+	 */
 	UWidgetBlueprintGeneratedClass* GetWidgetTreeOwningClass() const;
 
 	void UpdateCanTick();
@@ -236,7 +242,7 @@
 	void InitializeNamedSlots(bool bReparentToWidgetTree);
 
 public:
-	//UVisual interface
+	//~ Begin UVisual interface
 	virtual void ReleaseSlateResources(bool bReleaseChildren) override;
 	//~ End UVisual Interface
 
@@ -244,11 +250,11 @@
 	virtual void SynchronizeProperties() override;
 	//~ End UWidget Interface
 
-	// UNamedSlotInterface Begin
+	//~ Begin UNamedSlotInterface Begin
 	virtual void GetSlotNames(TArray<FName>& SlotNames) const override;
 	virtual UWidget* GetContentForSlot(FName SlotName) const override;
 	virtual void SetContentForSlot(FName SlotName, UWidget* Content) override;
-	// UNamedSlotInterface End
+	//~ UNamedSlotInterface End
 
 	/**
 	 * Adds it to the game's viewport and fills the entire screen, unless SetDesiredSizeInViewport is called
