Notifies users if a User Widget has a null animation track, prompting them to delete or rebind said track. Deleting null animation tracks can improve performance on larger user widgets.

diff --git a/Engine/Source/Editor/UMGEditor/Private/WidgetBlueprintCompiler.cpp b/Engine/Source/Editor/UMGEditor/Private/WidgetBlueprintCompiler.cpp
--- a/Engine/Source/Editor/UMGEditor/Private/WidgetBlueprintCompiler.cpp
+++ b/Engine/Source/Editor/UMGEditor/Private/WidgetBlueprintCompiler.cpp
@@ -1038,7 +1038,31 @@
 	TSubclassOf<UUserWidget> ClassContext;
 };
 
+void FWidgetBlueprintCompilerContext::ValidateWidgetAnimations()
+{
+	UWidgetBlueprintGeneratedClass* WidgetClass = NewWidgetBlueprintClass;
+	UWidgetBlueprint* WidgetBP = WidgetBlueprint();
+	UUserWidget& UserWidget = *WidgetClass->GetDefaultObject<UUserWidget>();
+	FBlueprintCompilerLog BlueprintLog(MessageLog, WidgetClass);
 
+	for (const UWidgetAnimation* InAnimation : WidgetBP->Animations)
+	{
+		for (const FWidgetAnimationBinding& Binding : InAnimation->AnimationBindings)
+		{
+			// Look for the object bindings within the widget
+			UObject* FoundObject = Binding.FindRuntimeObject(*WidgetBP->WidgetTree, UserWidget);
+
+			// If any of the FoundObjects is null, we do not play the animation.
+			if (FoundObject == nullptr)
+			{
+				// Notify the user of the null track in the editor
+				const FText AnimationNullTrackMessage = LOCTEXT("AnimationNullTrack", "UMG Animation '{0}' from '{1}' is trying to animate a non-existent widget through binding '{2}'. Please re-bind or delete this object from the animation.");
+				BlueprintLog.Warning(FText::Format(AnimationNullTrackMessage, InAnimation->GetDisplayName(), FText::FromString(UserWidget.GetClass()->GetName()), FText::FromName(Binding.WidgetName)));
+			}
+		}
+	}
+}
+
 void FWidgetBlueprintCompilerContext::OnPostCDOCompiled(const UObject::FPostCDOCompiledContext& Context)
 {
 	Super::OnPostCDOCompiled(Context);
@@ -1058,6 +1082,7 @@
 	{
 		FBlueprintCompilerLog BlueprintLog(MessageLog, WidgetClass);
 		WidgetClass->GetDefaultObject<UUserWidget>()->ValidateBlueprint(*WidgetBP->WidgetTree, BlueprintLog);
+		ValidateWidgetAnimations();
 	}
 }
 
diff --git a/Engine/Source/Editor/UMGEditor/Public/WidgetBlueprintCompiler.h b/Engine/Source/Editor/UMGEditor/Public/WidgetBlueprintCompiler.h
--- a/Engine/Source/Editor/UMGEditor/Public/WidgetBlueprintCompiler.h
+++ b/Engine/Source/Editor/UMGEditor/Public/WidgetBlueprintCompiler.h
@@ -58,6 +58,12 @@
 
 	void ValidateWidgetNames();
 
+	/**
+	 * Checks if the animations' bindings are valid. Shows a compiler warning if any of the animations have a null track, urging the user
+	 * to delete said track to avoid performance hitches when playing said null tracks in large user widgets.
+	*/
+	void ValidateWidgetAnimations();
+
 	// FKismetCompilerContext
 	virtual UEdGraphSchema_K2* CreateSchema() override;
 	virtual void CreateFunctionList() override;
diff --git a/Engine/Source/Runtime/UMG/Private/Animation/WidgetAnimationBinding.cpp b/Engine/Source/Runtime/UMG/Private/Animation/WidgetAnimationBinding.cpp
--- a/Engine/Source/Runtime/UMG/Private/Animation/WidgetAnimationBinding.cpp
+++ b/Engine/Source/Runtime/UMG/Private/Animation/WidgetAnimationBinding.cpp
@@ -1,33 +1,33 @@
-// Copyright Epic Games, Inc. All Rights Reserved.
-
-#include "Animation/WidgetAnimationBinding.h"
-#include "UObject/Object.h"
-#include "Components/Widget.h"
-#include "Blueprint/WidgetTree.h"
-
-
-/* FWidgetAnimationBinding interface
- *****************************************************************************/
-
-UObject* FWidgetAnimationBinding::FindRuntimeObject(UWidgetTree& WidgetTree, UUserWidget& UserWidget ) const
-{	
-	if (bIsRootWidget)
-	{
-		return &UserWidget;
-	}
-
-	UObject* FoundObject = FindObject<UObject>(&WidgetTree, *WidgetName.ToString());
-
-	if (FoundObject && (SlotWidgetName != NAME_None))
-	{
-		// if we were animating the slot, look up the slot that contains the widget 
-		UWidget* WidgetObject = Cast<UWidget>(FoundObject);
-
-		if (WidgetObject->Slot)
-		{
-			FoundObject = WidgetObject->Slot;
-		}
-	}
-
-	return FoundObject;
-}
+// Copyright Epic Games, Inc. All Rights Reserved.
+
+#include "Animation/WidgetAnimationBinding.h"
+#include "UObject/Object.h"
+#include "Components/Widget.h"
+#include "Blueprint/WidgetTree.h"
+
+
+/* FWidgetAnimationBinding interface
+ *****************************************************************************/
+
+UObject* FWidgetAnimationBinding::FindRuntimeObject(const UWidgetTree& WidgetTree, UUserWidget& UserWidget ) const
+{	
+	if (bIsRootWidget)
+	{
+		return &UserWidget;
+	}
+
+	UObject* FoundObject = WidgetTree.FindWidget(*WidgetName.ToString());
+
+	if (FoundObject && (SlotWidgetName != NAME_None))
+	{
+		// if we were animating the slot, look up the slot that contains the widget 
+		UWidget* WidgetObject = CastChecked<UWidget>(FoundObject);
+
+		if (WidgetObject->Slot)
+		{
+			FoundObject = WidgetObject->Slot;
+		}
+	}
+
+	return FoundObject;
+}
\ No newline at end of file
diff --git a/Engine/Source/Runtime/UMG/Public/Animation/WidgetAnimationBinding.h b/Engine/Source/Runtime/UMG/Public/Animation/WidgetAnimationBinding.h
--- a/Engine/Source/Runtime/UMG/Public/Animation/WidgetAnimationBinding.h
+++ b/Engine/Source/Runtime/UMG/Public/Animation/WidgetAnimationBinding.h
@@ -36,7 +36,7 @@
 	 * Locates a runtime object to animate from the provided tree of widgets.
 	 * @return the runtime object to animate or null if not found
 	 */
-	UMG_API UObject* FindRuntimeObject(UWidgetTree& WidgetTree, UUserWidget& UserWidget) const;
+	UMG_API UObject* FindRuntimeObject(const UWidgetTree& WidgetTree, UUserWidget& UserWidget) const;
 
 	/**
 	 * Compares two widget animation bindings for equality.
