UE-148560 - Unable to bind Float variables to certain properties in Widget elements

Since BP floats default to double precision, we need to relax binding requirements. Most UMG widgets uses delegates that expect single-precision floats, so we need to allow float<->double conversion in the UFloatBinding layer. Without this change, we can't bind any new float variables in UMG, which makes the feature fairly useless.

The fix simply updates UFloatBinding::IsSupportedSource to allow double properties. Additionally, when we need to retrieve the underlying value, we need to inspect the property first. If it's a double, we'll perform a cast. Otherwise, it defaults to existing behavior.

diff --git a/Engine/Source/Runtime/UMG/Private/Binding/FloatBinding.cpp b/Engine/Source/Runtime/UMG/Private/Binding/FloatBinding.cpp
--- a/Engine/Source/Runtime/UMG/Private/Binding/FloatBinding.cpp
+++ b/Engine/Source/Runtime/UMG/Private/Binding/FloatBinding.cpp
@@ -15,7 +15,7 @@ bool UFloatBinding::IsSupportedDestination(FProperty* Property) const
 
 bool UFloatBinding::IsSupportedSource(FProperty* Property) const
 {
-	return IsConcreteTypeCompatibleWithReflectedType<float>(Property);
+	return IsConcreteTypeCompatibleWithReflectedType<float>(Property) || IsConcreteTypeCompatibleWithReflectedType<double>(Property);
 }
 
 float UFloatBinding::GetValue() const
@@ -24,14 +24,40 @@ float UFloatBinding::GetValue() const
 
 	if ( UObject* Source = SourceObject.Get() )
 	{
-		float Value = 0;
-		if ( SourcePath.GetValue<float>(Source, Value) )
+		// Since we can bind to either a float or double, we need to perform a narrowing conversion where necessary.
+		// If this isn't a property, then we're assuming that a function is used to extract the float value.
+
+		float FloatValue = 0.0f;
+
+		SourcePath.Resolve(Source);
+		if (FProperty* Property = SourcePath.GetFProperty())
+		{
+			double DoubleValue = 0.0;
+			if (Property->IsA<FFloatProperty>() && SourcePath.GetValue<float>(Source, FloatValue))
+			{
+				return FloatValue;
+			}
+			else if (Property->IsA<FDoubleProperty>() && SourcePath.GetValue<double>(Source, DoubleValue))
+			{
+				FloatValue = static_cast<float>(DoubleValue);
+				return FloatValue;
+			}
+			else
+			{
+				checkf(false, TEXT("Unexpected property type: '%s'! Float bindings must use either a float or double property."), *Property->GetCPPType());
+			}
+		}
+		else
 		{
-			return Value;
+			check(SourcePath.GetCachedFunction());
+			if (SourcePath.GetValue<float>(Source, FloatValue))
+			{
+				return FloatValue;
+			}
 		}
 	}
 
-	return 0;
+	return 0.0f;
 }
 
 #undef LOCTEXT_NAMESPACE
