UMG - Further refining the named slot inheritance, we now extend the inheritance through multiple parent classes, and if a parent uses up a slot, then that slot becomes unavailable to subclasses to then be extended.  They cans also introduce new named slots along the way, so as it refines the structure, it can refine the available named slots.

diff --git a/Engine/Source/Editor/UMGEditor/Private/WidgetBlueprint.cpp b/Engine/Source/Editor/UMGEditor/Private/WidgetBlueprint.cpp
--- a/Engine/Source/Editor/UMGEditor/Private/WidgetBlueprint.cpp
+++ b/Engine/Source/Editor/UMGEditor/Private/WidgetBlueprint.cpp
@@ -751,7 +751,7 @@ bool UWidgetBlueprint::DetectSlateWidgetLeaks(TArray<FText>& ValidationErrors)
 
 	// Update the widget tree directly to match the blueprint tree.  That way the preview can update
 	// without needing to do a full recompile.
-	TempUserWidget->DuplicateAndInitializeFromWidgetTree(WidgetTree, nullptr);
+	TempUserWidget->DuplicateAndInitializeFromWidgetTree(WidgetTree, TMap<FName, UWidget*>());
 
 	// We don't want this widget doing all the normal startup and acting like it's the real deal
 	// trying to do gameplay stuff, so make sure it's in design mode.
@@ -1332,15 +1332,12 @@ bool UWidgetBlueprint::ArePropertyBindingsAllowed() const
 
 TArray<FName> UWidgetBlueprint::GetInheritedNamedSlots() const
 {
-	TArray<FName> SlotNames;
-	if ( UWidget* SuperClassCDO = GeneratedClass->GetSuperClass()->GetDefaultObject<UWidget>() )
+	if (const UWidgetBlueprintGeneratedClass* GeneratedBPClass = Cast<UWidgetBlueprintGeneratedClass>(GeneratedClass->GetSuperClass()))
 	{
-		if ( const INamedSlotInterface* NamedSlotHost = Cast<INamedSlotInterface>(SuperClassCDO) )
-		{
-			NamedSlotHost->GetSlotNames(SlotNames);
-		}
+		return GeneratedBPClass->AvailableNamedSlots;
 	}
-	return SlotNames;
+	
+	return TArray<FName>();
 }
 
 #if WITH_EDITOR
diff --git a/Engine/Source/Editor/UMGEditor/Private/WidgetBlueprintCompiler.cpp b/Engine/Source/Editor/UMGEditor/Private/WidgetBlueprintCompiler.cpp
--- a/Engine/Source/Editor/UMGEditor/Private/WidgetBlueprintCompiler.cpp
+++ b/Engine/Source/Editor/UMGEditor/Private/WidgetBlueprintCompiler.cpp
@@ -904,13 +904,30 @@ void FWidgetBlueprintCompilerContext::FinishCompilingClass(UClass* Class)
 		}
 
 		// Add all the names of the named slot widgets to the slot names structure.
-		BPGClass->NamedSlots.Reset();
-		WidgetBP->ForEachSourceWidget([&] (UWidget* Widget) {
-			if ( Widget && Widget->IsA<UNamedSlot>() )
+		{
+			BPGClass->AllNamedSlots.Reset();
+			UWidgetBlueprint* WidgetBPIt = WidgetBP;
+			while (WidgetBPIt)
 			{
-				BPGClass->NamedSlots.Add(Widget->GetFName());
+				WidgetBPIt->ForEachSourceWidget([&] (UWidget* Widget) {
+					if ( Widget && Widget->IsA<UNamedSlot>() )
+					{
+						BPGClass->AllNamedSlots.Add(Widget->GetFName());
+					}
+				});
+				
+				WidgetBPIt = Cast<UWidgetBlueprint>(WidgetBPIt->ParentClass->ClassGeneratedBy);
 			}
-		});
+
+			BPGClass->AvailableNamedSlots = BPGClass->AllNamedSlots;
+
+			// Remove any named slots from the available slots that has content for it.
+			BPGClass->GetNamedSlotArchetypeContent([BPGClass](FName SlotName, UWidget* Content)
+			{
+				// If we find content for this slot, remove it from the available set.
+				BPGClass->AvailableNamedSlots.Remove(SlotName);
+			});
+		}
 
 		// Make sure that we don't have dueling widget hierarchies
 		if (UWidgetBlueprintGeneratedClass* SuperBPGClass = Cast<UWidgetBlueprintGeneratedClass>(BPGClass->GetSuperClass()))
diff --git a/Engine/Source/Editor/UMGEditor/Private/WidgetBlueprintEditor.cpp b/Engine/Source/Editor/UMGEditor/Private/WidgetBlueprintEditor.cpp
--- a/Engine/Source/Editor/UMGEditor/Private/WidgetBlueprintEditor.cpp
+++ b/Engine/Source/Editor/UMGEditor/Private/WidgetBlueprintEditor.cpp
@@ -1610,9 +1610,27 @@ void FWidgetBlueprintEditor::UpdatePreview(UBlueprint* InBlueprint, bool bInForc
 				}
 			}
 
+			TMap<FName, UWidget*> NamedSlotContentToMerge;
+			UWidgetBlueprint* WidgetBPIt = PreviewBlueprint;
+			while (WidgetBPIt)
+			{
+				TArray<FName> SlotNames;
+				WidgetBPIt->WidgetTree->GetSlotNames(SlotNames);
+
+				for(const FName SlotName : SlotNames)
+				{
+					if(UWidget* Content = WidgetBPIt->WidgetTree->GetContentForSlot(SlotName))
+					{
+						NamedSlotContentToMerge.Add(SlotName, Content);
+					}
+				}
+
+				WidgetBPIt = Cast<UWidgetBlueprint>(WidgetBPIt->GeneratedClass->GetSuperClass()->ClassGeneratedBy);
+			}
+
 			// Update the widget tree directly to match the blueprint tree.  That way the preview can update
 			// without needing to do a full recompile.
-			PreviewUserWidget->DuplicateAndInitializeFromWidgetTree(LatestWidgetTree, LatestWidgetTree != PreviewBlueprint->WidgetTree ? PreviewBlueprint->WidgetTree : nullptr);
+			PreviewUserWidget->DuplicateAndInitializeFromWidgetTree(LatestWidgetTree, NamedSlotContentToMerge);
 
 			// Establish the widget as being in design time before initializing (so that IsDesignTime is reliable within Initialize)
             // We have to call it to make sure that all the WidgetTree had the DesignerFlags set correctly
diff --git a/Engine/Source/Runtime/UMG/Private/UserWidget.cpp b/Engine/Source/Runtime/UMG/Private/UserWidget.cpp
--- a/Engine/Source/Runtime/UMG/Private/UserWidget.cpp
+++ b/Engine/Source/Runtime/UMG/Private/UserWidget.cpp
@@ -182,7 +182,7 @@ void UUserWidget::InitializeNamedSlots()
 	}
 }
 
-void UUserWidget::DuplicateAndInitializeFromWidgetTree(UWidgetTree* InWidgetTree, INamedSlotInterface* NamedSlotsToMerge)
+void UUserWidget::DuplicateAndInitializeFromWidgetTree(UWidgetTree* InWidgetTree, const TMap<FName, UWidget*>& NamedSlotContentToMerge)
 {
 	TScopeCounter<uint32> ScopeInitializingFromWidgetTree(bInitializingFromWidgetTree);
 
@@ -211,34 +211,29 @@ void UUserWidget::DuplicateAndInitializeFromWidgetTree(UWidgetTree* InWidgetTree
 		});
 
 		// This block controls merging named slot content specified in a child class for the widget we're templated after.
-		if (NamedSlotsToMerge)
+		for (const auto& KVP_SlotContent : NamedSlotContentToMerge)
 		{
-			TArray<FName> SlotNames;
-			NamedSlotsToMerge->GetSlotNames(SlotNames);
-			for (FName SlotName : SlotNames)
+			// Don't insert the named slot content if the named slot is filled already.  This is a problematic
+			// scenario though, if someone inserted content, but we have class default instances, we sorta leave
+			// ourselves in a strange situation, because there are now potentially class variables that won't
+			// have an instance assigned.
+			if (!GetContentForSlot(KVP_SlotContent.Key))
 			{
-				// Don't insert the named slot content if the named slot is filled already.  This is a problematic
-				// scenario though, if someone inserted content, but we have class default instances, we sorta leave
-				// ourselves in a strange situation, because there are now potentially class variables that won't
-				// have an instance assigned.
-				if (!GetContentForSlot(SlotName))
+				if (UWidget* TemplateSlotContent = KVP_SlotContent.Value)
 				{
-					if (UWidget* TemplateSlotContent = NamedSlotsToMerge->GetContentForSlot(SlotName))
-					{
-						FObjectInstancingGraph NamedSlotInstancingGraph;
-						// We need to add a mapping from the template's widget tree to the new widget tree, that way
-						// as we instance the widget hierarchy it's grafted onto the new widget tree.
-						NamedSlotInstancingGraph.AddNewObject(WidgetTree, TemplateSlotContent->GetTypedOuter<UWidgetTree>());
-
-						// Instance the new widget from the foreign tree, but do it in a way that grafts it onto the tree we're instancing.
-						UWidget* Content = NewObject<UWidget>(WidgetTree, TemplateSlotContent->GetClass(), TemplateSlotContent->GetFName(), RF_Transactional, TemplateSlotContent, false, &NamedSlotInstancingGraph);
-						Content->SetFlags(RF_Transient | RF_DuplicateTransient);
-
-						// Insert the newly constructed widget into the named slot that corresponds.  The above creates
-						// it as if it was always part of the widget tree, but this actually puts it into a widget's
-						// slot for the named slot.
-						SetContentForSlot(SlotName, Content);
-					}
+					FObjectInstancingGraph NamedSlotInstancingGraph;
+					// We need to add a mapping from the template's widget tree to the new widget tree, that way
+					// as we instance the widget hierarchy it's grafted onto the new widget tree.
+					NamedSlotInstancingGraph.AddNewObject(WidgetTree, TemplateSlotContent->GetTypedOuter<UWidgetTree>());
+
+					// Instance the new widget from the foreign tree, but do it in a way that grafts it onto the tree we're instancing.
+					UWidget* Content = NewObject<UWidget>(WidgetTree, TemplateSlotContent->GetClass(), TemplateSlotContent->GetFName(), RF_Transactional, TemplateSlotContent, false, &NamedSlotInstancingGraph);
+					Content->SetFlags(RF_Transient | RF_DuplicateTransient);
+
+					// Insert the newly constructed widget into the named slot that corresponds.  The above creates
+					// it as if it was always part of the widget tree, but this actually puts it into a widget's
+					// slot for the named slot.
+					SetContentForSlot(KVP_SlotContent.Key, Content);
 				}
 			}
 		}
@@ -836,7 +831,7 @@ void UUserWidget::GetSlotNames(TArray<FName>& SlotNames) const
 	// Only do this if this widget is of a blueprint class
 	if (UWidgetBlueprintGeneratedClass* BGClass = GetWidgetTreeOwningClass())
 	{
-		SlotNames.Append(BGClass->NamedSlots);
+		SlotNames.Append(BGClass->AllNamedSlots);
 	}
 	else if (WidgetTree) // For non-blueprint widget blueprints we have to go through the widget tree to locate the named slots dynamically.
 	{
@@ -1941,7 +1936,7 @@ void UUserWidget::PreSave(FObjectPreSaveContext ObjectSaveContext)
 	// Remove bindings that are no longer contained in the class.
 	if ( UWidgetBlueprintGeneratedClass* BGClass = GetWidgetTreeOwningClass())
 	{
-		RemoveObsoleteBindings(BGClass->NamedSlots);
+		RemoveObsoleteBindings(BGClass->AllNamedSlots);
 	}
 
 	Super::PreSave(ObjectSaveContext);
diff --git a/Engine/Source/Runtime/UMG/Private/WidgetBlueprintGeneratedClass.cpp b/Engine/Source/Runtime/UMG/Private/WidgetBlueprintGeneratedClass.cpp
--- a/Engine/Source/Runtime/UMG/Private/WidgetBlueprintGeneratedClass.cpp
+++ b/Engine/Source/Runtime/UMG/Private/WidgetBlueprintGeneratedClass.cpp
@@ -204,16 +204,16 @@ void UWidgetBlueprintGeneratedClass::InitializeWidgetStatic(UUserWidget* UserWid
 	// similar to how we use to use the DesignerWidgetTree.
 	if ( CreatedWidgetTree == nullptr )
 	{
-		UWidgetTree* MergeNamedSlotTree = nullptr; 
-		if (UserWidget->GetClass() != InWidgetTreeWidgetClass)
+		TMap<FName, UWidget*> NamedSlotContentToMerge;
+		if (const UWidgetBlueprintGeneratedClass* WidgetsActualClass = Cast<UWidgetBlueprintGeneratedClass>(UserWidget->GetClass()))
 		{
-			if (UWidgetBlueprintGeneratedClass* WidgetsActualClass = Cast<UWidgetBlueprintGeneratedClass>(UserWidget->GetClass()))
+			WidgetsActualClass->GetNamedSlotArchetypeContent([&NamedSlotContentToMerge](FName SlotName, UWidget* Content)
 			{
-				MergeNamedSlotTree = WidgetsActualClass->GetWidgetTreeArchetype();
-			}
+				NamedSlotContentToMerge.Add(SlotName, Content);
+			});
 		}
 		
-		UserWidget->DuplicateAndInitializeFromWidgetTree(InWidgetTree, MergeNamedSlotTree);
+		UserWidget->DuplicateAndInitializeFromWidgetTree(InWidgetTree, NamedSlotContentToMerge);
 		CreatedWidgetTree = UserWidget->WidgetTree;
 	}
 
@@ -419,6 +419,26 @@ void UWidgetBlueprintGeneratedClass::SetWidgetTreeArchetype(UWidgetTree* InWidge
 	}
 }
 
+void UWidgetBlueprintGeneratedClass::GetNamedSlotArchetypeContent(TFunctionRef<void(FName /*SlotName*/, UWidget* /*Content*/)> Predicate) const
+{
+	for (const FName& SlotName : AllNamedSlots)
+	{
+		const UWidgetBlueprintGeneratedClass* BPGClass = this;
+		while (BPGClass)
+		{
+			UWidgetTree* Tree = BPGClass->GetWidgetTreeArchetype();
+			if (UWidget* ContentInSlot = Tree->GetContentForSlot(SlotName))
+			{
+				// Report the content in the slot, and break so we can test the next slot.
+				Predicate(SlotName, ContentInSlot);
+				break;
+			}
+
+			BPGClass = Cast<UWidgetBlueprintGeneratedClass>(BPGClass->GetSuperClass());
+		}
+	}
+}
+
 void UWidgetBlueprintGeneratedClass::Serialize(FArchive& Ar)
 {
 	Super::Serialize(Ar);
diff --git a/Engine/Source/Runtime/UMG/Public/Blueprint/UserWidget.h b/Engine/Source/Runtime/UMG/Public/Blueprint/UserWidget.h
--- a/Engine/Source/Runtime/UMG/Public/Blueprint/UserWidget.h
+++ b/Engine/Source/Runtime/UMG/Public/Blueprint/UserWidget.h
@@ -219,7 +219,7 @@ public:
 	virtual void PostLoad() override;
 	//~ End UObject Interface
 
-	void DuplicateAndInitializeFromWidgetTree(UWidgetTree* InWidgetTree, INamedSlotInterface* NamedSlotsToMerge);
+	void DuplicateAndInitializeFromWidgetTree(UWidgetTree* InWidgetTree, const TMap<FName, UWidget*>& NamedSlotContentToMerge);
 
 	virtual bool Initialize();
 
diff --git a/Engine/Source/Runtime/UMG/Public/Blueprint/WidgetBlueprintGeneratedClass.h b/Engine/Source/Runtime/UMG/Public/Blueprint/WidgetBlueprintGeneratedClass.h
--- a/Engine/Source/Runtime/UMG/Public/Blueprint/WidgetBlueprintGeneratedClass.h
+++ b/Engine/Source/Runtime/UMG/Public/Blueprint/WidgetBlueprintGeneratedClass.h
@@ -10,6 +10,7 @@
 
 #include "WidgetBlueprintGeneratedClass.generated.h"
 
+class UWidget;
 class UUserWidget;
 class UWidgetAnimation;
 class UWidgetBlueprintGeneratedClassExtension;
@@ -69,7 +70,7 @@ private:
 	/** A tree of the widget templates to be created */
 	UPROPERTY()
 	TObjectPtr<UWidgetTree> WidgetTree;
-
+	
 	/** The extension that are considered static to the class */
 	UPROPERTY()
 	TArray<TObjectPtr<UWidgetBlueprintGeneratedClassExtension>> Extensions;
@@ -97,13 +98,26 @@ public:
 	UPROPERTY()
 	TArray< TObjectPtr<UWidgetAnimation> > Animations;
 
+	/**
+	 * Exposed Named Slots for content.  These are slots that are accumulated from all super classes on compile.
+	 * and will exclude any named slots that are filled by a parent class.
+	 **/
+	UPROPERTY()
+	TArray<FName> AvailableNamedSlots;
+
+	/**
+	 * All named slots, even the ones that have content already filled into them by a parent class and are not
+	 * available for extension.
+	 **/
 	UPROPERTY()
-	TArray< FName > NamedSlots;
+	TArray<FName> AllNamedSlots;
 
 public:
 	UWidgetTree* GetWidgetTreeArchetype() const { return WidgetTree; }
 	void SetWidgetTreeArchetype(UWidgetTree* InWidgetTree);
 
+	void GetNamedSlotArchetypeContent(TFunctionRef<void(FName /*SlotName*/, UWidget* /*Content*/)> Predicate) const;
+
 	// Walks up the hierarchy looking for a valid widget tree.
 	UWidgetBlueprintGeneratedClass* FindWidgetTreeOwningClass() const;
 
