UMG: Initialize Extension that are added while initializing other Extension.

diff --git a/Engine/Source/Runtime/UMG/Private/UserWidget.cpp b/Engine/Source/Runtime/UMG/Private/UserWidget.cpp
--- a/Engine/Source/Runtime/UMG/Private/UserWidget.cpp
+++ b/Engine/Source/Runtime/UMG/Private/UserWidget.cpp
@@ -1375,8 +1375,11 @@ void UUserWidget::NativeOnInitialized()
 			});
 	}
 
-	for (UUserWidgetExtension* Extension : Extensions)
+	// Extension can add other extensions. Use index loop to initialize them all.
+	for (int32 Index = 0; Index < Extensions.Num(); ++Index)
 	{
+		UUserWidgetExtension* Extension = Extensions[Index];
+		check(Extension);
 		Extension->Initialize();
 	}
 
@@ -1416,9 +1419,17 @@ void UUserWidget::NativeDestruct()
 	StopListeningForAllInputActions();
 	Destruct();
 
-	for (UUserWidgetExtension* Extension : Extensions)
+	if (bInitialized)
 	{
-		Extension->Destruct();
+		// Extension can remove other extensions.
+		while (Extensions.Num() > 0)
+		{
+			int32 Index = Extensions.Num() - 1;
+			UUserWidgetExtension* Extension = Extensions[Index];
+			check(Extension);
+			Extensions.RemoveAtSwap(Index);
+			Extension->Destruct();
+		}
 	}
 
 	if (UWidgetBlueprintGeneratedClass* BPClass = Cast<UWidgetBlueprintGeneratedClass>(GetClass()))
@@ -1437,9 +1448,11 @@ void UUserWidget::NativeTick(const FGeometry& MyGeometry, float InDeltaTime)
 	{
 		GInitRunaway();
 
-		for (UUserWidgetExtension* Extension : Extensions)
+		// Extension can be added while ticking another extension.
+		//This loop does guarantee that they will all be updated this frame, if it's the case,  but it will not crash.
+		for (int32 Index = 0; Index < Extensions.Num(); ++Index)
 		{
-			Extension->Tick(MyGeometry, InDeltaTime);
+			Extensions[Index]->Tick(MyGeometry, InDeltaTime);
 		}
 
 #if WITH_EDITOR
@@ -2152,7 +2165,10 @@ void UUserWidget::RemoveExtension(UUserWidgetExtension* InExtension)
 	{
 		if (Extensions.RemoveSingleSwap(InExtension))
 		{
-			InExtension->Destruct();
+			if (bInitialized)
+			{
+				InExtension->Destruct();
+			}
 		}
 	}
 }
