UMG: Update HorizontalBoxSlot, SizeBoxSlot, UniformGridSlot, VerticalBoxSlot, WidgetSwitcherSlot, WindowTitleBarAreaSlot, WrapBoxSlot with accessors.

diff --git a/Engine/Plugins/Experimental/CommonUI/Source/CommonUI/Private/CommonLoadGuard.cpp b/Engine/Plugins/Experimental/CommonUI/Source/CommonUI/Private/CommonLoadGuard.cpp
--- a/Engine/Plugins/Experimental/CommonUI/Source/CommonUI/Private/CommonLoadGuard.cpp
+++ b/Engine/Plugins/Experimental/CommonUI/Source/CommonUI/Private/CommonLoadGuard.cpp
@@ -339,9 +339,9 @@ void UCommonLoadGuard::PostLoad()
 			// Copy properties from the old SizeBoxSlot
 			if (USizeBoxSlot* AsSizeBoxSlot = Cast<USizeBoxSlot>(PanelSlot))
 			{
-				LoadGuardSlot->SetPadding(AsSizeBoxSlot->Padding);
-				LoadGuardSlot->SetHorizontalAlignment(AsSizeBoxSlot->HorizontalAlignment);
-				LoadGuardSlot->SetVerticalAlignment(AsSizeBoxSlot->VerticalAlignment);
+				LoadGuardSlot->SetPadding(AsSizeBoxSlot->GetPadding());
+				LoadGuardSlot->SetHorizontalAlignment(AsSizeBoxSlot->GetHorizontalAlignment());
+				LoadGuardSlot->SetVerticalAlignment(AsSizeBoxSlot->GetVerticalAlignment());
 			}
 			Slots[0] = LoadGuardSlot;
 		}
diff --git a/Engine/Source/Editor/UMGEditor/Private/Extensions/UniformGridSlotExtension.cpp b/Engine/Source/Editor/UMGEditor/Private/Extensions/UniformGridSlotExtension.cpp
--- a/Engine/Source/Editor/UMGEditor/Private/Extensions/UniformGridSlotExtension.cpp
+++ b/Engine/Source/Editor/UMGEditor/Private/Extensions/UniformGridSlotExtension.cpp
@@ -106,13 +106,13 @@ FReply FUniformGridSlotExtension::HandleShiftColumn(int32 ShiftAmount)
 void FUniformGridSlotExtension::ShiftRow(UWidget* Widget, int32 ShiftAmount)
 {
 	UUniformGridSlot* Slot = Cast<UUniformGridSlot>(Widget->Slot);
-	Slot->SetRow(FMath::Max(Slot->Row + ShiftAmount, 0));
+	Slot->SetRow(FMath::Max(Slot->GetRow() + ShiftAmount, 0));
 }
 
 void FUniformGridSlotExtension::ShiftColumn(UWidget* Widget, int32 ShiftAmount)
 {
 	UUniformGridSlot* Slot = Cast<UUniformGridSlot>(Widget->Slot);
-	Slot->SetColumn(FMath::Max(Slot->Column + ShiftAmount, 0));
+	Slot->SetColumn(FMath::Max(Slot->GetColumn() + ShiftAmount, 0));
 }
 
 #undef LOCTEXT_NAMESPACE
diff --git a/Engine/Source/Runtime/UMG/Private/Components/BackgroundBlur.cpp b/Engine/Source/Runtime/UMG/Private/Components/BackgroundBlur.cpp
--- a/Engine/Source/Runtime/UMG/Private/Components/BackgroundBlur.cpp
+++ b/Engine/Source/Runtime/UMG/Private/Components/BackgroundBlur.cpp
@@ -73,9 +73,9 @@ void UBackgroundBlur::SynchronizeProperties()
 void UBackgroundBlur::OnSlotAdded(UPanelSlot* InSlot)
 {
 	UBackgroundBlurSlot* BackgroundBlurSlot = CastChecked<UBackgroundBlurSlot>(InSlot);
-	BackgroundBlurSlot->Padding = Padding;
-	BackgroundBlurSlot->HorizontalAlignment = HorizontalAlignment;
-	BackgroundBlurSlot->VerticalAlignment = VerticalAlignment;
+	BackgroundBlurSlot->SetPadding(Padding);
+	BackgroundBlurSlot->SetHorizontalAlignment(HorizontalAlignment);
+	BackgroundBlurSlot->SetVerticalAlignment(VerticalAlignment);
 
 	// Add the child to the live slot if it already exists
 	if (MyBackgroundBlur.IsValid())
diff --git a/Engine/Source/Runtime/UMG/Private/Components/BackgroundBlurSlot.cpp b/Engine/Source/Runtime/UMG/Private/Components/BackgroundBlurSlot.cpp
--- a/Engine/Source/Runtime/UMG/Private/Components/BackgroundBlurSlot.cpp
+++ b/Engine/Source/Runtime/UMG/Private/Components/BackgroundBlurSlot.cpp
@@ -11,10 +11,12 @@
 UBackgroundBlurSlot::UBackgroundBlurSlot(const FObjectInitializer& ObjectInitializer)
 	: Super(ObjectInitializer)
 {
-	Padding = FMargin(4, 2);
+PRAGMA_DISABLE_DEPRECATION_WARNINGS
+	Padding = FMargin(4.f, 2.f);
 
 	HorizontalAlignment = HAlign_Fill;
 	VerticalAlignment = VAlign_Fill;
+PRAGMA_ENABLE_DEPRECATION_WARNINGS
 }
 
 void UBackgroundBlurSlot::ReleaseSlateResources(bool bReleaseChildren)
@@ -28,36 +30,60 @@ void UBackgroundBlurSlot::BuildSlot(TSharedRef<SBackgroundBlur> InBackgroundBlur
 {
 	BackgroundBlur = InBackgroundBlur;
 
+PRAGMA_DISABLE_DEPRECATION_WARNINGS
 	BackgroundBlur->SetPadding(Padding);
 	BackgroundBlur->SetHAlign(HorizontalAlignment);
 	BackgroundBlur->SetVAlign(VerticalAlignment);
+PRAGMA_ENABLE_DEPRECATION_WARNINGS
 
 	BackgroundBlur->SetContent(Content ? Content->TakeWidget() : SNullWidget::NullWidget);
 }
 
+PRAGMA_DISABLE_DEPRECATION_WARNINGS
+FMargin UBackgroundBlurSlot::GetPadding() const
+{
+	return Padding;
+}
+
 void UBackgroundBlurSlot::SetPadding(FMargin InPadding)
 {
+	Padding = InPadding;
 	CastChecked<UBackgroundBlur>(Parent)->SetPadding(InPadding);
 }
 
+EHorizontalAlignment UBackgroundBlurSlot::GetHorizontalAlignment() const
+{
+	return HorizontalAlignment;
+}
+
 void UBackgroundBlurSlot::SetHorizontalAlignment(EHorizontalAlignment InHorizontalAlignment)
 {
+	HorizontalAlignment = InHorizontalAlignment;
 	CastChecked<UBackgroundBlur>(Parent)->SetHorizontalAlignment(InHorizontalAlignment);
 }
 
+EVerticalAlignment UBackgroundBlurSlot::GetVerticalAlignment() const
+{
+	return VerticalAlignment;
+}
+
 void UBackgroundBlurSlot::SetVerticalAlignment(EVerticalAlignment InVerticalAlignment)
 {
+	VerticalAlignment = InVerticalAlignment;
 	CastChecked<UBackgroundBlur>(Parent)->SetVerticalAlignment(InVerticalAlignment);
 }
+PRAGMA_ENABLE_DEPRECATION_WARNINGS
 
 void UBackgroundBlurSlot::SynchronizeProperties()
 {
+PRAGMA_DISABLE_DEPRECATION_WARNINGS
 	if ( BackgroundBlur.IsValid() )
 	{
 		SetPadding(Padding);
 		SetHorizontalAlignment(HorizontalAlignment);
 		SetVerticalAlignment(VerticalAlignment);
 	}
+PRAGMA_ENABLE_DEPRECATION_WARNINGS
 }
 
 #if WITH_EDITOR
diff --git a/Engine/Source/Runtime/UMG/Private/Components/BorderSlot.cpp b/Engine/Source/Runtime/UMG/Private/Components/BorderSlot.cpp
--- a/Engine/Source/Runtime/UMG/Private/Components/BorderSlot.cpp
+++ b/Engine/Source/Runtime/UMG/Private/Components/BorderSlot.cpp
@@ -12,10 +12,12 @@
 UBorderSlot::UBorderSlot(const FObjectInitializer& ObjectInitializer)
 	: Super(ObjectInitializer)
 {
-	Padding = FMargin(4, 2);
+PRAGMA_DISABLE_DEPRECATION_WARNINGS
+	Padding = FMargin(4.f, 2.f);
 
 	HorizontalAlignment = HAlign_Fill;
 	VerticalAlignment = VAlign_Fill;
+PRAGMA_ENABLE_DEPRECATION_WARNINGS
 }
 
 void UBorderSlot::ReleaseSlateResources(bool bReleaseChildren)
@@ -29,36 +31,60 @@ void UBorderSlot::BuildSlot(TSharedRef<SBorder> InBorder)
 {
 	Border = InBorder;
 
-	Border.Pin()->SetPadding(Padding);
-	Border.Pin()->SetHAlign(HorizontalAlignment);
-	Border.Pin()->SetVAlign(VerticalAlignment);
+PRAGMA_DISABLE_DEPRECATION_WARNINGS
+	InBorder->SetPadding(Padding);
+	InBorder->SetHAlign(HorizontalAlignment);
+	InBorder->SetVAlign(VerticalAlignment);
+PRAGMA_ENABLE_DEPRECATION_WARNINGS
 
-	Border.Pin()->SetContent(Content ? Content->TakeWidget() : SNullWidget::NullWidget);
+	InBorder->SetContent(Content ? Content->TakeWidget() : SNullWidget::NullWidget);
+}
+
+PRAGMA_DISABLE_DEPRECATION_WARNINGS
+FMargin UBorderSlot::GetPadding() const
+{
+	return Padding;
 }
 
 void UBorderSlot::SetPadding(FMargin InPadding)
 {
+	Padding = InPadding;
 	CastChecked<UBorder>(Parent)->SetPadding(InPadding);
 }
 
+EHorizontalAlignment UBorderSlot::GetHorizontalAlignment() const
+{
+	return HorizontalAlignment;
+}
+
 void UBorderSlot::SetHorizontalAlignment(EHorizontalAlignment InHorizontalAlignment)
 {
+	HorizontalAlignment = InHorizontalAlignment;
 	CastChecked<UBorder>(Parent)->SetHorizontalAlignment(InHorizontalAlignment);
 }
 
+EVerticalAlignment UBorderSlot::GetVerticalAlignment() const
+{
+	return VerticalAlignment;
+}
+
 void UBorderSlot::SetVerticalAlignment(EVerticalAlignment InVerticalAlignment)
 {
+	VerticalAlignment = InVerticalAlignment;
 	CastChecked<UBorder>(Parent)->SetVerticalAlignment(InVerticalAlignment);
 }
+PRAGMA_ENABLE_DEPRECATION_WARNINGS
 
 void UBorderSlot::SynchronizeProperties()
 {
+PRAGMA_DISABLE_DEPRECATION_WARNINGS
 	if ( Border.IsValid() )
 	{
 		SetPadding(Padding);
 		SetHorizontalAlignment(HorizontalAlignment);
 		SetVerticalAlignment(VerticalAlignment);
 	}
+PRAGMA_ENABLE_DEPRECATION_WARNINGS
 }
 
 #if WITH_EDITOR
diff --git a/Engine/Source/Runtime/UMG/Private/Components/HorizontalBoxSlot.cpp b/Engine/Source/Runtime/UMG/Private/Components/HorizontalBoxSlot.cpp
--- a/Engine/Source/Runtime/UMG/Private/Components/HorizontalBoxSlot.cpp
+++ b/Engine/Source/Runtime/UMG/Private/Components/HorizontalBoxSlot.cpp
@@ -9,11 +9,13 @@
 
 UHorizontalBoxSlot::UHorizontalBoxSlot(const FObjectInitializer& ObjectInitializer)
 	: Super(ObjectInitializer)
+	, Slot(nullptr)
 {
+PRAGMA_DISABLE_DEPRECATION_WARNINGS
 	HorizontalAlignment = HAlign_Fill;
 	VerticalAlignment = VAlign_Fill;
-	Slot = nullptr;
 	Size = FSlateChildSize(ESlateSizeRule::Automatic);
+PRAGMA_ENABLE_DEPRECATION_WARNINGS
 }
 
 void UHorizontalBoxSlot::ReleaseSlateResources(bool bReleaseChildren)
@@ -24,7 +26,8 @@ void UHorizontalBoxSlot::ReleaseSlateResources(bool bReleaseChildren)
 }
 
 void UHorizontalBoxSlot::BuildSlot(TSharedRef<SHorizontalBox> HorizontalBox)
-{ 
+{
+PRAGMA_DISABLE_DEPRECATION_WARNINGS
 	HorizontalBox->AddSlot()
 		.Expose(Slot)
 		.HAlign(HorizontalAlignment)
@@ -34,6 +37,13 @@ void UHorizontalBoxSlot::BuildSlot(TSharedRef<SHorizontalBox> HorizontalBox)
 		[
 			Content == NULL ? SNullWidget::NullWidget : Content->TakeWidget()
 		];
+PRAGMA_ENABLE_DEPRECATION_WARNINGS
+}
+
+PRAGMA_DISABLE_DEPRECATION_WARNINGS
+FMargin UHorizontalBoxSlot::GetPadding() const
+{
+	return Padding;
 }
 
 void UHorizontalBoxSlot::SetPadding(FMargin InPadding)
@@ -45,6 +55,11 @@ void UHorizontalBoxSlot::SetPadding(FMargin InPadding)
 	}
 }
 
+FSlateChildSize UHorizontalBoxSlot::GetSize() const
+{
+	return Size;
+}
+
 void UHorizontalBoxSlot::SetSize(FSlateChildSize InSize)
 {
 	Size = InSize;
@@ -54,6 +69,11 @@ void UHorizontalBoxSlot::SetSize(FSlateChildSize InSize)
 	}
 }
 
+EHorizontalAlignment UHorizontalBoxSlot::GetHorizontalAlignment() const
+{
+	return HorizontalAlignment;
+}
+
 void UHorizontalBoxSlot::SetHorizontalAlignment(EHorizontalAlignment InHorizontalAlignment)
 {
 	HorizontalAlignment = InHorizontalAlignment;
@@ -63,6 +83,11 @@ void UHorizontalBoxSlot::SetHorizontalAlignment(EHorizontalAlignment InHorizonta
 	}
 }
 
+EVerticalAlignment UHorizontalBoxSlot::GetVerticalAlignment() const
+{
+	return VerticalAlignment;
+}
+
 void UHorizontalBoxSlot::SetVerticalAlignment(EVerticalAlignment InVerticalAlignment)
 {
 	VerticalAlignment = InVerticalAlignment;
@@ -71,13 +96,16 @@ void UHorizontalBoxSlot::SetVerticalAlignment(EVerticalAlignment InVerticalAlign
 		Slot->SetVerticalAlignment(InVerticalAlignment);
 	}
 }
+PRAGMA_ENABLE_DEPRECATION_WARNINGS
 
 void UHorizontalBoxSlot::SynchronizeProperties()
 {
+PRAGMA_DISABLE_DEPRECATION_WARNINGS
 	SetPadding(Padding);
 	SetSize(Size);
 	SetHorizontalAlignment(HorizontalAlignment);
 	SetVerticalAlignment(VerticalAlignment);
+PRAGMA_ENABLE_DEPRECATION_WARNINGS
 }
 
 #if WITH_EDITOR
diff --git a/Engine/Source/Runtime/UMG/Private/Components/SizeBoxSlot.cpp b/Engine/Source/Runtime/UMG/Private/Components/SizeBoxSlot.cpp
--- a/Engine/Source/Runtime/UMG/Private/Components/SizeBoxSlot.cpp
+++ b/Engine/Source/Runtime/UMG/Private/Components/SizeBoxSlot.cpp
@@ -11,10 +11,12 @@
 USizeBoxSlot::USizeBoxSlot(const FObjectInitializer& ObjectInitializer)
 	: Super(ObjectInitializer)
 {
-	Padding = FMargin(0, 0);
+PRAGMA_DISABLE_DEPRECATION_WARNINGS
+	Padding = FMargin(0.f, 0.f);
 
 	HorizontalAlignment = HAlign_Fill;
 	VerticalAlignment = VAlign_Fill;
+PRAGMA_ENABLE_DEPRECATION_WARNINGS
 }
 
 void USizeBoxSlot::ReleaseSlateResources(bool bReleaseChildren)
@@ -33,6 +35,12 @@ void USizeBoxSlot::BuildSlot(TSharedRef<SBox> InSizeBox)
 	SizeBox.Pin()->SetContent(Content ? Content->TakeWidget() : SNullWidget::NullWidget);
 }
 
+PRAGMA_DISABLE_DEPRECATION_WARNINGS
+FMargin USizeBoxSlot::GetPadding() const
+{
+	return Padding;
+}
+
 void USizeBoxSlot::SetPadding(FMargin InPadding)
 {
 	Padding = InPadding;
@@ -42,6 +50,11 @@ void USizeBoxSlot::SetPadding(FMargin InPadding)
 	}
 }
 
+EHorizontalAlignment USizeBoxSlot::GetHorizontalAlignment() const
+{
+	return HorizontalAlignment;
+}
+
 void USizeBoxSlot::SetHorizontalAlignment(EHorizontalAlignment InHorizontalAlignment)
 {
 	HorizontalAlignment = InHorizontalAlignment;
@@ -51,6 +64,11 @@ void USizeBoxSlot::SetHorizontalAlignment(EHorizontalAlignment InHorizontalAlign
 	}
 }
 
+EVerticalAlignment USizeBoxSlot::GetVerticalAlignment() const
+{
+	return VerticalAlignment;
+}
+
 void USizeBoxSlot::SetVerticalAlignment(EVerticalAlignment InVerticalAlignment)
 {
 	VerticalAlignment = InVerticalAlignment;
@@ -59,10 +77,13 @@ void USizeBoxSlot::SetVerticalAlignment(EVerticalAlignment InVerticalAlignment)
 		SizeBox.Pin()->SetVAlign(InVerticalAlignment);
 	}
 }
+PRAGMA_ENABLE_DEPRECATION_WARNINGS
 
 void USizeBoxSlot::SynchronizeProperties()
 {
+PRAGMA_DISABLE_DEPRECATION_WARNINGS
 	SetPadding(Padding);
 	SetHorizontalAlignment(HorizontalAlignment);
 	SetVerticalAlignment(VerticalAlignment);
+PRAGMA_ENABLE_DEPRECATION_WARNINGS
 }
diff --git a/Engine/Source/Runtime/UMG/Private/Components/UniformGridSlot.cpp b/Engine/Source/Runtime/UMG/Private/Components/UniformGridSlot.cpp
--- a/Engine/Source/Runtime/UMG/Private/Components/UniformGridSlot.cpp
+++ b/Engine/Source/Runtime/UMG/Private/Components/UniformGridSlot.cpp
@@ -10,8 +10,10 @@ UUniformGridSlot::UUniformGridSlot(const FObjectInitializer& ObjectInitializer)
 	: Super(ObjectInitializer)
 	, Slot(nullptr)
 {
+PRAGMA_DISABLE_DEPRECATION_WARNINGS
 	HorizontalAlignment = HAlign_Left;
 	VerticalAlignment = VAlign_Top;
+PRAGMA_ENABLE_DEPRECATION_WARNINGS
 }
 
 void UUniformGridSlot::ReleaseSlateResources(bool bReleaseChildren)
@@ -23,6 +25,7 @@ void UUniformGridSlot::ReleaseSlateResources(bool bReleaseChildren)
 
 void UUniformGridSlot::BuildSlot(TSharedRef<SUniformGridPanel> GridPanel)
 {
+PRAGMA_DISABLE_DEPRECATION_WARNINGS
 	GridPanel->AddSlot(Column, Row)
 		.Expose(Slot)
 		.HAlign(HorizontalAlignment)
@@ -30,6 +33,13 @@ void UUniformGridSlot::BuildSlot(TSharedRef<SUniformGridPanel> GridPanel)
 		[
 			Content == nullptr ? SNullWidget::NullWidget : Content->TakeWidget()
 		];
+PRAGMA_ENABLE_DEPRECATION_WARNINGS
+}
+
+PRAGMA_DISABLE_DEPRECATION_WARNINGS
+int32 UUniformGridSlot::GetRow() const
+{
+	return Row;
 }
 
 void UUniformGridSlot::SetRow(int32 InRow)
@@ -41,6 +51,11 @@ void UUniformGridSlot::SetRow(int32 InRow)
 	}
 }
 
+int32 UUniformGridSlot::GetColumn() const
+{
+	return Column;
+}
+
 void UUniformGridSlot::SetColumn(int32 InColumn)
 {
 	Column = InColumn;
@@ -50,6 +65,11 @@ void UUniformGridSlot::SetColumn(int32 InColumn)
 	}
 }
 
+EHorizontalAlignment UUniformGridSlot::GetHorizontalAlignment() const
+{
+	return HorizontalAlignment;
+}
+
 void UUniformGridSlot::SetHorizontalAlignment(EHorizontalAlignment InHorizontalAlignment)
 {
 	HorizontalAlignment = InHorizontalAlignment;
@@ -59,6 +79,11 @@ void UUniformGridSlot::SetHorizontalAlignment(EHorizontalAlignment InHorizontalA
 	}
 }
 
+EVerticalAlignment UUniformGridSlot::GetVerticalAlignment() const
+{
+	return VerticalAlignment;
+}
+
 void UUniformGridSlot::SetVerticalAlignment(EVerticalAlignment InVerticalAlignment)
 {
 	VerticalAlignment = InVerticalAlignment;
@@ -67,13 +92,16 @@ void UUniformGridSlot::SetVerticalAlignment(EVerticalAlignment InVerticalAlignme
 		Slot->SetVerticalAlignment(InVerticalAlignment);
 	}
 }
+PRAGMA_ENABLE_DEPRECATION_WARNINGS
 
 void UUniformGridSlot::SynchronizeProperties()
 {
+PRAGMA_DISABLE_DEPRECATION_WARNINGS
 	SetRow(Row);
 	SetColumn(Column);
 	SetHorizontalAlignment(HorizontalAlignment);
 	SetVerticalAlignment(VerticalAlignment);
+PRAGMA_ENABLE_DEPRECATION_WARNINGS
 }
 
 #if WITH_EDITOR
@@ -81,10 +109,10 @@ void UUniformGridSlot::SynchronizeProperties()
 bool UUniformGridSlot::NudgeByDesigner(const FVector2D& NudgeDirection, const TOptional<int32>& GridSnapSize)
 {
 	const FVector2D ClampedDirection = NudgeDirection.ClampAxes(-1.0f, 1.0f);
-	const int32 NewColumn = Column + ClampedDirection.X;
-	const int32 NewRow = Row + ClampedDirection.Y;
+	const int32 NewColumn = GetColumn() + ClampedDirection.X;
+	const int32 NewRow = GetRow() + ClampedDirection.Y;
 
-	if (NewColumn < 0 || NewRow < 0 || (NewColumn == Column && NewRow == Row))
+	if (NewColumn < 0 || NewRow < 0 || (NewColumn == GetColumn() && NewRow == GetRow()))
 	{
 		return false;
 	}
@@ -100,8 +128,8 @@ bool UUniformGridSlot::NudgeByDesigner(const FVector2D& NudgeDirection, const TO
 void UUniformGridSlot::SynchronizeFromTemplate(const UPanelSlot* const TemplateSlot)
 {
 	const ThisClass* const TemplateUniformGridSlot = CastChecked<ThisClass>(TemplateSlot);
-	SetRow(TemplateUniformGridSlot->Row);
-	SetColumn(TemplateUniformGridSlot->Column);
+	SetRow(TemplateUniformGridSlot->GetRow());
+	SetColumn(TemplateUniformGridSlot->GetColumn());
 }
 
 #endif
diff --git a/Engine/Source/Runtime/UMG/Private/Components/VerticalBoxSlot.cpp b/Engine/Source/Runtime/UMG/Private/Components/VerticalBoxSlot.cpp
--- a/Engine/Source/Runtime/UMG/Private/Components/VerticalBoxSlot.cpp
+++ b/Engine/Source/Runtime/UMG/Private/Components/VerticalBoxSlot.cpp
@@ -11,9 +11,11 @@ UVerticalBoxSlot::UVerticalBoxSlot(const FObjectInitializer& ObjectInitializer)
 	: Super(ObjectInitializer)
 	, Slot(nullptr)
 {
+PRAGMA_DISABLE_DEPRECATION_WARNINGS
 	HorizontalAlignment = HAlign_Fill;
 	VerticalAlignment = VAlign_Fill;
 	Size = FSlateChildSize(ESlateSizeRule::Automatic);
+PRAGMA_ENABLE_DEPRECATION_WARNINGS
 }
 
 void UVerticalBoxSlot::ReleaseSlateResources(bool bReleaseChildren)
@@ -25,16 +27,23 @@ void UVerticalBoxSlot::ReleaseSlateResources(bool bReleaseChildren)
 
 void UVerticalBoxSlot::BuildSlot(TSharedRef<SVerticalBox> VerticalBox)
 {
+PRAGMA_DISABLE_DEPRECATION_WARNINGS
 	VerticalBox->AddSlot()
 		.Expose(Slot)
 		.Padding(Padding)
 		.HAlign(HorizontalAlignment)
 		.VAlign(VerticalAlignment)
 		.SizeParam(UWidget::ConvertSerializedSizeParamToRuntime(Size))
-		.Expose(Slot)
 		[
 			Content == nullptr ? SNullWidget::NullWidget : Content->TakeWidget()
 		];
+PRAGMA_ENABLE_DEPRECATION_WARNINGS
+}
+
+PRAGMA_DISABLE_DEPRECATION_WARNINGS
+FMargin UVerticalBoxSlot::GetPadding() const
+{
+	return Padding;
 }
 
 void UVerticalBoxSlot::SetPadding(FMargin InPadding)
@@ -46,6 +55,11 @@ void UVerticalBoxSlot::SetPadding(FMargin InPadding)
 	}
 }
 
+FSlateChildSize UVerticalBoxSlot::GetSize() const
+{
+	return Size;
+}
+
 void UVerticalBoxSlot::SetSize(FSlateChildSize InSize)
 {
 	Size = InSize;
@@ -55,6 +69,11 @@ void UVerticalBoxSlot::SetSize(FSlateChildSize InSize)
 	}
 }
 
+EHorizontalAlignment UVerticalBoxSlot::GetHorizontalAlignment() const
+{
+	return HorizontalAlignment;
+}
+
 void UVerticalBoxSlot::SetHorizontalAlignment(EHorizontalAlignment InHorizontalAlignment)
 {
 	HorizontalAlignment = InHorizontalAlignment;
@@ -64,6 +83,11 @@ void UVerticalBoxSlot::SetHorizontalAlignment(EHorizontalAlignment InHorizontalA
 	}
 }
 
+EVerticalAlignment UVerticalBoxSlot::GetVerticalAlignment() const
+{
+	return VerticalAlignment;
+}
+
 void UVerticalBoxSlot::SetVerticalAlignment(EVerticalAlignment InVerticalAlignment)
 {
 	VerticalAlignment = InVerticalAlignment;
@@ -72,13 +96,16 @@ void UVerticalBoxSlot::SetVerticalAlignment(EVerticalAlignment InVerticalAlignme
 		Slot->SetVerticalAlignment(InVerticalAlignment);
 	}
 }
+PRAGMA_ENABLE_DEPRECATION_WARNINGS
 
 void UVerticalBoxSlot::SynchronizeProperties()
 {
+PRAGMA_DISABLE_DEPRECATION_WARNINGS
 	SetPadding(Padding);
 	SetSize(Size);
 	SetHorizontalAlignment(HorizontalAlignment);
 	SetVerticalAlignment(VerticalAlignment);
+PRAGMA_ENABLE_DEPRECATION_WARNINGS
 }
 
 #if WITH_EDITOR
diff --git a/Engine/Source/Runtime/UMG/Private/Components/WidgetSwitcherSlot.cpp b/Engine/Source/Runtime/UMG/Private/Components/WidgetSwitcherSlot.cpp
--- a/Engine/Source/Runtime/UMG/Private/Components/WidgetSwitcherSlot.cpp
+++ b/Engine/Source/Runtime/UMG/Private/Components/WidgetSwitcherSlot.cpp
@@ -11,8 +11,10 @@ UWidgetSwitcherSlot::UWidgetSwitcherSlot(const FObjectInitializer& ObjectInitial
 	: Super(ObjectInitializer)
 	, Slot(nullptr)
 {
+PRAGMA_DISABLE_DEPRECATION_WARNINGS
 	HorizontalAlignment = HAlign_Fill;
 	VerticalAlignment = VAlign_Fill;
+PRAGMA_ENABLE_DEPRECATION_WARNINGS
 }
 
 void UWidgetSwitcherSlot::ReleaseSlateResources(bool bReleaseChildren)
@@ -24,6 +26,7 @@ void UWidgetSwitcherSlot::ReleaseSlateResources(bool bReleaseChildren)
 
 void UWidgetSwitcherSlot::BuildSlot(TSharedRef<SWidgetSwitcher> WidgetSwitcher)
 {
+PRAGMA_DISABLE_DEPRECATION_WARNINGS
 	WidgetSwitcher->AddSlot()
 		.Expose(Slot)
 		.Padding(Padding)
@@ -32,6 +35,7 @@ void UWidgetSwitcherSlot::BuildSlot(TSharedRef<SWidgetSwitcher> WidgetSwitcher)
 		[
 			Content == nullptr ? SNullWidget::NullWidget : Content->TakeWidget()
 		];
+PRAGMA_ENABLE_DEPRECATION_WARNINGS
 }
 
 void UWidgetSwitcherSlot::SetContent(UWidget* NewContent)
@@ -43,6 +47,12 @@ void UWidgetSwitcherSlot::SetContent(UWidget* NewContent)
 	}
 }
 
+PRAGMA_DISABLE_DEPRECATION_WARNINGS
+FMargin UWidgetSwitcherSlot::GetPadding() const
+{
+	return Padding;
+}
+
 void UWidgetSwitcherSlot::SetPadding(FMargin InPadding)
 {
 	Padding = InPadding;
@@ -52,6 +62,11 @@ void UWidgetSwitcherSlot::SetPadding(FMargin InPadding)
 	}
 }
 
+EHorizontalAlignment UWidgetSwitcherSlot::GetHorizontalAlignment() const
+{
+	return HorizontalAlignment;
+}
+
 void UWidgetSwitcherSlot::SetHorizontalAlignment(EHorizontalAlignment InHorizontalAlignment)
 {
 	HorizontalAlignment = InHorizontalAlignment;
@@ -61,6 +76,11 @@ void UWidgetSwitcherSlot::SetHorizontalAlignment(EHorizontalAlignment InHorizont
 	}
 }
 
+EVerticalAlignment UWidgetSwitcherSlot::GetVerticalAlignment() const
+{
+	return VerticalAlignment;
+}
+
 void UWidgetSwitcherSlot::SetVerticalAlignment(EVerticalAlignment InVerticalAlignment)
 {
 	VerticalAlignment = InVerticalAlignment;
@@ -69,10 +89,13 @@ void UWidgetSwitcherSlot::SetVerticalAlignment(EVerticalAlignment InVerticalAlig
 		Slot->SetVerticalAlignment(InVerticalAlignment);
 	}
 }
+PRAGMA_ENABLE_DEPRECATION_WARNINGS
 
 void UWidgetSwitcherSlot::SynchronizeProperties()
 {
+PRAGMA_DISABLE_DEPRECATION_WARNINGS
 	SetPadding(Padding);
 	SetHorizontalAlignment(HorizontalAlignment);
 	SetVerticalAlignment(VerticalAlignment);
+PRAGMA_ENABLE_DEPRECATION_WARNINGS
 }
diff --git a/Engine/Source/Runtime/UMG/Private/Components/WindowTitleBarAreaSlot.cpp b/Engine/Source/Runtime/UMG/Private/Components/WindowTitleBarAreaSlot.cpp
--- a/Engine/Source/Runtime/UMG/Private/Components/WindowTitleBarAreaSlot.cpp
+++ b/Engine/Source/Runtime/UMG/Private/Components/WindowTitleBarAreaSlot.cpp
@@ -10,46 +10,72 @@
 UWindowTitleBarAreaSlot::UWindowTitleBarAreaSlot(const FObjectInitializer& ObjectInitializer)
 	: Super(ObjectInitializer)
 {
-	Padding = FMargin(0);
+PRAGMA_DISABLE_DEPRECATION_WARNINGS
+	Padding = FMargin(0.f);
 
 	HorizontalAlignment = HAlign_Fill;
 	VerticalAlignment = VAlign_Fill;
+PRAGMA_ENABLE_DEPRECATION_WARNINGS
 }
 
 void UWindowTitleBarAreaSlot::BuildSlot(TSharedRef<SWindowTitleBarArea> InWindowTitleBarArea)
 {
 	WindowTitleBarArea = InWindowTitleBarArea;
 
+PRAGMA_DISABLE_DEPRECATION_WARNINGS
 	WindowTitleBarArea->SetPadding(Padding);
 	WindowTitleBarArea->SetHAlign(HorizontalAlignment);
 	WindowTitleBarArea->SetVAlign(VerticalAlignment);
+PRAGMA_ENABLE_DEPRECATION_WARNINGS
 
 	WindowTitleBarArea->SetContent(Content ? Content->TakeWidget() : SNullWidget::NullWidget);
 }
 
+PRAGMA_DISABLE_DEPRECATION_WARNINGS
+FMargin UWindowTitleBarAreaSlot::GetPadding() const
+{
+	return Padding;
+}
+
 void UWindowTitleBarAreaSlot::SetPadding(FMargin InPadding)
 {
+	Padding = InPadding;
 	CastChecked<UWindowTitleBarArea>(Parent)->SetPadding(InPadding);
 }
 
+EHorizontalAlignment UWindowTitleBarAreaSlot::GetHorizontalAlignment() const
+{
+	return HorizontalAlignment;
+}
+
 void UWindowTitleBarAreaSlot::SetHorizontalAlignment(EHorizontalAlignment InHorizontalAlignment)
 {
+	HorizontalAlignment = InHorizontalAlignment;
 	CastChecked<UWindowTitleBarArea>(Parent)->SetHorizontalAlignment(InHorizontalAlignment);
 }
 
+EVerticalAlignment UWindowTitleBarAreaSlot::GetVerticalAlignment() const
+{
+	return VerticalAlignment;
+}
+
 void UWindowTitleBarAreaSlot::SetVerticalAlignment(EVerticalAlignment InVerticalAlignment)
 {
+	VerticalAlignment = InVerticalAlignment;
 	CastChecked<UWindowTitleBarArea>(Parent)->SetVerticalAlignment(InVerticalAlignment);
 }
+PRAGMA_ENABLE_DEPRECATION_WARNINGS
 
 void UWindowTitleBarAreaSlot::SynchronizeProperties()
 {
+PRAGMA_DISABLE_DEPRECATION_WARNINGS
 	if (WindowTitleBarArea.IsValid())
 	{
 		SetPadding(Padding);
 		SetHorizontalAlignment(HorizontalAlignment);
 		SetVerticalAlignment(VerticalAlignment);
 	}
+PRAGMA_ENABLE_DEPRECATION_WARNINGS
 }
 
 void UWindowTitleBarAreaSlot::ReleaseSlateResources(bool bReleaseChildren)
diff --git a/Engine/Source/Runtime/UMG/Private/Components/WrapBoxSlot.cpp b/Engine/Source/Runtime/UMG/Private/Components/WrapBoxSlot.cpp
--- a/Engine/Source/Runtime/UMG/Private/Components/WrapBoxSlot.cpp
+++ b/Engine/Source/Runtime/UMG/Private/Components/WrapBoxSlot.cpp
@@ -10,12 +10,14 @@ UWrapBoxSlot::UWrapBoxSlot(const FObjectInitializer& ObjectInitializer)
 	: Super(ObjectInitializer)
 	, Slot(nullptr)
 {
+PRAGMA_DISABLE_DEPRECATION_WARNINGS
 	HorizontalAlignment = HAlign_Fill;
 	VerticalAlignment = VAlign_Fill;
 
 	bFillEmptySpace = false;
 	FillSpanWhenLessThan = 0;
 	bForceNewLine = false;
+PRAGMA_ENABLE_DEPRECATION_WARNINGS
 }
 
 void UWrapBoxSlot::ReleaseSlateResources(bool bReleaseChildren)
@@ -27,8 +29,8 @@ void UWrapBoxSlot::ReleaseSlateResources(bool bReleaseChildren)
 
 void UWrapBoxSlot::BuildSlot(TSharedRef<SWrapBox> WrapBox)
 {
+PRAGMA_DISABLE_DEPRECATION_WARNINGS
 	WrapBox->AddSlot()
-		.Expose(Slot)
 		.Padding(Padding)
 		.HAlign(HorizontalAlignment)
 		.VAlign(VerticalAlignment)
@@ -39,6 +41,13 @@ void UWrapBoxSlot::BuildSlot(TSharedRef<SWrapBox> WrapBox)
 		[
 			Content == nullptr ? SNullWidget::NullWidget : Content->TakeWidget()
 		];
+PRAGMA_ENABLE_DEPRECATION_WARNINGS
+}
+
+PRAGMA_DISABLE_DEPRECATION_WARNINGS
+FMargin UWrapBoxSlot::GetPadding() const
+{
+	return Padding;
 }
 
 void UWrapBoxSlot::SetPadding(FMargin InPadding)
@@ -50,6 +59,11 @@ void UWrapBoxSlot::SetPadding(FMargin InPadding)
 	}
 }
 
+bool UWrapBoxSlot::DoesFillEmptySpace() const
+{
+	return bFillEmptySpace;
+}
+
 void UWrapBoxSlot::SetFillEmptySpace(bool InbFillEmptySpace)
 {
 	bFillEmptySpace = InbFillEmptySpace;
@@ -59,6 +73,11 @@ void UWrapBoxSlot::SetFillEmptySpace(bool InbFillEmptySpace)
 	}
 }
 
+float UWrapBoxSlot::GetFillSpanWhenLessThan() const
+{
+	return FillSpanWhenLessThan;
+}
+
 void UWrapBoxSlot::SetFillSpanWhenLessThan(float InFillSpanWhenLessThan)
 {
 	FillSpanWhenLessThan = InFillSpanWhenLessThan;
@@ -68,6 +87,11 @@ void UWrapBoxSlot::SetFillSpanWhenLessThan(float InFillSpanWhenLessThan)
 	}
 }
 
+EHorizontalAlignment UWrapBoxSlot::GetHorizontalAlignment() const
+{
+	return HorizontalAlignment;
+}
+
 void UWrapBoxSlot::SetHorizontalAlignment(EHorizontalAlignment InHorizontalAlignment)
 {
 	HorizontalAlignment = InHorizontalAlignment;
@@ -77,6 +101,11 @@ void UWrapBoxSlot::SetHorizontalAlignment(EHorizontalAlignment InHorizontalAlign
 	}
 }
 
+EVerticalAlignment UWrapBoxSlot::GetVerticalAlignment() const
+{
+	return VerticalAlignment;
+}
+
 void UWrapBoxSlot::SetVerticalAlignment(EVerticalAlignment InVerticalAlignment)
 {
 	VerticalAlignment = InVerticalAlignment;
@@ -86,6 +115,11 @@ void UWrapBoxSlot::SetVerticalAlignment(EVerticalAlignment InVerticalAlignment)
 	}
 }
 
+bool UWrapBoxSlot::DoesForceNewLine() const
+{
+	return bForceNewLine;
+}
+
 void UWrapBoxSlot::SetNewLine(bool InForceNewLine)
 {
 	bForceNewLine = InForceNewLine;
@@ -94,13 +128,16 @@ void UWrapBoxSlot::SetNewLine(bool InForceNewLine)
 		Slot->SetForceNewLine(InForceNewLine);
 	}
 }
+PRAGMA_ENABLE_DEPRECATION_WARNINGS
 
 void UWrapBoxSlot::SynchronizeProperties()
 {
+PRAGMA_DISABLE_DEPRECATION_WARNINGS
 	SetPadding(Padding);
 	SetFillEmptySpace(bFillEmptySpace);
 	SetFillSpanWhenLessThan(FillSpanWhenLessThan);
 	SetHorizontalAlignment(HorizontalAlignment);
 	SetVerticalAlignment(VerticalAlignment);
 	SetNewLine(bForceNewLine);
+PRAGMA_ENABLE_DEPRECATION_WARNINGS
 }
diff --git a/Engine/Source/Runtime/UMG/Public/Components/BackgroundBlurSlot.h b/Engine/Source/Runtime/UMG/Public/Components/BackgroundBlurSlot.h
--- a/Engine/Source/Runtime/UMG/Public/Components/BackgroundBlurSlot.h
+++ b/Engine/Source/Runtime/UMG/Public/Components/BackgroundBlurSlot.h
@@ -22,33 +22,42 @@ class UMG_API UBackgroundBlurSlot : public UPanelSlot
 
 public:
 
+	FMargin GetPadding() const;
+
 	UFUNCTION(BlueprintCallable, Category="Layout|Background Blur Slot")
 	void SetPadding(FMargin InPadding);
 
+	EHorizontalAlignment GetHorizontalAlignment() const;
+
 	UFUNCTION(BlueprintCallable, Category="Layout|Background Blur Slot")
 	void SetHorizontalAlignment(EHorizontalAlignment InHorizontalAlignment);
 
+	EVerticalAlignment GetVerticalAlignment() const;
+
 	UFUNCTION(BlueprintCallable, Category="Layout|Background Blur Slot")
 	void SetVerticalAlignment(EVerticalAlignment InVerticalAlignment);
 
 protected:
 	/** The padding area between the slot and the content it contains. */
-	UPROPERTY(EditAnywhere, Category="Layout|Background Blur Slot")
+	UE_DEPRECATED(5.1, "Direct access to Padding is deprecated. Please use the getter or setter.")
+	UPROPERTY(EditAnywhere, BlueprintReadWrite, Setter, BlueprintSetter="SetPadding", Category="Layout|Background Blur Slot")
 	FMargin Padding;
 
 	/** The alignment of the object horizontally. */
-	UPROPERTY(EditAnywhere, Category="Layout|Background Blur Slot")
+	UE_DEPRECATED(5.1, "Direct access to HorizontalAlignment is deprecated. Please use the getter or setter.")
+	UPROPERTY(EditAnywhere, BlueprintReadWrite, Setter, BlueprintSetter="SetHorizontalAlignment", Category="Layout|Background Blur Slot")
 	TEnumAsByte<EHorizontalAlignment> HorizontalAlignment;
 
 	/** The alignment of the object vertically. */
-	UPROPERTY(EditAnywhere, Category="Layout|Background Blur Slot")
+	UE_DEPRECATED(5.1, "Direct access to VerticalAlignment is deprecated. Please use the getter or setter.")
+	UPROPERTY(EditAnywhere, BlueprintReadWrite, Setter, BlueprintSetter="SetVerticalAlignment", Category="Layout|Background Blur Slot")
 	TEnumAsByte<EVerticalAlignment> VerticalAlignment;
 
 public:
 
-	// UPanelSlot interface
+	//~ Begin UPanelSlot interface
 	virtual void SynchronizeProperties() override;
-	// End of UPanelSlot interface
+	//~ End of UPanelSlot interface
 
 	/** Builds the underlying slot for the slate BackgroundBlur. */
 	void BuildSlot(TSharedRef<SBackgroundBlur> InBackgroundBlur);
@@ -59,9 +68,9 @@ public:
 
 #if WITH_EDITOR
 
-	// UObject interface
+	//~ Begin UObject interface
 	virtual void PostEditChangeProperty(struct FPropertyChangedEvent& PropertyChangedEvent) override;
-	// End of UObject interface
+	//~ End of UObject interface
 
 #endif
 
diff --git a/Engine/Source/Runtime/UMG/Public/Components/BorderSlot.h b/Engine/Source/Runtime/UMG/Public/Components/BorderSlot.h
--- a/Engine/Source/Runtime/UMG/Public/Components/BorderSlot.h
+++ b/Engine/Source/Runtime/UMG/Public/Components/BorderSlot.h
@@ -23,33 +23,42 @@ class UMG_API UBorderSlot : public UPanelSlot
 
 public:
 
+	FMargin GetPadding() const;
+
 	UFUNCTION(BlueprintCallable, Category="Layout|Border Slot")
 	void SetPadding(FMargin InPadding);
 
+	EHorizontalAlignment GetHorizontalAlignment() const;
+
 	UFUNCTION(BlueprintCallable, Category="Layout|Border Slot")
 	void SetHorizontalAlignment(EHorizontalAlignment InHorizontalAlignment);
 
+	EVerticalAlignment GetVerticalAlignment() const;
+
 	UFUNCTION(BlueprintCallable, Category="Layout|Border Slot")
 	void SetVerticalAlignment(EVerticalAlignment InVerticalAlignment);
 
 protected:
 	/** The padding area between the slot and the content it contains. */
-	UPROPERTY(EditAnywhere, Category="Layout|Border Slot")
+	UE_DEPRECATED(5.1, "Direct access to Padding is deprecated. Please use the getter or setter.")
+	UPROPERTY(EditAnywhere, BlueprintReadWrite, Setter, BlueprintSetter="SetPadding", Category="Layout|Border Slot")
 	FMargin Padding;
 
 	/** The alignment of the object horizontally. */
-	UPROPERTY(EditAnywhere, Category="Layout|Border Slot")
+	UE_DEPRECATED(5.1, "Direct access to HorizontalAlignment is deprecated. Please use the getter or setter.")
+	UPROPERTY(EditAnywhere, BlueprintReadWrite, Setter, BlueprintSetter="SetHorizontalAlignment", Category="Layout|Border Slot")
 	TEnumAsByte<EHorizontalAlignment> HorizontalAlignment;
 
 	/** The alignment of the object vertically. */
-	UPROPERTY(EditAnywhere, Category="Layout|Border Slot")
+	UE_DEPRECATED(5.1, "Direct access to VerticalAlignment is deprecated. Please use the getter or setter.")
+	UPROPERTY(EditAnywhere, BlueprintReadWrite, Setter, BlueprintSetter="SetVerticalAlignment", Category="Layout|Border Slot")
 	TEnumAsByte<EVerticalAlignment> VerticalAlignment;
 
 public:
 
-	// UPanelSlot interface
+	//~ Begin UPanelSlot interface
 	virtual void SynchronizeProperties() override;
-	// End of UPanelSlot interface
+	//~ End of UPanelSlot interface
 
 	/** Builds the underlying slot for the slate border. */
 	void BuildSlot(TSharedRef<SBorder> InBorder);
@@ -60,9 +69,9 @@ public:
 
 #if WITH_EDITOR
 
-	// UObject interface
+	//~ Begin UObject interface
 	virtual void PostEditChangeProperty(struct FPropertyChangedEvent& PropertyChangedEvent) override;
-	// End of UObject interface
+	//~ End of UObject interface
 
 #endif
 
diff --git a/Engine/Source/Runtime/UMG/Public/Components/HorizontalBoxSlot.h b/Engine/Source/Runtime/UMG/Public/Components/HorizontalBoxSlot.h
--- a/Engine/Source/Runtime/UMG/Public/Components/HorizontalBoxSlot.h
+++ b/Engine/Source/Runtime/UMG/Public/Components/HorizontalBoxSlot.h
@@ -22,32 +22,48 @@ private:
 
 public:
 	
-	/** The amount of padding between the slots parent and the content. */
-	UPROPERTY(EditAnywhere, BlueprintReadOnly, Category="Layout|Horizontal Box Slot")
-	FMargin Padding;
-
 	/** How much space this slot should occupy in the direction of the panel. */
-	UPROPERTY(EditAnywhere, BlueprintReadOnly, Category="Layout|Horizontal Box Slot")
+	UE_DEPRECATED(5.1, "Direct access to Size is deprecated. Please use the getter or setter.")
+	UPROPERTY(EditAnywhere, BlueprintReadWrite, Setter, BlueprintSetter = "SetSize", Category = "Layout|Horizontal Box Slot")
 	FSlateChildSize Size;
 
-	UPROPERTY(EditAnywhere, BlueprintReadOnly, Category="Layout|Horizontal Box Slot")
+	/** The amount of padding between the slots parent and the content. */
+	UE_DEPRECATED(5.1, "Direct access to Padding is deprecated. Please use the getter or setter.")
+	UPROPERTY(EditAnywhere, BlueprintReadWrite, Setter, BlueprintSetter = "SetPadding", Category = "Layout|Horizontal Box Slot")
+	FMargin Padding;
+
+	UE_DEPRECATED(5.1, "Direct access to HorizontalAlignment is deprecated. Please use the getter or setter.")
+	UPROPERTY(EditAnywhere, BlueprintReadWrite, Setter, BlueprintSetter = "SetHorizontalAlignment", Category = "Layout|Horizontal Box Slot")
 	TEnumAsByte<EHorizontalAlignment> HorizontalAlignment;
 
-	UPROPERTY(EditAnywhere, BlueprintReadOnly, Category="Layout|Horizontal Box Slot")
+	UE_DEPRECATED(5.1, "Direct access to VerticalAlignment is deprecated. Please use the getter or setter.")
+	UPROPERTY(EditAnywhere, BlueprintReadWrite, Setter, BlueprintSetter = "SetVerticalAlignment", Category = "Layout|Horizontal Box Slot")
 	TEnumAsByte<EVerticalAlignment> VerticalAlignment;
 
+public:
+
+	FMargin GetPadding() const;
+
 	UFUNCTION(BlueprintCallable, Category="Layout|Horizontal Box Slot")
 	void SetPadding(FMargin InPadding);
 
+	FSlateChildSize GetSize() const;
+
 	UFUNCTION(BlueprintCallable, Category="Layout|Horizontal Box Slot")
 	void SetSize(FSlateChildSize InSize);
 
+	EHorizontalAlignment GetHorizontalAlignment() const;
+
 	UFUNCTION(BlueprintCallable, Category="Layout|Horizontal Box Slot")
 	void SetHorizontalAlignment(EHorizontalAlignment InHorizontalAlignment);
 
+	EVerticalAlignment GetVerticalAlignment() const;
+
 	UFUNCTION(BlueprintCallable, Category="Layout|Horizontal Box Slot")
 	void SetVerticalAlignment(EVerticalAlignment InVerticalAlignment);
 
+public:
+
 	void BuildSlot(TSharedRef<SHorizontalBox> HorizontalBox);
 
 	// UPanelSlot interface
diff --git a/Engine/Source/Runtime/UMG/Public/Components/SizeBoxSlot.h b/Engine/Source/Runtime/UMG/Public/Components/SizeBoxSlot.h
--- a/Engine/Source/Runtime/UMG/Public/Components/SizeBoxSlot.h
+++ b/Engine/Source/Runtime/UMG/Public/Components/SizeBoxSlot.h
@@ -20,7 +20,8 @@ class UMG_API USizeBoxSlot : public UPanelSlot
 
 public:
 	/** The padding area between the slot and the content it contains. */
-	UPROPERTY(EditAnywhere, BlueprintReadOnly, Category="Layout|SizeBox Slot")
+	UE_DEPRECATED(5.1, "Direct access to Padding is deprecated. Please use the getter or setter.")
+	UPROPERTY(EditAnywhere, BlueprintReadWrite, Setter, BlueprintSetter="SetPadding", Category="Layout|SizeBox Slot")
 	FMargin Padding;
 
 private:
@@ -29,26 +30,37 @@ private:
 
 public:
 	/** The alignment of the object horizontally. */
-	UPROPERTY(EditAnywhere, BlueprintReadOnly, Category="Layout|SizeBox Slot")
+	UE_DEPRECATED(5.1, "Direct access to HorizontalAlignment is deprecated. Please use the getter or setter.")
+	UPROPERTY(EditAnywhere, BlueprintReadWrite, Setter, BlueprintSetter="SetHorizontalAlignment", Category="Layout|SizeBox Slot")
 	TEnumAsByte<EHorizontalAlignment> HorizontalAlignment;
 
 	/** The alignment of the object vertically. */
-	UPROPERTY(EditAnywhere, BlueprintReadOnly, Category="Layout|SizeBox Slot")
+	UE_DEPRECATED(5.1, "Direct access to VerticalAlignment is deprecated. Please use the getter or setter.")
+	UPROPERTY(EditAnywhere, BlueprintReadWrite, Setter, BlueprintSetter="SetVerticalAlignment", Category="Layout|SizeBox Slot")
 	TEnumAsByte<EVerticalAlignment> VerticalAlignment;
 
+public:
+
+	FMargin GetPadding() const;
 
 	UFUNCTION(BlueprintCallable, Category="Layout|SizeBox Slot")
 	void SetPadding(FMargin InPadding);
 
+	EHorizontalAlignment GetHorizontalAlignment() const;
+
 	UFUNCTION(BlueprintCallable, Category="Layout|SizeBox Slot")
 	void SetHorizontalAlignment(EHorizontalAlignment InHorizontalAlignment);
 
+	EVerticalAlignment GetVerticalAlignment() const;
+
 	UFUNCTION(BlueprintCallable, Category="Layout|SizeBox Slot")
 	void SetVerticalAlignment(EVerticalAlignment InVerticalAlignment);
 
-	// UPanelSlot interface
+public:
+
+	//~ UPanelSlot interface
 	virtual void SynchronizeProperties() override;
-	// End of UPanelSlot interface
+	//~ End of UPanelSlot interface
 
 	/** Builds the underlying slot for the slate SizeBox. */
 	void BuildSlot(TSharedRef<SBox> InSizeBox);
diff --git a/Engine/Source/Runtime/UMG/Public/Components/UniformGridSlot.h b/Engine/Source/Runtime/UMG/Public/Components/UniformGridSlot.h
--- a/Engine/Source/Runtime/UMG/Public/Components/UniformGridSlot.h
+++ b/Engine/Source/Runtime/UMG/Public/Components/UniformGridSlot.h
@@ -22,44 +22,60 @@ class UMG_API UUniformGridSlot : public UPanelSlot
 public:
 
 	/** The alignment of the object horizontally. */
-	UPROPERTY(EditAnywhere, BlueprintReadOnly, Category="Layout|Uniform Grid Slot")
+	UE_DEPRECATED(5.1, "Direct access to HorizontalAlignment is deprecated. Please use the getter or setter.")
+	UPROPERTY(EditAnywhere, BlueprintReadWrite, Setter, BlueprintSetter="SetHorizontalAlignment", Category="Layout|Uniform Grid Slot")
 	TEnumAsByte<EHorizontalAlignment> HorizontalAlignment;
 
 	/** The alignment of the object vertically. */
-	UPROPERTY(EditAnywhere, BlueprintReadOnly, Category="Layout|Uniform Grid Slot")
+	UE_DEPRECATED(5.1, "Direct access to VerticalAlignment is deprecated. Please use the getter or setter.")
+	UPROPERTY(EditAnywhere, BlueprintReadWrite, Setter, BlueprintSetter="SetVerticalAlignment", Category="Layout|Uniform Grid Slot")
 	TEnumAsByte<EVerticalAlignment> VerticalAlignment;
 	
 	/** The row index of the cell this slot is in */
-	UPROPERTY(EditAnywhere, BlueprintReadOnly, meta=( UIMin = "0" ), Category="Layout|Uniform Grid Slot")
+	UE_DEPRECATED(5.1, "Direct access to Row is deprecated. Please use the getter or setter.")
+	UPROPERTY(EditAnywhere, BlueprintReadWrite, Setter, BlueprintSetter="SetRow", meta=( UIMin = "0" ), Category="Layout|Uniform Grid Slot")
 	int32 Row;
 	
 	/** The column index of the cell this slot is in */
-	UPROPERTY(EditAnywhere, BlueprintReadOnly, meta=( UIMin = "0" ), Category="Layout|Uniform Grid Slot")
+	UE_DEPRECATED(5.1, "Direct access to Column is deprecated. Please use the getter or setter.")
+	UPROPERTY(EditAnywhere, BlueprintReadWrite, Setter, BlueprintSetter="SetColumn", meta=( UIMin = "0" ), Category="Layout|Uniform Grid Slot")
 	int32 Column;
 
 public:
 
+	/** Gets the row index of the slot */
+	int32 GetRow() const;
+
 	/** Sets the row index of the slot, this determines what cell the slot is in the panel */
 	UFUNCTION(BlueprintCallable, Category="Layout|Uniform Grid Slot")
 	void SetRow(int32 InRow);
 
+	/** Gets the column index of the slot. */
+	int32 GetColumn() const;
+
 	/** Sets the column index of the slot, this determines what cell the slot is in the panel */
 	UFUNCTION(BlueprintCallable, Category="Layout|Uniform Grid Slot")
 	void SetColumn(int32 InColumn);
 
+	/**  */
+	EHorizontalAlignment GetHorizontalAlignment() const;
+
 	/**  */
 	UFUNCTION(BlueprintCallable, Category="Layout|Uniform Grid Slot")
 	void SetHorizontalAlignment(EHorizontalAlignment InHorizontalAlignment);
 
+	/**  */
+	EVerticalAlignment GetVerticalAlignment() const;
+
 	/**  */
 	UFUNCTION(BlueprintCallable, Category="Layout|Uniform Grid Slot")
 	void SetVerticalAlignment(EVerticalAlignment InVerticalAlignment);
 
 public:
 
-	// UPanelSlot interface
+	//~ UPanelSlot interface
 	virtual void SynchronizeProperties() override;
-	// End of UPanelSlot interface
+	//~ End of UPanelSlot interface
 
 	/** Builds the underlying FSlot for the Slate layout panel. */
 	void BuildSlot(TSharedRef<SUniformGridPanel> GridPanel);
@@ -72,8 +88,6 @@ public:
 #endif //WITH_EDITOR
 
 private:
-	//TODO UMG Slots should hold weak or shared refs to slots.
-
 	/** A raw pointer to the slot to allow us to adjust the size, padding...etc at runtime. */
 	SUniformGridPanel::FSlot* Slot;
 };
diff --git a/Engine/Source/Runtime/UMG/Public/Components/VerticalBoxSlot.h b/Engine/Source/Runtime/UMG/Public/Components/VerticalBoxSlot.h
--- a/Engine/Source/Runtime/UMG/Public/Components/VerticalBoxSlot.h
+++ b/Engine/Source/Runtime/UMG/Public/Components/VerticalBoxSlot.h
@@ -21,40 +21,49 @@ class UMG_API UVerticalBoxSlot : public UPanelSlot
 public:
 	
 	/** How much space this slot should occupy in the direction of the panel. */
-	UPROPERTY(EditAnywhere, BlueprintReadOnly, Category="Layout|Vertical Box Slot", meta = (DisplayAfter = "Padding"))
+	UE_DEPRECATED(5.1, "Direct access to Size is deprecated. Please use the getter or setter.")
+	UPROPERTY(EditAnywhere, BlueprintReadWrite, Setter, BlueprintSetter="SetSize", Category="Layout|Vertical Box Slot", meta = (DisplayAfter = "Padding"))
 	FSlateChildSize Size;
 
 	/** The padding area between the slot and the content it contains. */
-	UPROPERTY(EditAnywhere, BlueprintReadOnly, Category="Layout|Vertical Box Slot")
+	UE_DEPRECATED(5.1, "Direct access to Padding is deprecated. Please use the getter or setter.")
+	UPROPERTY(EditAnywhere, BlueprintReadWrite, Setter, BlueprintSetter="SetPadding", Category="Layout|Vertical Box Slot")
 	FMargin Padding;
 
-private:
-	//TODO UMG Slots should hold weak or shared refs to slots.
-
-	/** A raw pointer to the slot to allow us to adjust the size, padding...etc at runtime. */
-	SVerticalBox::FSlot* Slot;
-
-public:
-
 	/** The alignment of the object horizontally. */
-	UPROPERTY(EditAnywhere, BlueprintReadOnly, Category="Layout|Vertical Box Slot")
+	UE_DEPRECATED(5.1, "Direct access to HorizontalAlignment is deprecated. Please use the getter or setter.")
+	UPROPERTY(EditAnywhere, BlueprintReadWrite, Setter, BlueprintSetter="SetHorizontalAlignment", Category="Layout|Vertical Box Slot")
 	TEnumAsByte<EHorizontalAlignment> HorizontalAlignment;
 
 	/** The alignment of the object vertically. */
-	UPROPERTY(EditAnywhere, BlueprintReadOnly, Category="Layout|Vertical Box Slot")
+	UE_DEPRECATED(5.1, "Direct access to VerticalAlignment is deprecated. Please use the getter or setter.")
+	UPROPERTY(EditAnywhere, BlueprintReadWrite, Setter, BlueprintSetter="SetVerticalAlignment", Category="Layout|Vertical Box Slot")
 	TEnumAsByte<EVerticalAlignment> VerticalAlignment;
 
+private:
+
+	/** A raw pointer to the slot to allow us to adjust the size, padding...etc at runtime. */
+	SVerticalBox::FSlot* Slot;
+
 public:
 
+	FMargin GetPadding() const;
+
 	UFUNCTION(BlueprintCallable, Category="Layout|Vertical Box Slot")
 	void SetPadding(FMargin InPadding);
 
+	FSlateChildSize GetSize() const;
+
 	UFUNCTION(BlueprintCallable, Category="Layout|Vertical Box Slot")
 	void SetSize(FSlateChildSize InSize);
 
+	EHorizontalAlignment GetHorizontalAlignment() const;
+
 	UFUNCTION(BlueprintCallable, Category="Layout|Vertical Box Slot")
 	void SetHorizontalAlignment(EHorizontalAlignment InHorizontalAlignment);
 
+	EVerticalAlignment GetVerticalAlignment() const;
+
 	UFUNCTION(BlueprintCallable, Category="Layout|Vertical Box Slot")
 	void SetVerticalAlignment(EVerticalAlignment InVerticalAlignment);
 
diff --git a/Engine/Source/Runtime/UMG/Public/Components/WidgetSwitcherSlot.h b/Engine/Source/Runtime/UMG/Public/Components/WidgetSwitcherSlot.h
--- a/Engine/Source/Runtime/UMG/Public/Components/WidgetSwitcherSlot.h
+++ b/Engine/Source/Runtime/UMG/Public/Components/WidgetSwitcherSlot.h
@@ -28,33 +28,42 @@ private:
 public:
 
 	/** The padding area between the slot and the content it contains. */
-	UPROPERTY(EditAnywhere, BlueprintReadOnly, Category="Layout|Widget Switcher Slot")
+	UE_DEPRECATED(5.1, "Direct access to Padding is deprecated. Please use the getter or setter.")
+	UPROPERTY(EditAnywhere, BlueprintReadWrite, Setter, BlueprintSetter="SetPadding", Category="Layout|Widget Switcher Slot")
 	FMargin Padding;
 
 	/** The alignment of the object horizontally. */
-	UPROPERTY(EditAnywhere, BlueprintReadOnly, Category="Layout|Widget Switcher Slot")
+	UE_DEPRECATED(5.1, "Direct access to HorizontalAlignment is deprecated. Please use the getter or setter.")
+	UPROPERTY(EditAnywhere, BlueprintReadWrite, Setter, BlueprintSetter="SetHorizontalAlignment", Category="Layout|Widget Switcher Slot")
 	TEnumAsByte<EHorizontalAlignment> HorizontalAlignment;
 
 	/** The alignment of the object vertically. */
-	UPROPERTY(EditAnywhere, BlueprintReadOnly, Category="Layout|Widget Switcher Slot")
+	UE_DEPRECATED(5.1, "Direct access to VerticalAlignment is deprecated. Please use the getter or setter.")
+	UPROPERTY(EditAnywhere, BlueprintReadWrite, Setter, BlueprintSetter="SetVerticalAlignment", Category="Layout|Widget Switcher Slot")
 	TEnumAsByte<EVerticalAlignment> VerticalAlignment;
 
 public:
 
+	FMargin GetPadding() const;
+
 	UFUNCTION(BlueprintCallable, Category="Layout|Widget Switcher Slot")
 	void SetPadding(FMargin InPadding);
 
+	EHorizontalAlignment GetHorizontalAlignment() const;
+
 	UFUNCTION(BlueprintCallable, Category="Layout|Widget Switcher Slot")
 	void SetHorizontalAlignment(EHorizontalAlignment InHorizontalAlignment);
 
+	EVerticalAlignment GetVerticalAlignment() const;
+
 	UFUNCTION(BlueprintCallable, Category="Layout|Widget Switcher Slot")
 	void SetVerticalAlignment(EVerticalAlignment InVerticalAlignment);
 
 public:
 
-	// UPanelSlot interface
+	//~ UPanelSlot interface
 	virtual void SynchronizeProperties() override;
-	// End of UPanelSlot interface
+	//~ End of UPanelSlot interface
 
 	virtual void ReleaseSlateResources(bool bReleaseChildren) override;
 
diff --git a/Engine/Source/Runtime/UMG/Public/Components/WindowTitleBarAreaSlot.h b/Engine/Source/Runtime/UMG/Public/Components/WindowTitleBarAreaSlot.h
--- a/Engine/Source/Runtime/UMG/Public/Components/WindowTitleBarAreaSlot.h
+++ b/Engine/Source/Runtime/UMG/Public/Components/WindowTitleBarAreaSlot.h
@@ -20,34 +20,43 @@ class UMG_API UWindowTitleBarAreaSlot : public UPanelSlot
 
 public:
 
+	FMargin GetPadding() const;
+
 	UFUNCTION(BlueprintCallable, Category="Layout|WindowTitleBarArea Slot")
 	void SetPadding(FMargin InPadding);
 
+	EHorizontalAlignment GetHorizontalAlignment() const;
+
 	UFUNCTION(BlueprintCallable, Category="Layout|WindowTitleBarArea Slot")
 	void SetHorizontalAlignment(EHorizontalAlignment InHorizontalAlignment);
 
+	EVerticalAlignment GetVerticalAlignment() const;
+
 	UFUNCTION(BlueprintCallable, Category = "Layout|WindowTitleBarArea Slot")
 	void SetVerticalAlignment(EVerticalAlignment InVerticalAlignment);
 
 protected:
 
 	/** The padding area between the slot and the content it contains. */
-	UPROPERTY(EditAnywhere, BlueprintReadOnly, Category = "Layout|WindowTitleBarArea Slot")
+	UE_DEPRECATED(5.1, "Direct access to Padding is deprecated. Please use the getter or setter.")
+	UPROPERTY(EditAnywhere, BlueprintReadWrite, Setter, BlueprintSetter = "SetPadding", Category = "Layout|WindowTitleBarArea Slot")
 	FMargin Padding;
 
 	/** The alignment of the object horizontally. */
-	UPROPERTY(EditAnywhere, BlueprintReadOnly, Category = "Layout|WindowTitleBarArea Slot")
+	UE_DEPRECATED(5.1, "Direct access to HorizontalAlignment is deprecated. Please use the getter or setter.")
+	UPROPERTY(EditAnywhere, BlueprintReadWrite, Setter, BlueprintSetter = "SetHorizontalAlignment", Category = "Layout|WindowTitleBarArea Slot")
 	TEnumAsByte<EHorizontalAlignment> HorizontalAlignment;
 
 	/** The alignment of the object vertically. */
-	UPROPERTY(EditAnywhere, Category = "Layout|WindowTitleBarArea Slot")
+	UE_DEPRECATED(5.1, "Direct access to VerticalAlignment is deprecated. Please use the getter or setter.")
+	UPROPERTY(EditAnywhere, BlueprintReadWrite, Setter, BlueprintSetter = "SetVerticalAlignment", Category = "Layout|WindowTitleBarArea Slot")
 	TEnumAsByte<EVerticalAlignment> VerticalAlignment;
 
 public:
 
-	// UPanelSlot interface
+	//~ UPanelSlot interface
 	virtual void SynchronizeProperties() override;
-	// End of UPanelSlot interface
+	//~ End of UPanelSlot interface
 
 	/** Builds the underlying FSlot for the Slate layout panel. */
 	void BuildSlot(TSharedRef<SWindowTitleBarArea> WindowTitleBarArea);
diff --git a/Engine/Source/Runtime/UMG/Public/Components/WrapBoxSlot.h b/Engine/Source/Runtime/UMG/Public/Components/WrapBoxSlot.h
--- a/Engine/Source/Runtime/UMG/Public/Components/WrapBoxSlot.h
+++ b/Engine/Source/Runtime/UMG/Public/Components/WrapBoxSlot.h
@@ -20,49 +20,67 @@ class UMG_API UWrapBoxSlot : public UPanelSlot
 public:
 
 	/** The padding area between the slot and the content it contains. */
-	UPROPERTY(EditAnywhere, BlueprintReadOnly, Category="Layout|Wrap Box Slot")
+	UE_DEPRECATED(5.1, "Direct access to Padding is deprecated. Please use the getter or setter.")
+	UPROPERTY(EditAnywhere, BlueprintReadWrite, Setter, BlueprintSetter="SetPadding", Category="Layout|Wrap Box Slot")
 	FMargin Padding;
 
-	/** Should this slot fill the remaining space on the line? */
-	UPROPERTY(EditAnywhere, BlueprintReadOnly, Category="Layout|Wrap Box Slot")
-	bool bFillEmptySpace;
-
 	/**
 	 * If the total available space in the wrap panel drops below this threshold, this slot will attempt to fill an entire line.
 	 * NOTE: A value of 0, denotes no filling will occur.
 	 */
-	UPROPERTY(EditAnywhere, BlueprintReadOnly, Category="Layout|Wrap Box Slot")
+	UE_DEPRECATED(5.1, "Direct access to FillSpanWhenLessThan is deprecated. Please use the getter or setter.")
+	UPROPERTY(EditAnywhere, BlueprintReadWrite, Setter, BlueprintSetter="SetFillSpanWhenLessThan", Category="Layout|Wrap Box Slot")
 	float FillSpanWhenLessThan;
 
 	/** The alignment of the object horizontally. */
-	UPROPERTY(EditAnywhere, BlueprintReadOnly, Category="Layout|Wrap Box Slot")
+	UE_DEPRECATED(5.1, "Direct access to HorizontalAlignment is deprecated. Please use the getter or setter.")
+	UPROPERTY(EditAnywhere, BlueprintReadWrite, Setter, BlueprintSetter="SetHorizontalAlignment", Category="Layout|Wrap Box Slot")
 	TEnumAsByte<EHorizontalAlignment> HorizontalAlignment;
 
 	/** The alignment of the object vertically. */
-	UPROPERTY(EditAnywhere, BlueprintReadOnly, Category="Layout|Wrap Box Slot")
+	UE_DEPRECATED(5.1, "Direct access to VerticalAlignment is deprecated. Please use the getter or setter.")
+	UPROPERTY(EditAnywhere, BlueprintReadWrite, Setter, BlueprintSetter="SetVerticalAlignment", Category="Layout|Wrap Box Slot")
 	TEnumAsByte<EVerticalAlignment> VerticalAlignment;
 
+	/** Should this slot fill the remaining space on the line? */
+	UE_DEPRECATED(5.1, "Direct access to bFillEmptySpace is deprecated. Please use the getter or setter.")
+	UPROPERTY(EditAnywhere, BlueprintReadWrite, Setter="SetFillEmptySpace", BlueprintSetter="SetFillEmptySpace", Category = "Layout|Wrap Box Slot")
+	bool bFillEmptySpace;
+
 	/** Should this slot start on a new line? */
-	UPROPERTY(EditAnywhere, BlueprintReadOnly, Category="Layout|Wrap Box Slot")
+	UE_DEPRECATED(5.1, "Direct access to bForceNewLine is deprecated. Please use the getter or setter.")
+	UPROPERTY(EditAnywhere, BlueprintReadWrite, Setter="SetNewLine", BlueprintSetter="SetNewLine", Category = "Layout|Wrap Box Slot")
 	bool bForceNewLine;
 
 public:
 
+	FMargin GetPadding() const;
+
 	UFUNCTION(BlueprintCallable, Category="Layout|Wrap Box Slot")
 	void SetPadding(FMargin InPadding);
 
+	bool DoesFillEmptySpace() const;
+
 	UFUNCTION(BlueprintCallable, Category="Layout|Wrap Box Slot")
 	void SetFillEmptySpace(bool InbFillEmptySpace);
 
+	float GetFillSpanWhenLessThan() const;
+
 	UFUNCTION(BlueprintCallable, Category="Layout|Wrap Box Slot")
 	void SetFillSpanWhenLessThan(float InFillSpanWhenLessThan);
 
+	EHorizontalAlignment  GetHorizontalAlignment() const;
+
 	UFUNCTION(BlueprintCallable, Category="Layout|Wrap Box Slot")
 	void SetHorizontalAlignment(EHorizontalAlignment InHorizontalAlignment);
 
+	EVerticalAlignment GetVerticalAlignment() const;
+
 	UFUNCTION(BlueprintCallable, Category="Layout|Wrap Box Slot")
 	void SetVerticalAlignment(EVerticalAlignment InVerticalAlignment);
 
+	bool DoesForceNewLine() const;
+
 	UFUNCTION(BlueprintCallable, Category="Layout|Wrap Box Slot")
 	void SetNewLine(bool InForceNewLine);
 
