Hide UObject::IsSelected on UVisual to avoid confusion with UI selection state

diff --git a/Engine/Source/Runtime/UMG/Public/Components/Visual.h b/Engine/Source/Runtime/UMG/Public/Components/Visual.h
--- a/Engine/Source/Runtime/UMG/Public/Components/Visual.h
+++ b/Engine/Source/Runtime/UMG/Public/Components/Visual.h
@@ -20,4 +20,8 @@ public:
 	virtual void BeginDestroy() override;
 	virtual bool NeedsLoadForServer() const override;
 	//~ End UObject Interface
+
+private:
+	// Hide this to avoid confusion with UI selection state
+	using UObject::IsSelected;
 };
