UMG: Update the EditableTextBox widget with accessors.

diff --git a/Engine/Source/Runtime/UMG/Private/Components/EditableTextBox.cpp b/Engine/Source/Runtime/UMG/Private/Components/EditableTextBox.cpp
--- a/Engine/Source/Runtime/UMG/Private/Components/EditableTextBox.cpp
+++ b/Engine/Source/Runtime/UMG/Private/Components/EditableTextBox.cpp
@@ -21,6 +21,8 @@ static FEditableTextBoxStyle* EditorEditableTextBoxStyle = nullptr;
 UEditableTextBox::UEditableTextBox(const FObjectInitializer& ObjectInitializer)
 	: Super(ObjectInitializer)
 {
+
+	PRAGMA_DISABLE_DEPRECATION_WARNINGS
 	IsReadOnly = false;
 	IsPassword = false;
 	MinimumDesiredWidth = 0.0f;
@@ -32,6 +34,7 @@ UEditableTextBox::UEditableTextBox(const FObjectInitializer& ObjectInitializer)
 	AllowContextMenu = true;
 	VirtualKeyboardDismissAction = EVirtualKeyboardDismissAction::TextChangeOnDismiss;
 	OverflowPolicy = ETextOverflowPolicy::Clip;
+	PRAGMA_ENABLE_DEPRECATION_WARNINGS
 
 	if (DefaultEditableTextBoxStyle == nullptr)
 	{
@@ -74,6 +77,7 @@ void UEditableTextBox::ReleaseSlateResources(bool bReleaseChildren)
 	MyEditableTextBlock.Reset();
 }
 
+PRAGMA_DISABLE_DEPRECATION_WARNINGS
 TSharedRef<SWidget> UEditableTextBox::RebuildWidget()
 {
 	MyEditableTextBlock = SNew(SEditableTextBox)
@@ -142,6 +146,16 @@ void UEditableTextBox::SetText(FText InText)
 	}
 }
 
+FText UEditableTextBox::GetHintText() const
+{
+	if (MyEditableTextBlock.IsValid())
+	{
+		return MyEditableTextBlock->GetHintText();
+	}
+
+	return HintText;
+}
+
 void UEditableTextBox::SetHintText(FText InText)
 {
 	HintText = InText;
@@ -151,6 +165,20 @@ void UEditableTextBox::SetHintText(FText InText)
 	}
 }
 
+float UEditableTextBox::GetMinimumDesiredWidth() const
+{
+	return MinimumDesiredWidth;
+}
+
+void UEditableTextBox::SetMinDesiredWidth(float InMinDesiredWidth)
+{
+	MinimumDesiredWidth = InMinDesiredWidth;
+	if (MyEditableTextBlock.IsValid())
+	{
+		MyEditableTextBlock->SetMinimumDesiredWidth(MinimumDesiredWidth);
+	}
+}
+
 void UEditableTextBox::SetForegroundColor(FLinearColor color)
 {
 	WidgetStyle.ForegroundColor = color;
@@ -168,6 +196,11 @@ void UEditableTextBox::SetError(FText InError)
 	}
 }
 
+bool UEditableTextBox::GetIsReadOnly() const
+{
+	return IsReadOnly;
+}
+
 void UEditableTextBox::SetIsReadOnly(bool bReadOnly)
 {
 	IsReadOnly = bReadOnly;
@@ -177,6 +210,11 @@ void UEditableTextBox::SetIsReadOnly(bool bReadOnly)
 	}
 }
 
+bool UEditableTextBox::GetIsPassword() const
+{
+	return IsPassword;
+}
+
 void UEditableTextBox::SetIsPassword(bool bIsPassword)
 {
 	IsPassword = bIsPassword;
@@ -204,6 +242,12 @@ bool UEditableTextBox::HasError() const
 	return false;
 }
 
+
+ETextJustify::Type UEditableTextBox::GetJustification() const
+{
+	return Justification;
+}
+
 void UEditableTextBox::SetJustification(ETextJustify::Type InJustification)
 {
 	Justification = InJustification;
@@ -213,6 +257,11 @@ void UEditableTextBox::SetJustification(ETextJustify::Type InJustification)
 	}
 }
 
+ETextOverflowPolicy UEditableTextBox::GetTextOverflowPolicy() const
+{
+	return OverflowPolicy;
+}
+
 void UEditableTextBox::SetTextOverflowPolicy(ETextOverflowPolicy InOverflowPolicy)
 {
 	OverflowPolicy = InOverflowPolicy;
@@ -233,6 +282,7 @@ void UEditableTextBox::HandleOnTextCommitted(const FText& InText, ETextCommit::T
 	Text = InText;
 	OnTextCommitted.Broadcast(InText, CommitMethod);
 }
+PRAGMA_ENABLE_DEPRECATION_WARNINGS
 
 #if WITH_ACCESSIBILITY
 TSharedPtr<SWidget> UEditableTextBox::GetAccessibleWidget() const
diff --git a/Engine/Source/Runtime/UMG/Public/Components/EditableTextBox.h b/Engine/Source/Runtime/UMG/Public/Components/EditableTextBox.h
--- a/Engine/Source/Runtime/UMG/Public/Components/EditableTextBox.h
+++ b/Engine/Source/Runtime/UMG/Public/Components/EditableTextBox.h
@@ -34,7 +34,8 @@ public:
 
 public:
 	/** The text content for this editable text box widget */
-	UPROPERTY(EditAnywhere, Category=Content)
+	UE_DEPRECATED(5.1, "Direct access to Text is deprecated. Please use the getter or setter.")
+	UPROPERTY(EditAnywhere, BlueprintReadWrite, Getter, Setter, BlueprintGetter = "GetText", BlueprintSetter = "SetText", Category = "Content")
 	FText Text;
 
 	/** A bindable delegate to allow logic to drive the text of the widget */
@@ -48,7 +49,8 @@ public:
 	FEditableTextBoxStyle WidgetStyle;
 
 	/** Hint text that appears when there is no text in the text box */
-	UPROPERTY(EditAnywhere, BlueprintReadOnly, Category=Content)
+	UE_DEPRECATED(5.1, "Direct access to Hint Text is deprecated. Please use the getter or setter.")
+	UPROPERTY(EditAnywhere, BlueprintReadWrite, Getter, Setter, BlueprintSetter = "SetHintText", Category = Content, meta = (MultiLine = "true"))
 	FText HintText;
 
 	/** A bindable delegate to allow logic to drive the hint text of the widget */
@@ -56,15 +58,18 @@ public:
 	FGetText HintTextDelegate;
 
 	/** Sets whether this text box can actually be modified interactively by the user */
-	UPROPERTY(EditAnywhere, BlueprintReadOnly, Category=Appearance)
+	UE_DEPRECATED(5.1, "Direct access to IsReadOnly is deprecated. Please use the getter or setter.")
+	UPROPERTY(EditAnywhere, BlueprintReadWrite, Getter, Setter, Category = Appearance)
 	bool IsReadOnly;
 
 	/** Sets whether this text box is for storing a password */
-	UPROPERTY(EditAnywhere, BlueprintReadOnly, Category=Appearance)
+	UE_DEPRECATED(5.1, "Direct access to IsPassword is deprecated. Please use the getter or setter.")
+	UPROPERTY(EditAnywhere, BlueprintReadWrite, Getter, Setter, BlueprintSetter = "SetIsPassword", Category = Appearance)
 	bool IsPassword;
 
-	/** Minimum width that a text block should be */
-	UPROPERTY(EditAnywhere, BlueprintReadOnly, Category=Appearance)
+	/** The minimum desired size for the text */
+	UE_DEPRECATED(5.1, "Direct access to MinimumDesiredWidth is deprecated. Please use the getter or setter.")
+	UPROPERTY(EditAnywhere, BlueprintReadWrite, Getter, Setter=SetMinDesiredWidth, Category = "Appearance")
 	float MinimumDesiredWidth;
 
 	/** Workaround as we lose focus when the auto completion closes. */
@@ -108,11 +113,13 @@ public:
 	EVirtualKeyboardDismissAction VirtualKeyboardDismissAction;
 	
 	/** How the text should be aligned with the margin. */
-	UPROPERTY(EditAnywhere, BlueprintReadWrite, BlueprintSetter = SetJustification, Category=Appearance)
+	UE_DEPRECATED(5.1, "Direct access to Justification is deprecated. Please use the getter or setter.")
+	UPROPERTY(EditAnywhere, BlueprintReadWrite, Getter, Setter, BlueprintSetter = SetJustification, Category = Appearance)
 	TEnumAsByte<ETextJustify::Type> Justification;
 
-	/** Sets what should happen when text is clipped because the block does not have enough space */
-	UPROPERTY(EditAnywhere, BlueprintReadWrite, BlueprintSetter = SetTextOverflowPolicy, Category = Appearance)
+	/** Sets what happens to text that is clipped and doesn't fit within the clip rect for this widget */
+	UE_DEPRECATED(5.1, "Direct access to OverflowPolicy is deprecated. Please use the getter or setter.")
+	UPROPERTY(EditAnywhere, BlueprintReadWrite, Getter="GetTextOverflowPolicy", Setter= "SetTextOverflowPolicy", BlueprintSetter = "SetTextOverflowPolicy", Category = "Clipping", AdvancedDisplay, meta = (DisplayName = "Overflow Policy"))
 	ETextOverflowPolicy OverflowPolicy;
 
 
@@ -135,24 +142,52 @@ public:
 
 public:
 
-	/**  */
+	/**
+	* Gets the widget text
+	* @return The widget text
+	*/
 	UFUNCTION(BlueprintCallable, Category="Widget", meta=(DisplayName="GetText (Text Box)"))
 	FText GetText() const;
 
-	/**  */
+	/**
+	* Directly sets the widget text.
+	* Warning: This will wipe any binding created for the Text property!
+	* @param InText The text to assign to the widget
+	*/
 	UFUNCTION(BlueprintCallable, Category="Widget", meta=(DisplayName="SetText (Text Box)"))
 	void SetText(FText InText);
+	
+	/** Gets the Hint text that appears when there is no text in the text box */
+	FText GetHintText() const;
 
+	/**
+	* Sets the Hint text that appears when there is no text in the text box
+	* @param InHintText The text that appears when there is no text in the text box
+	*/
 	UFUNCTION(BlueprintCallable, Category = "Widget", meta = (DisplayName = "Set Hint Text (Text Box)"))
 	void SetHintText(FText InText);
 
+	/** @return the minimum desired width for this text box */
+	float GetMinimumDesiredWidth() const;
+
+	/**
+	*  Set the minimum desired width for this text box
+	*
+	*  @param InMinDesiredWidth new minimum desired width
+	*/
+	void SetMinDesiredWidth(float InMinDesiredWidth);
+
 	UFUNCTION(BlueprintCallable, Category="Widget",  meta=(DisplayName="SetError (Text Box)"))
 	void SetError(FText InError);
 
-	UFUNCTION(BlueprintCallable, Category="Widget", meta=(DisplayName="SetIsReadOnly (Text Box)"))
-	void SetIsReadOnly(bool bReadOnly);
+	bool GetIsReadOnly() const;
 
-	UFUNCTION(BlueprintCallable, Category = "Widget")
+	UFUNCTION(BlueprintCallable, Category = "Widget", meta = (DisplayName = "SetIsReadOnly (Editable Text)"))
+	void SetIsReadOnly(UPARAM(DisplayName = "ReadyOnly") bool bReadOnly);
+
+	bool GetIsPassword() const;
+
+	UFUNCTION(BlueprintCallable, Category = "Widget", meta = (DisplayName = "IsPassword"))
 	void SetIsPassword(bool bIsPassword);
 
 	UFUNCTION(BlueprintCallable, Category="Widget")
@@ -161,9 +196,19 @@ public:
 	UFUNCTION(BlueprintCallable, Category = "Widget")
 	bool HasError() const;
 
+	ETextJustify::Type GetJustification() const;
+
 	UFUNCTION(BlueprintSetter)
 	void SetJustification(ETextJustify::Type InJustification);
 
+	/** @return the text overflow policy for this text block. */
+	ETextOverflowPolicy GetTextOverflowPolicy() const;
+
+	/**
+	 * Set the text overflow policy for this text box.
+	 *
+	 * @param InOverflowPolicy the new text overflow policy.
+	 */
 	UFUNCTION(BlueprintSetter)
 	void SetTextOverflowPolicy(ETextOverflowPolicy InOverflowPolicy);
 
