Initilaize the VectorArt built from a Mesh correctly.

diff --git a/Engine/Source/Runtime/UMG/Private/Slate/SlateVectorArtData.cpp b/Engine/Source/Runtime/UMG/Private/Slate/SlateVectorArtData.cpp
--- a/Engine/Source/Runtime/UMG/Private/Slate/SlateVectorArtData.cpp
+++ b/Engine/Source/Runtime/UMG/Private/Slate/SlateVectorArtData.cpp
@@ -148,7 +148,19 @@ UMaterialInstanceDynamic* USlateVectorArtData::ConvertToMaterialInstanceDynamic(
 void USlateVectorArtData::EnsureValidData()
 {
 #if WITH_EDITORONLY_DATA
-	//InitFromStaticMesh(*MeshAsset);
+	if (MeshAsset)
+	{
+		InitFromStaticMesh(*MeshAsset);
+	}
+	else
+	{
+		SourceMaterial = nullptr;
+		VertexData.Reset();
+		IndexData.Reset();
+		Material = nullptr;
+		ExtentMin = FVector2D(FLT_MAX, FLT_MAX);
+		ExtentMax = FVector2D(-FLT_MAX, -FLT_MAX);
+	}
 #endif
 }
 
