Add accessors to MultilineEditableText and MultilineEditableTextBox

diff --git a/Engine/Source/Runtime/UMG/Private/Components/EditableTextBox.cpp b/Engine/Source/Runtime/UMG/Private/Components/EditableTextBox.cpp
--- a/Engine/Source/Runtime/UMG/Private/Components/EditableTextBox.cpp
+++ b/Engine/Source/Runtime/UMG/Private/Components/EditableTextBox.cpp
@@ -139,7 +139,16 @@ FText UEditableTextBox::GetText() const
 
 void UEditableTextBox::SetText(FText InText)
 {
+	// We detect if the Text is internal pointing to the same thing if so, nothing to do.
+	if (Text.IdenticalTo(InText))
+	{
+		return;
+	}
+
 	Text = InText;
+
+	BroadcastFieldValueChanged(FFieldNotificationClassDescriptor::Text);
+
 	if ( MyEditableTextBlock.IsValid() )
 	{
 		MyEditableTextBlock->SetText(Text);
@@ -201,9 +210,9 @@ bool UEditableTextBox::GetIsReadOnly() const
 	return IsReadOnly;
 }
 
-void UEditableTextBox::SetIsReadOnly(bool bReadOnly)
+void UEditableTextBox::SetIsReadOnly(bool bIsReadOnly)
 {
-	IsReadOnly = bReadOnly;
+	IsReadOnly = bIsReadOnly;
 	if ( MyEditableTextBlock.IsValid() )
 	{
 		MyEditableTextBlock->SetIsReadOnly(IsReadOnly);
diff --git a/Engine/Source/Runtime/UMG/Private/Components/MultiLineEditableText.cpp b/Engine/Source/Runtime/UMG/Private/Components/MultiLineEditableText.cpp
--- a/Engine/Source/Runtime/UMG/Private/Components/MultiLineEditableText.cpp
+++ b/Engine/Source/Runtime/UMG/Private/Components/MultiLineEditableText.cpp
@@ -18,6 +18,7 @@ static FTextBlockStyle* DefaultMultiLineEditableTextStyle = nullptr;
 static FTextBlockStyle* EditorMultiLineEditableTextStyle = nullptr;
 #endif 
 
+PRAGMA_DISABLE_DEPRECATION_WARNINGS
 UMultiLineEditableText::UMultiLineEditableText(const FObjectInitializer& ObjectInitializer)
 	: Super(ObjectInitializer)
 {
@@ -146,7 +147,16 @@ FText UMultiLineEditableText::GetText() const
 
 void UMultiLineEditableText::SetText(FText InText)
 {
+	// We detect if the Text is internal pointing to the same thing if so, nothing to do.
+	if (Text.IdenticalTo(InText))
+	{
+		return;
+	}
+
 	Text = InText;
+
+	BroadcastFieldValueChanged(FFieldNotificationClassDescriptor::Text);
+
 	if ( MyMultiLineEditableText.IsValid() )
 	{
 		MyMultiLineEditableText->SetText(Text);
@@ -174,6 +184,11 @@ void UMultiLineEditableText::SetHintText(FText InHintText)
 	}
 }
 
+bool UMultiLineEditableText::GetIsReadOnly() const
+{
+	return bIsReadOnly;
+}
+
 void UMultiLineEditableText::SetIsReadOnly(bool bReadOnly)
 {
 	bIsReadOnly = bReadOnly;
@@ -183,6 +198,7 @@ void UMultiLineEditableText::SetIsReadOnly(bool bReadOnly)
 		MyMultiLineEditableText->SetIsReadOnly(bIsReadOnly);
 	}
 }
+PRAGMA_ENABLE_DEPRECATION_WARNINGS
 
 void UMultiLineEditableText::SetWidgetStyle(const FTextBlockStyle& InWidgetStyle)
 {
diff --git a/Engine/Source/Runtime/UMG/Private/Components/MultiLineEditableTextBox.cpp b/Engine/Source/Runtime/UMG/Private/Components/MultiLineEditableTextBox.cpp
--- a/Engine/Source/Runtime/UMG/Private/Components/MultiLineEditableTextBox.cpp
+++ b/Engine/Source/Runtime/UMG/Private/Components/MultiLineEditableTextBox.cpp
@@ -20,6 +20,7 @@ static FEditableTextBoxStyle* EditorMultiLineEditableTextBoxStyle = nullptr;
 static FTextBlockStyle* EditorMultiLineEditableTextBoxTextStyle = nullptr;
 #endif 
 
+PRAGMA_DISABLE_DEPRECATION_WARNINGS
 UMultiLineEditableTextBox::UMultiLineEditableTextBox(const FObjectInitializer& ObjectInitializer)
 	: Super(ObjectInitializer)
 {
@@ -184,6 +185,11 @@ void UMultiLineEditableTextBox::SetError(FText InError)
 	}
 }
 
+bool UMultiLineEditableTextBox::GetIsReadOnly() const
+{
+	return bIsReadOnly;
+}
+
 void UMultiLineEditableTextBox::SetIsReadOnly(bool bReadOnly)
 {
 	bIsReadOnly = bReadOnly;
@@ -194,6 +200,8 @@ void UMultiLineEditableTextBox::SetIsReadOnly(bool bReadOnly)
 	}
 }
 
+PRAGMA_ENABLE_DEPRECATION_WARNINGS
+
 void UMultiLineEditableTextBox::SetTextStyle(const FTextBlockStyle& InTextStyle)
 {
 	TextStyle = InTextStyle;
diff --git a/Engine/Source/Runtime/UMG/Public/Components/EditableTextBox.h b/Engine/Source/Runtime/UMG/Public/Components/EditableTextBox.h
--- a/Engine/Source/Runtime/UMG/Public/Components/EditableTextBox.h
+++ b/Engine/Source/Runtime/UMG/Public/Components/EditableTextBox.h
@@ -35,7 +35,7 @@ public:
 public:
 	/** The text content for this editable text box widget */
 	UE_DEPRECATED(5.1, "Direct access to Text is deprecated. Please use the getter or setter.")
-	UPROPERTY(EditAnywhere, BlueprintReadWrite, Getter, Setter, BlueprintGetter = "GetText", BlueprintSetter = "SetText", Category = "Content")
+	UPROPERTY(EditAnywhere, BlueprintReadWrite, Getter, Setter, BlueprintGetter = "GetText", BlueprintSetter = "SetText", FieldNotify, Category = "Content")
 	FText Text;
 
 	/** A bindable delegate to allow logic to drive the text of the widget */
@@ -57,7 +57,7 @@ public:
 	UPROPERTY()
 	FGetText HintTextDelegate;
 
-	/** Sets whether this text box can actually be modified interactively by the user */
+	/** Sets the Text Box as Readonly to prevent it from being modified interactively by the user */
 	UE_DEPRECATED(5.1, "Direct access to IsReadOnly is deprecated. Please use the getter or setter.")
 	UPROPERTY(EditAnywhere, BlueprintReadWrite, Getter, Setter, Category = Appearance)
 	bool IsReadOnly;
@@ -180,8 +180,10 @@ public:
 	UFUNCTION(BlueprintCallable, Category="Widget",  meta=(DisplayName="SetError (Text Box)"))
 	void SetError(FText InError);
 
+	/** Return true when this text cannot be modified interactively by the user */
 	bool GetIsReadOnly() const;
 
+	/** Sets the Text as Readonly to prevent it from being modified interactively by the user */
 	UFUNCTION(BlueprintCallable, Category = "Widget", meta = (DisplayName = "SetIsReadOnly (Editable Text)"))
 	void SetIsReadOnly(UPARAM(DisplayName = "ReadyOnly") bool bReadOnly);
 
diff --git a/Engine/Source/Runtime/UMG/Public/Components/MultiLineEditableText.h b/Engine/Source/Runtime/UMG/Public/Components/MultiLineEditableText.h
--- a/Engine/Source/Runtime/UMG/Public/Components/MultiLineEditableText.h
+++ b/Engine/Source/Runtime/UMG/Public/Components/MultiLineEditableText.h
@@ -28,11 +28,13 @@ public:
 
 public:
 	/** The text content for this editable text box widget */
-	UPROPERTY(EditAnywhere, Category=Content, meta=(MultiLine="true"))
+	UE_DEPRECATED(5.1, "Direct access to Text is deprecated. Please use the getter or setter.")
+	UPROPERTY(EditAnywhere, BlueprintReadWrite, Getter, Setter, BlueprintGetter = "GetText", BlueprintSetter = "SetText", FieldNotify, Category = Content, meta = (MultiLine = "true"))
 	FText Text;
 
 	/** Hint text that appears when there is no text in the text box */
-	UPROPERTY(EditAnywhere, Category=Content, meta=(MultiLine="true"))
+	UE_DEPRECATED(5.1, "Direct access to HintText is deprecated. Please use the getter or setter.")
+	UPROPERTY(EditAnywhere, BlueprintReadWrite, Getter, Setter, BlueprintGetter = "GetHintText", BlueprintSetter = "SetHintText", Category = Content, meta = (MultiLine = "true"))
 	FText HintText;
 
 	/** A bindable delegate to allow logic to drive the hint text of the widget */
@@ -44,8 +46,9 @@ public:
 	UPROPERTY(EditAnywhere, BlueprintReadWrite, BlueprintSetter=SetWidgetStyle, Category="Style", meta=(ShowOnlyInnerProperties))
 	FTextBlockStyle WidgetStyle;
 
-	/** Sets whether this text block can be modified interactively by the user */
-	UPROPERTY(EditAnywhere, BlueprintReadOnly, Category="Appearance")
+	/** Sets the Text as Readonly to prevent it from being modified interactively by the user */
+	UE_DEPRECATED(5.1, "Direct access to IsReadOnly is deprecated. Please use the getter or setter.")
+	UPROPERTY(EditAnywhere, BlueprintReadWrite, Getter = GetIsReadOnly, Setter = SetIsReadOnly, BlueprintSetter = "SetIsReadOnly", Category = Appearance)
 	bool bIsReadOnly;
 
 	/** Whether to select all text when the user clicks to give focus on the widget */
@@ -86,22 +89,35 @@ public:
 
 public:
 
-	/**  */
+	/**
+	* Gets the widget text
+	* @return The widget text
+	*/
 	UFUNCTION(BlueprintCallable, Category="Widget", meta=(DisplayName="GetText (Multi-Line Editable Text)"))
 	FText GetText() const;
 
-	/**  */
+	/**
+	* Directly sets the widget text.
+	* @param InText The text to assign to the widget
+	*/
 	UFUNCTION(BlueprintCallable, Category="Widget", meta=(DisplayName="SetText (Multi-Line Editable Text)"))
 	void SetText(FText InText);
 
-	/**  */
+	/** Returns the Hint text that appears when there is no text in the text box */
 	UFUNCTION(BlueprintCallable, Category="Widget", meta=(DisplayName="GetHintText (Multi-Line Editable Text)"))
 	FText GetHintText() const;
 
-	/**  */
+	/** 
+	* Sets the Hint text that appears when there is no text in the text box 
+	* @param InHintText The text that appears when there is no text in the text box 
+	*/
 	UFUNCTION(BlueprintCallable, Category="Widget", meta=(DisplayName="SetHintText (Multi-Line Editable Text)"))
 	void SetHintText(FText InHintText);
 
+	/** Return true when this text cannot be modified interactively by the user */
+	bool GetIsReadOnly() const;
+
+	/** Sets the Text as Readonly to prevent it from being modified interactively by the user */
 	UFUNCTION(BlueprintCallable, Category="Widget", meta=(DisplayName="SetIsReadOnly (Multi-Line Editable Text"))
 	void SetIsReadOnly(bool bReadOnly);
 
diff --git a/Engine/Source/Runtime/UMG/Public/Components/MultiLineEditableTextBox.h b/Engine/Source/Runtime/UMG/Public/Components/MultiLineEditableTextBox.h
--- a/Engine/Source/Runtime/UMG/Public/Components/MultiLineEditableTextBox.h
+++ b/Engine/Source/Runtime/UMG/Public/Components/MultiLineEditableTextBox.h
@@ -30,11 +30,13 @@ public:
 public:
 	
 	/** The text content for this editable text box widget */
-	UPROPERTY(EditAnywhere, Category=Content, meta=(MultiLine="true"))
+	UE_DEPRECATED(5.1, "Direct access to Text is deprecated. Please use the getter or setter.")
+	UPROPERTY(EditAnywhere, BlueprintReadWrite, Getter, Setter, BlueprintGetter = "GetText", BlueprintSetter = "SetText", Category = Content, meta=(MultiLine="true"))
 	FText Text;
 
 	/** Hint text that appears when there is no text in the text box */
-	UPROPERTY(EditAnywhere, Category=Content, meta=(MultiLine="true"))
+	UE_DEPRECATED(5.1, "Direct access to HintText is deprecated. Please use the getter or setter.")
+	UPROPERTY(EditAnywhere, BlueprintReadWrite, Getter, Setter, BlueprintGetter = "GetHintText", BlueprintSetter = "SetHintText", Category = Content, meta = (MultiLine = "true"))
 	FText HintText;
 
 	/** A bindable delegate to allow logic to drive the hint text of the widget */
@@ -50,8 +52,9 @@ public:
 	UPROPERTY(EditAnywhere, BlueprintReadWrite, BlueprintSetter = SetTextStyle, Category="Style", meta=(DisplayName="Text Style"))
 	FTextBlockStyle TextStyle;
 
-	/** Sets whether this text block can be modified interactively by the user */
-	UPROPERTY(EditAnywhere, BlueprintReadOnly, Category="Appearance")
+	/** Sets the Text as Readonly to prevent it from being modified interactively by the user */
+	UE_DEPRECATED(5.1, "Direct access to IsReadOnly is deprecated. Please use the getter or setter.")
+	UPROPERTY(EditAnywhere, BlueprintReadWrite, Getter = GetIsReadOnly, Setter = SetIsReadOnly, Category = Appearance)
 	bool bIsReadOnly;
 
 	/** Whether the context menu can be opened */
@@ -79,27 +82,41 @@ public:
 
 public:
 
-	/**  */
+	/**
+	 * Gets the widget text
+	 * @return The widget text
+	 */
 	UFUNCTION(BlueprintCallable, Category="Widget", meta=(DisplayName="GetText (Multi-Line Text Box)"))
 	FText GetText() const;
 
-	/**  */
+	/**
+	 * Directly sets the widget text.
+	 * Warning: This will wipe any binding created for the Text property!
+	 * @param InText The text to assign to the widget
+	 */
 	UFUNCTION(BlueprintCallable, Category="Widget", meta=(DisplayName="SetText (Multi-Line Text Box)"))
 	void SetText(FText InText);
 
-	/**  */
-	UFUNCTION(BlueprintCallable, Category="Widget", meta=(DisplayName="GetHintText (Multi-Line Text Box)"))
+	/** Returns the Hint text that appears when there is no text in the text box */
+	UFUNCTION(BlueprintCallable, Category = "Widget", meta = (DisplayName = "GetHintText (Multi-Line Text Box)"))
 	FText GetHintText() const;
 
-	/**  */
+	/**
+	* Sets the Hint text that appears when there is no text in the text box
+	* @param InHintText The text that appears when there is no text in the text box
+	*/
 	UFUNCTION(BlueprintCallable, Category="Widget", meta=(DisplayName="SetHintText (Multi-Line Text Box)"))
 	void SetHintText(FText InHintText);
 
 	UFUNCTION(BlueprintCallable, Category="Widget", meta=(DisplayName="SetError (Multi-Line Text Box)"))
 	void SetError(FText InError);
 
-	UFUNCTION(BlueprintCallable, Category="Widget", meta=(DisplayName="SetIsReadOnly (Multi-Line Text Box)"))
-	void SetIsReadOnly(bool bReadOnly);
+	/** Return true when this text cannot be modified interactively by the user */
+	bool GetIsReadOnly() const;
+
+	/** Sets the Text as Readonly to prevent it from being modified interactively by the user */
+	UFUNCTION(BlueprintCallable, Category = "Widget", meta = (DisplayName = "SetIsReadOnly (Multi-Line Text Box)"))
+	void SetIsReadOnly(UPARAM(DisplayName = "ReadyOnly") bool bReadOnly);
 
 	UFUNCTION(BlueprintSetter)
 	void SetTextStyle(const FTextBlockStyle& InTextStyle);
