Add more accessor to EditableTextBox and EditableText so that both API match

diff --git a/Engine/Source/Runtime/UMG/Private/Components/EditableText.cpp b/Engine/Source/Runtime/UMG/Private/Components/EditableText.cpp
--- a/Engine/Source/Runtime/UMG/Private/Components/EditableText.cpp
+++ b/Engine/Source/Runtime/UMG/Private/Components/EditableText.cpp
@@ -202,6 +202,67 @@ void UEditableText::SetMinimumDesiredWidth(float InMinDesiredWidth)
 	}
 }
 
+void UEditableText::SetIsCaretMovedWhenGainFocus(bool bIsCaretMovedWhenGainFocus)
+{
+	IsCaretMovedWhenGainFocus = bIsCaretMovedWhenGainFocus;
+	if (MyEditableText.IsValid())
+	{
+		MyEditableText->SetIsCaretMovedWhenGainFocus(bIsCaretMovedWhenGainFocus);
+	}
+}
+
+bool UEditableText::GetIsCaretMovedWhenGainFocus() const
+{
+	return IsCaretMovedWhenGainFocus;
+}
+
+void UEditableText::SetSelectAllTextWhenFocused(bool bSelectAllTextWhenFocused)
+{
+	SelectAllTextWhenFocused = bSelectAllTextWhenFocused;
+	if (MyEditableText.IsValid())
+	{
+		MyEditableText->SetSelectAllTextWhenFocused(bSelectAllTextWhenFocused);
+	}
+}
+
+bool UEditableText::GetSelectAllTextWhenFocused() const
+{
+	return SelectAllTextWhenFocused;
+}
+
+void UEditableText::SetRevertTextOnEscape(bool bRevertTextOnEscape)
+{
+	RevertTextOnEscape = bRevertTextOnEscape;
+	if (MyEditableText.IsValid())
+	{
+		MyEditableText->SetRevertTextOnEscape(bRevertTextOnEscape);
+	}
+}
+
+bool UEditableText::GetRevertTextOnEscape() const
+{
+	return RevertTextOnEscape;
+}
+
+bool UEditableText::GetClearKeyboardFocusOnCommit() const
+{
+	return ClearKeyboardFocusOnCommit;
+}
+
+void UEditableText::SetSelectAllTextOnCommit(bool bSelectAllTextOnCommit)
+{
+	SelectAllTextOnCommit = bSelectAllTextOnCommit;
+	if (MyEditableText.IsValid())
+	{
+		MyEditableText->SetSelectAllTextOnCommit(bSelectAllTextOnCommit);
+	}
+}
+
+bool UEditableText::GetSelectAllTextOnCommit() const
+{
+	return SelectAllTextOnCommit;
+}
+
 bool UEditableText::GetIsReadOnly() const
 {
 	return IsReadOnly;
@@ -251,13 +312,15 @@ void UEditableText::SetTextOverflowPolicy(ETextOverflowPolicy InOverflowPolicy)
 	}
 }
 
-PRAGMA_ENABLE_DEPRECATION_WARNINGS
-
-void UEditableText::SetClearKeyboardFocusOnCommit(bool bInClearKeyboardFocusOnCommit)
+void UEditableText::SetClearKeyboardFocusOnCommit(bool bClearKeyboardFocusOnCommit)
 {
-	ClearKeyboardFocusOnCommit = bInClearKeyboardFocusOnCommit;
-	MyEditableText->SetClearKeyboardFocusOnCommit(ClearKeyboardFocusOnCommit);
+	ClearKeyboardFocusOnCommit = bClearKeyboardFocusOnCommit;
+	if (MyEditableText.IsValid())
+	{
+		MyEditableText->SetClearKeyboardFocusOnCommit(ClearKeyboardFocusOnCommit);
+	}
 }
+PRAGMA_ENABLE_DEPRECATION_WARNINGS
 
 void UEditableText::SetKeyboardType(EVirtualKeyboardType::Type Type)
 {
diff --git a/Engine/Source/Runtime/UMG/Private/Components/EditableTextBox.cpp b/Engine/Source/Runtime/UMG/Private/Components/EditableTextBox.cpp
--- a/Engine/Source/Runtime/UMG/Private/Components/EditableTextBox.cpp
+++ b/Engine/Source/Runtime/UMG/Private/Components/EditableTextBox.cpp
@@ -188,6 +188,77 @@ void UEditableTextBox::SetMinDesiredWidth(float InMinDesiredWidth)
 	}
 }
 
+void UEditableTextBox::SetIsCaretMovedWhenGainFocus(bool bIsCaretMovedWhenGainFocus)
+{
+	IsCaretMovedWhenGainFocus = bIsCaretMovedWhenGainFocus;
+	if (MyEditableTextBlock.IsValid())
+	{
+		MyEditableTextBlock->SetIsCaretMovedWhenGainFocus(bIsCaretMovedWhenGainFocus);
+	}
+}
+
+bool UEditableTextBox::GetIsCaretMovedWhenGainFocus() const
+{
+	return IsCaretMovedWhenGainFocus;
+}
+
+void UEditableTextBox::SetSelectAllTextWhenFocused(bool bSelectAllTextWhenFocused)
+{
+	SelectAllTextWhenFocused = bSelectAllTextWhenFocused;
+	if (MyEditableTextBlock.IsValid())
+	{
+		MyEditableTextBlock->SetSelectAllTextWhenFocused(bSelectAllTextWhenFocused);
+	}
+}
+
+bool UEditableTextBox::GetSelectAllTextWhenFocused() const
+{
+	return SelectAllTextWhenFocused;
+}
+
+void UEditableTextBox::SetRevertTextOnEscape(bool bRevertTextOnEscape)
+{
+	RevertTextOnEscape = bRevertTextOnEscape;
+	if (MyEditableTextBlock.IsValid())
+	{
+		MyEditableTextBlock->SetRevertTextOnEscape(bRevertTextOnEscape);
+	}
+}
+
+bool UEditableTextBox::GetRevertTextOnEscape() const
+{
+	return RevertTextOnEscape;
+}
+
+void UEditableTextBox::SetClearKeyboardFocusOnCommit(bool bClearKeyboardFocusOnCommit)
+{
+	ClearKeyboardFocusOnCommit = bClearKeyboardFocusOnCommit;
+	if (MyEditableTextBlock.IsValid())
+	{
+		MyEditableTextBlock->SetClearKeyboardFocusOnCommit(bClearKeyboardFocusOnCommit);
+	}
+}
+
+bool UEditableTextBox::GetClearKeyboardFocusOnCommit() const
+{
+	return ClearKeyboardFocusOnCommit;
+}
+
+void UEditableTextBox::SetSelectAllTextOnCommit(bool bSelectAllTextOnCommit)
+{
+	SelectAllTextOnCommit = bSelectAllTextOnCommit;
+	if (MyEditableTextBlock.IsValid())
+	{
+		MyEditableTextBlock->SetSelectAllTextOnCommit(bSelectAllTextOnCommit);
+	}
+
+}
+
+bool UEditableTextBox::GetSelectAllTextOnCommit() const
+{
+	return SelectAllTextOnCommit;
+}
+
 void UEditableTextBox::SetForegroundColor(FLinearColor color)
 {
 	WidgetStyle.ForegroundColor = color;
diff --git a/Engine/Source/Runtime/UMG/Private/Components/MultiLineEditableText.cpp b/Engine/Source/Runtime/UMG/Private/Components/MultiLineEditableText.cpp
--- a/Engine/Source/Runtime/UMG/Private/Components/MultiLineEditableText.cpp
+++ b/Engine/Source/Runtime/UMG/Private/Components/MultiLineEditableText.cpp
@@ -184,6 +184,66 @@ void UMultiLineEditableText::SetHintText(FText InHintText)
 	}
 }
 
+void UMultiLineEditableText::SetSelectAllTextWhenFocused(bool bSelectAllTextWhenFocused)
+{
+	SelectAllTextWhenFocused = bSelectAllTextWhenFocused;
+
+	if (MyMultiLineEditableText.IsValid())
+	{
+		MyMultiLineEditableText->SetSelectAllTextWhenFocused(bSelectAllTextWhenFocused);
+	}
+}
+
+bool UMultiLineEditableText::GetSelectAllTextWhenFocused() const
+{
+	return SelectAllTextWhenFocused;
+}
+
+void UMultiLineEditableText::SetClearTextSelectionOnFocusLoss(bool bClearTextSelectionOnFocusLoss)
+{
+	ClearTextSelectionOnFocusLoss = bClearTextSelectionOnFocusLoss;
+
+	if (MyMultiLineEditableText.IsValid())
+	{
+		MyMultiLineEditableText->SetClearTextSelectionOnFocusLoss(bClearTextSelectionOnFocusLoss);
+	}
+}
+
+bool UMultiLineEditableText::GetClearTextSelectionOnFocusLoss() const
+{
+	return ClearTextSelectionOnFocusLoss;
+}
+
+void UMultiLineEditableText::SetRevertTextOnEscape(bool bRevertTextOnEscape)
+{
+	RevertTextOnEscape = bRevertTextOnEscape;
+
+	if (MyMultiLineEditableText.IsValid())
+	{
+		MyMultiLineEditableText->SetRevertTextOnEscape(bRevertTextOnEscape);
+	}
+}
+
+bool UMultiLineEditableText::GetRevertTextOnEscape() const
+{
+	return RevertTextOnEscape;
+}
+
+void UMultiLineEditableText::SetClearKeyboardFocusOnCommit(bool bClearKeyboardFocusOnCommit)
+{
+	ClearKeyboardFocusOnCommit = bClearKeyboardFocusOnCommit;
+
+	if (MyMultiLineEditableText.IsValid())
+	{
+		MyMultiLineEditableText->SetClearKeyboardFocusOnCommit(bClearKeyboardFocusOnCommit);
+	}
+}
+
+bool UMultiLineEditableText::GetClearKeyboardFocusOnCommit() const
+{
+	return ClearKeyboardFocusOnCommit;
+}
+
 bool UMultiLineEditableText::GetIsReadOnly() const
 {
 	return bIsReadOnly;
diff --git a/Engine/Source/Runtime/UMG/Public/Components/EditableText.h b/Engine/Source/Runtime/UMG/Public/Components/EditableText.h
--- a/Engine/Source/Runtime/UMG/Public/Components/EditableText.h
+++ b/Engine/Source/Runtime/UMG/Public/Components/EditableText.h
@@ -70,24 +70,29 @@ public:
 	UPROPERTY(EditAnywhere, BlueprintReadWrite, Getter, Setter, BlueprintSetter = "SetMinimumDesiredWidth", Category = "Appearance")
 	float MinimumDesiredWidth;
 
-	/** Workaround as we lose focus when the auto completion closes. */
-	UPROPERTY(EditAnywhere, Category=Behavior, AdvancedDisplay)
+	/** When set to true the caret is moved when gaining focus */
+	UE_DEPRECATED(5.1, "Direct access to IsCaretMovedWhenGainFocus is deprecated. Please use the getter or setter.")
+	UPROPERTY(EditAnywhere, BlueprintReadWrite, Getter, Setter, Category = Behavior, AdvancedDisplay)
 	bool IsCaretMovedWhenGainFocus;
 
 	/** Whether to select all text when the user clicks to give focus on the widget */
-	UPROPERTY(EditAnywhere, Category=Behavior, AdvancedDisplay)
+	UE_DEPRECATED(5.1, "Direct access to SelectAllTextWhenFocused is deprecated. Please use the getter or setter.")
+	UPROPERTY(EditAnywhere, BlueprintReadWrite, Getter, Setter, Category = Behavior, AdvancedDisplay)
 	bool SelectAllTextWhenFocused;
 
 	/** Whether to allow the user to back out of changes when they press the escape key */
-	UPROPERTY(EditAnywhere, Category=Behavior, AdvancedDisplay)
+	UE_DEPRECATED(5.1, "Direct access to RevertTextOnEscape is deprecated. Please use the getter or setter.")
+	UPROPERTY(EditAnywhere, BlueprintReadWrite, Getter, Setter, Category = Behavior, AdvancedDisplay)
 	bool RevertTextOnEscape;
 
 	/** Whether to clear keyboard focus when pressing enter to commit changes */
-	UPROPERTY(EditAnywhere, Category=Behavior, AdvancedDisplay)
+	UE_DEPRECATED(5.1, "Direct access to ClearKeyboardFocusOnCommit is deprecated. Please use the getter or setter.")
+	UPROPERTY(EditAnywhere, BlueprintReadWrite, Getter, Setter, Category = Behavior, AdvancedDisplay)
 	bool ClearKeyboardFocusOnCommit;
 
 	/** Whether to select all text when pressing enter to commit changes */
-	UPROPERTY(EditAnywhere, Category=Behavior, AdvancedDisplay)
+	UE_DEPRECATED(5.1, "Direct access to SelectAllTextOnCommit is deprecated. Please use the getter or setter.")
+	UPROPERTY(EditAnywhere, BlueprintReadWrite, Getter, Setter, Category = Behavior, AdvancedDisplay)
 	bool SelectAllTextOnCommit;
 
 	/** Whether the context menu can be opened */
@@ -171,6 +176,33 @@ public:
 	UFUNCTION(BlueprintCallable, Category = "Appearance")
 	void SetMinimumDesiredWidth(float InMinDesiredWidth);
 
+	/** When set to true the caret is moved when gaining focus */
+	void SetIsCaretMovedWhenGainFocus(bool bIsCaretMovedWhenGainFocus);
+
+	/** Return true when the caret is moved when gaining focus */
+	bool GetIsCaretMovedWhenGainFocus() const;
+
+	/** Set to true to select all text when the user clicks to give focus on the widget */
+	void SetSelectAllTextWhenFocused(bool bSelectAllTextWhenFocused);
+
+	/** Whether to select all text when the user clicks to give focus on the widget */
+	bool GetSelectAllTextWhenFocused() const;
+
+	/** Set to true to allow the user to back out of changes when they press the escape key */
+	void SetRevertTextOnEscape(bool bRevertTextOnEscape);
+
+	/** Whether to allow the user to back out of changes when they press the escape key  */
+	bool GetRevertTextOnEscape() const;
+
+	/** Whether to clear keyboard focus when pressing enter to commit changes */
+	bool GetClearKeyboardFocusOnCommit() const;
+
+	/** Set to true to select all text when pressing enter to commit changes */
+	void SetSelectAllTextOnCommit(bool bSelectAllTextOnCommit);
+
+	/** Whether to select all text when pressing enter to commit changes */
+	bool GetSelectAllTextOnCommit() const;
+
 	bool GetIsReadOnly() const;
 
 	UFUNCTION(BlueprintCallable, Category = "Widget", meta=(DisplayName="SetIsReadOnly (Editable Text)"))
@@ -188,6 +220,7 @@ public:
 	UFUNCTION(BlueprintSetter)
 	void SetTextOverflowPolicy(ETextOverflowPolicy InOverflowPolicy);
 
+	/** Set to true to clear keyboard focus when pressing enter to commit changes */
 	void SetClearKeyboardFocusOnCommit(bool bInClearKeyboardFocusOnCommit);
 
 	void SetKeyboardType(EVirtualKeyboardType::Type Type);
diff --git a/Engine/Source/Runtime/UMG/Public/Components/EditableTextBox.h b/Engine/Source/Runtime/UMG/Public/Components/EditableTextBox.h
--- a/Engine/Source/Runtime/UMG/Public/Components/EditableTextBox.h
+++ b/Engine/Source/Runtime/UMG/Public/Components/EditableTextBox.h
@@ -73,23 +73,28 @@ public:
 	float MinimumDesiredWidth;
 
 	/** Workaround as we lose focus when the auto completion closes. */
-	UPROPERTY(EditAnywhere, BlueprintReadOnly, Category=Behavior, AdvancedDisplay)
+	UE_DEPRECATED(5.1, "Direct access to IsCaretMovedWhenGainFocus is deprecated. Please use the getter or setter.")
+	UPROPERTY(EditAnywhere, BlueprintReadWrite, Getter, Setter, Category=Behavior, AdvancedDisplay)
 	bool IsCaretMovedWhenGainFocus;
 
 	/** Whether to select all text when the user clicks to give focus on the widget */
-	UPROPERTY(EditAnywhere, BlueprintReadOnly, Category=Behavior, AdvancedDisplay)
+	UE_DEPRECATED(5.1, "Direct access to SelectAllTextWhenFocused is deprecated. Please use the getter or setter.")
+	UPROPERTY(EditAnywhere, BlueprintReadWrite, Getter, Setter, Category=Behavior, AdvancedDisplay)
 	bool SelectAllTextWhenFocused;
 
 	/** Whether to allow the user to back out of changes when they press the escape key */
-	UPROPERTY(EditAnywhere, BlueprintReadOnly, Category=Behavior, AdvancedDisplay)
+	UE_DEPRECATED(5.1, "Direct access to RevertTextOnEscape is deprecated. Please use the getter or setter.")
+	UPROPERTY(EditAnywhere, BlueprintReadWrite, Getter, Setter, Category=Behavior, AdvancedDisplay)
 	bool RevertTextOnEscape;
 
 	/** Whether to clear keyboard focus when pressing enter to commit changes */
-	UPROPERTY(EditAnywhere, BlueprintReadOnly, Category=Behavior, AdvancedDisplay)
+	UE_DEPRECATED(5.1, "Direct access to ClearKeyboardFocusOnCommit is deprecated. Please use the getter or setter.")
+	UPROPERTY(EditAnywhere, BlueprintReadWrite, Getter, Setter, Category=Behavior, AdvancedDisplay)
 	bool ClearKeyboardFocusOnCommit;
 
 	/** Whether to select all text when pressing enter to commit changes */
-	UPROPERTY(EditAnywhere, BlueprintReadOnly, Category=Behavior, AdvancedDisplay)
+	UE_DEPRECATED(5.1, "Direct access to SelectAllTextOnCommit is deprecated. Please use the getter or setter.")
+	UPROPERTY(EditAnywhere, BlueprintReadWrite, Getter, Setter, Category=Behavior, AdvancedDisplay)
 	bool SelectAllTextOnCommit;
 
 	/** Whether the context menu can be opened */
@@ -176,6 +181,36 @@ public:
 	*  @param InMinDesiredWidth new minimum desired width
 	*/
 	void SetMinDesiredWidth(float InMinDesiredWidth);
+	
+	/** When set to true the caret is moved when gaining focus */
+	void SetIsCaretMovedWhenGainFocus(bool bIsCaretMovedWhenGainFocus);
+
+	/** Return true when the caret is moved when gaining focus */
+	bool GetIsCaretMovedWhenGainFocus() const;
+
+	/** Set to true to select all text when the user clicks to give focus on the widget */
+	void SetSelectAllTextWhenFocused(bool bSelectAllTextWhenFocused);
+
+	/** Whether to select all text when the user clicks to give focus on the widget */
+	bool GetSelectAllTextWhenFocused() const;
+
+	/** Set to true to allow the user to back out of changes when they press the escape key */
+	void SetRevertTextOnEscape(bool bRevertTextOnEscape);
+
+	/** Whether to allow the user to back out of changes when they press the escape key  */
+	bool GetRevertTextOnEscape() const;
+
+	/** Set to true to clear keyboard focus when pressing enter to commit changes */
+	void SetClearKeyboardFocusOnCommit(bool bClearKeyboardFocusOnCommit);
+
+	/** Whether to clear keyboard focus when pressing enter to commit changes */
+	bool GetClearKeyboardFocusOnCommit() const;
+
+	/** Set to true to select all text when pressing enter to commit changes */
+	void SetSelectAllTextOnCommit(bool bSelectAllTextOnCommit);
+
+	/** Whether to select all text when pressing enter to commit changes */
+	bool GetSelectAllTextOnCommit() const;
 
 	UFUNCTION(BlueprintCallable, Category="Widget",  meta=(DisplayName="SetError (Text Box)"))
 	void SetError(FText InError);
diff --git a/Engine/Source/Runtime/UMG/Public/Components/MultiLineEditableText.h b/Engine/Source/Runtime/UMG/Public/Components/MultiLineEditableText.h
--- a/Engine/Source/Runtime/UMG/Public/Components/MultiLineEditableText.h
+++ b/Engine/Source/Runtime/UMG/Public/Components/MultiLineEditableText.h
@@ -52,19 +52,23 @@ public:
 	bool bIsReadOnly;
 
 	/** Whether to select all text when the user clicks to give focus on the widget */
-	UPROPERTY(EditAnywhere, Category=Behavior, AdvancedDisplay)
+	UE_DEPRECATED(5.1, "Direct access to SelectAllTextWhenFocused is deprecated. Please use the getter or setter.")
+	UPROPERTY(EditAnywhere, BlueprintReadWrite, Getter, Setter, Category = Behavior, AdvancedDisplay)
 	bool SelectAllTextWhenFocused;
 
 	/** Whether to clear text selection when focus is lost */
-	UPROPERTY(EditAnywhere, Category=Behavior, AdvancedDisplay)
+	UE_DEPRECATED(5.1, "Direct access to ClearTextSelectionOnFocusLoss is deprecated. Please use the getter or setter.")
+	UPROPERTY(EditAnywhere, Category = Behavior, AdvancedDisplay)
 	bool ClearTextSelectionOnFocusLoss;
 
 	/** Whether to allow the user to back out of changes when they press the escape key */
-	UPROPERTY(EditAnywhere, Category=Behavior, AdvancedDisplay)
+	UE_DEPRECATED(5.1, "Direct access to RevertTextOnEscape is deprecated. Please use the getter or setter.")
+	UPROPERTY(EditAnywhere, BlueprintReadWrite, Getter, Setter, Category = Behavior, AdvancedDisplay)
 	bool RevertTextOnEscape;
 
 	/** Whether to clear keyboard focus when pressing enter to commit changes */
-	UPROPERTY(EditAnywhere, Category=Behavior, AdvancedDisplay)
+	UE_DEPRECATED(5.1, "Direct access to ClearKeyboardFocusOnCommit is deprecated. Please use the getter or setter.")
+	UPROPERTY(EditAnywhere, BlueprintReadWrite, Getter, Setter, Category = Behavior, AdvancedDisplay)
 	bool ClearKeyboardFocusOnCommit;
 
 	/** Whether the context menu can be opened */
@@ -114,6 +118,30 @@ public:
 	UFUNCTION(BlueprintCallable, Category="Widget", meta=(DisplayName="SetHintText (Multi-Line Editable Text)"))
 	void SetHintText(FText InHintText);
 
+	/** Set to true to select all text when the user clicks to give focus on the widget */
+	void SetSelectAllTextWhenFocused(bool bSelectAllTextWhenFocused);
+
+	/** Whether to select all text when the user clicks to give focus on the widget */
+	bool GetSelectAllTextWhenFocused() const;
+
+	/** Set to true to clear text selection when focus is lost */
+	void SetClearTextSelectionOnFocusLoss(bool bClearTextSelectionOnFocusLoss);
+
+	/** Whether to clear text selection when focus is lost */
+	bool GetClearTextSelectionOnFocusLoss() const;
+
+	/** Set to true to allow the user to back out of changes when they press the escape key */
+	void SetRevertTextOnEscape(bool bRevertTextOnEscape);
+
+	/** Whether to allow the user to back out of changes when they press the escape key  */
+	bool GetRevertTextOnEscape() const;
+
+	/** Set to true to clear keyboard focus when pressing enter to commit changes */
+	void SetClearKeyboardFocusOnCommit(bool bClearKeyboardFocusOnCommit);
+
+	/** Whether to clear keyboard focus when pressing enter to commit changes */
+	bool GetClearKeyboardFocusOnCommit() const;	
+
 	/** Return true when this text cannot be modified interactively by the user */
 	bool GetIsReadOnly() const;
 
