UMG: Update ButtonSlot, CanvasPanelSlot and GridSlot with accessors.

diff --git a/Engine/Source/Editor/UMGEditor/Private/Customizations/CanvasSlotCustomization.cpp b/Engine/Source/Editor/UMGEditor/Private/Customizations/CanvasSlotCustomization.cpp
--- a/Engine/Source/Editor/UMGEditor/Private/Customizations/CanvasSlotCustomization.cpp
+++ b/Engine/Source/Editor/UMGEditor/Private/Customizations/CanvasSlotCustomization.cpp
@@ -223,7 +223,11 @@ void FCanvasSlotCustomization::FillOutChildren(TSharedRef<IPropertyHandle> Prope
 
 void FCanvasSlotCustomization::CustomizeLayoutData(TSharedRef<IPropertyHandle> PropertyHandle, IDetailChildrenBuilder& ChildBuilder, IPropertyTypeCustomizationUtils& CustomizationUtils)
 {
-	TSharedPtr<IPropertyHandle> LayoutData = PropertyHandle->GetChildHandle(GET_MEMBER_NAME_CHECKED(UCanvasPanelSlot, LayoutData));
+PRAGMA_DISABLE_DEPRECATION_WARNINGS
+	const FName LayoutDataName = GET_MEMBER_NAME_CHECKED(UCanvasPanelSlot, LayoutData);
+PRAGMA_ENABLE_DEPRECATION_WARNINGS
+
+	TSharedPtr<IPropertyHandle> LayoutData = PropertyHandle->GetChildHandle(LayoutDataName);
 	if ( LayoutData.IsValid() )
 	{
 		LayoutData->MarkHiddenByCustomization();
diff --git a/Engine/Source/Editor/UMGEditor/Private/Designer/STransformHandle.cpp b/Engine/Source/Editor/UMGEditor/Private/Designer/STransformHandle.cpp
--- a/Engine/Source/Editor/UMGEditor/Private/Designer/STransformHandle.cpp
+++ b/Engine/Source/Editor/UMGEditor/Private/Designer/STransformHandle.cpp
@@ -163,7 +163,7 @@ void STransformHandle::Resize(UCanvasPanelSlot* Slot, const FVector2D& Direction
 	}
 
 	FMargin Offsets = StartingOffsets;
-	const FAnchorData& LayoutData = Slot->LayoutData;
+	const FAnchorData LayoutData = Slot->GetLayout();
 
 	FVector2D Movement = Amount * Direction;
 	FVector2D PositionMovement = Movement * ( FVector2D(1.0f, 1.0f) - LayoutData.Alignment );
diff --git a/Engine/Source/Editor/UMGEditor/Private/Extensions/CanvasSlotExtension.cpp b/Engine/Source/Editor/UMGEditor/Private/Extensions/CanvasSlotExtension.cpp
--- a/Engine/Source/Editor/UMGEditor/Private/Extensions/CanvasSlotExtension.cpp
+++ b/Engine/Source/Editor/UMGEditor/Private/Extensions/CanvasSlotExtension.cpp
@@ -225,16 +225,17 @@ EVisibility FCanvasSlotExtension::GetAnchorVisibility(EAnchorWidget AnchorType)
 		{
 			if ( UCanvasPanelSlot* PreviewCanvasSlot = Cast<UCanvasPanelSlot>(PreviewWidget->Slot) )
 			{
+				FAnchorData PreviewLayout = PreviewCanvasSlot->GetLayout();
 				switch ( AnchorType )
 				{
 				case EAnchorWidget::Center:
-					return PreviewCanvasSlot->LayoutData.Anchors.Minimum == PreviewCanvasSlot->LayoutData.Anchors.Maximum ? EVisibility::Visible : EVisibility::Collapsed;
+					return PreviewLayout.Anchors.Minimum == PreviewLayout.Anchors.Maximum ? EVisibility::Visible : EVisibility::Collapsed;
 				case EAnchorWidget::Left:
 				case EAnchorWidget::Right:
-					return PreviewCanvasSlot->LayoutData.Anchors.Minimum.Y == PreviewCanvasSlot->LayoutData.Anchors.Maximum.Y ? EVisibility::Visible : EVisibility::Collapsed;
+					return PreviewLayout.Anchors.Minimum.Y == PreviewLayout.Anchors.Maximum.Y ? EVisibility::Visible : EVisibility::Collapsed;
 				case EAnchorWidget::Top:
 				case EAnchorWidget::Bottom:
-					return PreviewCanvasSlot->LayoutData.Anchors.Minimum.X == PreviewCanvasSlot->LayoutData.Anchors.Maximum.X ? EVisibility::Visible : EVisibility::Collapsed;
+					return PreviewLayout.Anchors.Minimum.X == PreviewLayout.Anchors.Maximum.X ? EVisibility::Visible : EVisibility::Collapsed;
 				}
 
 				return EVisibility::Visible;
@@ -254,8 +255,8 @@ FVector2D FCanvasSlotExtension::GetAnchorAlignment(EAnchorWidget AnchorType) con
 		{
 			if ( UCanvasPanelSlot* PreviewCanvasSlot = Cast<UCanvasPanelSlot>(PreviewWidget->Slot) )
 			{
-				FVector2D Minimum = PreviewCanvasSlot->LayoutData.Anchors.Minimum;
-				FVector2D Maximum = PreviewCanvasSlot->LayoutData.Anchors.Maximum;
+				FVector2D Minimum = PreviewCanvasSlot->GetLayout().Anchors.Minimum;
+				FVector2D Maximum = PreviewCanvasSlot->GetLayout().Anchors.Maximum;
 
 				switch ( AnchorType )
 				{
@@ -337,7 +338,7 @@ FReply FCanvasSlotExtension::HandleAnchorBeginDrag(const FGeometry& Geometry, co
 	MouseDownPosition = Event.GetScreenSpacePosition();
 
 	UCanvasPanelSlot* PreviewCanvasSlot = Cast<UCanvasPanelSlot>(SelectionCache[0].GetPreview()->Slot);
-	BeginAnchors = PreviewCanvasSlot->LayoutData.Anchors;
+	BeginAnchors = PreviewCanvasSlot->GetLayout().Anchors;
 
 	Designer->PushDesignerMessage(LOCTEXT("CenterAnchorControls", "Hold [Ctrl] to update widget position"));
 
@@ -389,7 +390,7 @@ FReply FCanvasSlotExtension::HandleAnchorDragging(const FGeometry& Geometry, con
 
 					FVector2D AnchorDelta = LocalPositionDelta / CanvasGeometry.GetLocalSize();
 
-					const FAnchorData OldLayoutData = PreviewCanvasSlot->LayoutData;
+					const FAnchorData OldLayoutData = PreviewCanvasSlot->GetLayout();
 					FAnchorData LayoutData = OldLayoutData;
 					
 					switch ( AnchorType )
@@ -475,11 +476,11 @@ FReply FCanvasSlotExtension::HandleAnchorDragging(const FGeometry& Geometry, con
 					// result.
 					{
 						PreviewCanvasSlot->SaveBaseLayout();
-						PreviewCanvasSlot->LayoutData = LayoutData;
+						PreviewCanvasSlot->SetLayout(LayoutData);
 						PreviewCanvasSlot->RebaseLayout();
 
-						LayoutData = PreviewCanvasSlot->LayoutData;
-						PreviewCanvasSlot->LayoutData = OldLayoutData;
+						LayoutData = PreviewCanvasSlot->GetLayout();
+						PreviewCanvasSlot->SetLayout(OldLayoutData);
 					}
 
 					// If control is pressed reset all positional offset information
@@ -569,7 +570,7 @@ void FCanvasSlotExtension::PaintDragPercentages(const TSet< FWidgetReference >&
 
 				FVector2D CanvasSize = CanvasGeometry.GetLocalSize();
 
-				const FAnchorData LayoutData = PreviewCanvasSlot->LayoutData;
+				const FAnchorData LayoutData = PreviewCanvasSlot->GetLayout();
 				const FVector2D AnchorMin = LayoutData.Anchors.Minimum;
 				const FVector2D AnchorMax = LayoutData.Anchors.Maximum;
 
diff --git a/Engine/Source/Editor/UMGEditor/Private/Extensions/GridSlotExtension.cpp b/Engine/Source/Editor/UMGEditor/Private/Extensions/GridSlotExtension.cpp
--- a/Engine/Source/Editor/UMGEditor/Private/Extensions/GridSlotExtension.cpp
+++ b/Engine/Source/Editor/UMGEditor/Private/Extensions/GridSlotExtension.cpp
@@ -106,13 +106,13 @@ FReply FGridSlotExtension::HandleShiftColumn(int32 ShiftAmount)
 void FGridSlotExtension::ShiftRow(UWidget* Widget, int32 ShiftAmount)
 {
 	UGridSlot* Slot = Cast<UGridSlot>(Widget->Slot);
-	Slot->SetRow(FMath::Max(Slot->Row + ShiftAmount, 0));
+	Slot->SetRow(FMath::Max(Slot->GetRow() + ShiftAmount, 0));
 }
 
 void FGridSlotExtension::ShiftColumn(UWidget* Widget, int32 ShiftAmount)
 {
 	UGridSlot* Slot = Cast<UGridSlot>(Widget->Slot);
-	Slot->SetColumn(FMath::Max(Slot->Column + ShiftAmount, 0));
+	Slot->SetColumn(FMath::Max(Slot->GetColumn() + ShiftAmount, 0));
 }
 
 #undef LOCTEXT_NAMESPACE
diff --git a/Engine/Source/Runtime/UMG/Private/Components/ButtonSlot.cpp b/Engine/Source/Runtime/UMG/Private/Components/ButtonSlot.cpp
--- a/Engine/Source/Runtime/UMG/Private/Components/ButtonSlot.cpp
+++ b/Engine/Source/Runtime/UMG/Private/Components/ButtonSlot.cpp
@@ -11,10 +11,12 @@
 UButtonSlot::UButtonSlot(const FObjectInitializer& ObjectInitializer)
 	: Super(ObjectInitializer)
 {
-	Padding = FMargin(4, 2);
+PRAGMA_DISABLE_DEPRECATION_WARNINGS
+	Padding = FMargin(4.f, 2.f);
 
 	HorizontalAlignment = HAlign_Center;
 	VerticalAlignment = VAlign_Center;
+PRAGMA_ENABLE_DEPRECATION_WARNINGS
 }
 
 void UButtonSlot::ReleaseSlateResources(bool bReleaseChildren)
@@ -28,11 +30,19 @@ void UButtonSlot::BuildSlot(TSharedRef<SButton> InButton)
 {
 	Button = InButton;
 
-	Button.Pin()->SetContentPadding(Padding);
-	Button.Pin()->SetHAlign(HorizontalAlignment);
-	Button.Pin()->SetVAlign(VerticalAlignment);
+PRAGMA_DISABLE_DEPRECATION_WARNINGS
+	InButton->SetContentPadding(Padding);
+	InButton->SetHAlign(HorizontalAlignment);
+	InButton->SetVAlign(VerticalAlignment);
+PRAGMA_ENABLE_DEPRECATION_WARNINGS
 
-	Button.Pin()->SetContent(Content ? Content->TakeWidget() : SNullWidget::NullWidget);
+	InButton->SetContent(Content ? Content->TakeWidget() : SNullWidget::NullWidget);
+}
+
+PRAGMA_DISABLE_DEPRECATION_WARNINGS
+FMargin UButtonSlot::GetPadding() const
+{
+	return Padding;
 }
 
 void UButtonSlot::SetPadding(FMargin InPadding)
@@ -44,6 +54,11 @@ void UButtonSlot::SetPadding(FMargin InPadding)
 	}
 }
 
+EHorizontalAlignment UButtonSlot::GetHorizontalAlignment() const
+{
+	return HorizontalAlignment;
+}
+
 void UButtonSlot::SetHorizontalAlignment(EHorizontalAlignment InHorizontalAlignment)
 {
 	HorizontalAlignment = InHorizontalAlignment;
@@ -53,6 +68,11 @@ void UButtonSlot::SetHorizontalAlignment(EHorizontalAlignment InHorizontalAlignm
 	}
 }
 
+EVerticalAlignment UButtonSlot::GetVerticalAlignment() const
+{
+	return VerticalAlignment;
+}
+
 void UButtonSlot::SetVerticalAlignment(EVerticalAlignment InVerticalAlignment)
 {
 	VerticalAlignment = InVerticalAlignment;
@@ -61,10 +81,13 @@ void UButtonSlot::SetVerticalAlignment(EVerticalAlignment InVerticalAlignment)
 		Button.Pin()->SetVAlign(InVerticalAlignment);
 	}
 }
+PRAGMA_ENABLE_DEPRECATION_WARNINGS
 
 void UButtonSlot::SynchronizeProperties()
 {
+PRAGMA_DISABLE_DEPRECATION_WARNINGS
 	SetPadding(Padding);
 	SetHorizontalAlignment(HorizontalAlignment);
 	SetVerticalAlignment(VerticalAlignment);
+PRAGMA_ENABLE_DEPRECATION_WARNINGS
 }
diff --git a/Engine/Source/Runtime/UMG/Private/Components/CanvasPanelSlot.cpp b/Engine/Source/Runtime/UMG/Private/Components/CanvasPanelSlot.cpp
--- a/Engine/Source/Runtime/UMG/Private/Components/CanvasPanelSlot.cpp
+++ b/Engine/Source/Runtime/UMG/Private/Components/CanvasPanelSlot.cpp
@@ -10,11 +10,13 @@ UCanvasPanelSlot::UCanvasPanelSlot(const FObjectInitializer& ObjectInitializer)
 	: Super(ObjectInitializer)
 	, Slot(nullptr)
 {
-	LayoutData.Offsets = FMargin(0, 0, 100, 30);
+PRAGMA_DISABLE_DEPRECATION_WARNINGS
+	LayoutData.Offsets = FMargin(0.f, 0.f, 100.f, 30.f);
 	LayoutData.Anchors = FAnchors(0.0f, 0.0f);
 	LayoutData.Alignment = FVector2D(0.0f, 0.0f);
 	bAutoSize = false;
 	ZOrder = 0;
+PRAGMA_ENABLE_DEPRECATION_WARNINGS
 }
 
 void UCanvasPanelSlot::ReleaseSlateResources(bool bReleaseChildren)
@@ -134,6 +136,7 @@ void UCanvasPanelSlot::SynchronizeFromTemplate(const UPanelSlot* const TemplateS
 
 #endif //WITH_EDITOR
 
+PRAGMA_DISABLE_DEPRECATION_WARNINGS
 void UCanvasPanelSlot::SetLayout(const FAnchorData& InLayoutData)
 {
 	LayoutData = InLayoutData;
@@ -307,14 +310,17 @@ void UCanvasPanelSlot::SetMaximum(FVector2D InMaximumAnchors)
 		Slot->SetAnchors(LayoutData.Anchors);
 	}
 }
+PRAGMA_ENABLE_DEPRECATION_WARNINGS
 
 void UCanvasPanelSlot::SynchronizeProperties()
 {
+PRAGMA_DISABLE_DEPRECATION_WARNINGS
 	SetOffsets(LayoutData.Offsets);
 	SetAnchors(LayoutData.Anchors);
 	SetAlignment(LayoutData.Alignment);
 	SetAutoSize(bAutoSize);
 	SetZOrder(ZOrder);
+PRAGMA_ENABLE_DEPRECATION_WARNINGS
 }
 
 #if WITH_EDITOR
@@ -356,7 +362,7 @@ void UCanvasPanelSlot::SaveBaseLayout()
 		if ( Canvas->GetGeometryForSlot(this, Geometry) )
 		{
 			PreEditGeometry = Geometry;
-			PreEditLayoutData = LayoutData;
+			PreEditLayoutData = GetLayout();
 		}
 	}
 }
@@ -376,131 +382,131 @@ void UCanvasPanelSlot::RebaseLayout(bool PreserveSize)
 		{
 			// Calculate the default anchor offset, ie where would this control be laid out if no offset were provided.
 			FVector2D CanvasSize = Canvas->GetCanvasWidget()->GetCachedGeometry().Size;
+			FAnchorData LocalLayoutData = GetLayout();
 			FMargin AnchorPositions = FMargin(
-				LayoutData.Anchors.Minimum.X * CanvasSize.X,
-				LayoutData.Anchors.Minimum.Y * CanvasSize.Y,
-				LayoutData.Anchors.Maximum.X * CanvasSize.X,
-				LayoutData.Anchors.Maximum.Y * CanvasSize.Y);
+				LocalLayoutData.Anchors.Minimum.X * CanvasSize.X,
+				LocalLayoutData.Anchors.Minimum.Y * CanvasSize.Y,
+				LocalLayoutData.Anchors.Maximum.X * CanvasSize.X,
+				LocalLayoutData.Anchors.Maximum.Y * CanvasSize.Y);
 			FVector2D DefaultAnchorPosition = FVector2D(AnchorPositions.Left, AnchorPositions.Top);
 
 			// Determine the amount that would be offset from the anchor position if alignment was applied.
-			FVector2D AlignmentOffset = LayoutData.Alignment * PreEditGeometry.Size;
+			FVector2D AlignmentOffset = LocalLayoutData.Alignment * PreEditGeometry.Size;
 
 			FVector2D MoveDelta = FVector2D(Geometry.Position) - FVector2D(PreEditGeometry.Position);
 
 			// Determine where the widget's new position needs to be to maintain a stable location when the anchors change.
 			FVector2D LeftTopDelta = FVector2D(PreEditGeometry.Position) - FVector2D(DefaultAnchorPosition);
 
-			const bool bAnchorsMoved = PreEditLayoutData.Anchors.Minimum != LayoutData.Anchors.Minimum || PreEditLayoutData.Anchors.Maximum != LayoutData.Anchors.Maximum;
-			const bool bMoved = PreEditLayoutData.Offsets.Left != LayoutData.Offsets.Left || PreEditLayoutData.Offsets.Top != LayoutData.Offsets.Top;
+			const bool bAnchorsMoved = PreEditLayoutData.Anchors.Minimum != LocalLayoutData.Anchors.Minimum || PreEditLayoutData.Anchors.Maximum != LocalLayoutData.Anchors.Maximum;
+			const bool bMoved = PreEditLayoutData.Offsets.Left != LocalLayoutData.Offsets.Left || PreEditLayoutData.Offsets.Top != LocalLayoutData.Offsets.Top;
 
 			if ( bAnchorsMoved )
 			{
 				// Adjust the size to remain constant
-				if ( !LayoutData.Anchors.IsStretchedHorizontal() && PreEditLayoutData.Anchors.IsStretchedHorizontal() )
+				if ( !LocalLayoutData.Anchors.IsStretchedHorizontal() && PreEditLayoutData.Anchors.IsStretchedHorizontal() )
 				{
 					// Adjust the position to remain constant
-					LayoutData.Offsets.Left = LeftTopDelta.X + AlignmentOffset.X;
-					LayoutData.Offsets.Right = PreEditGeometry.Size.X;
+					LocalLayoutData.Offsets.Left = LeftTopDelta.X + AlignmentOffset.X;
+					LocalLayoutData.Offsets.Right = PreEditGeometry.Size.X;
 				}
-				else if ( !PreserveSize && LayoutData.Anchors.IsStretchedHorizontal() && !PreEditLayoutData.Anchors.IsStretchedHorizontal() )
+				else if ( !PreserveSize && LocalLayoutData.Anchors.IsStretchedHorizontal() && !PreEditLayoutData.Anchors.IsStretchedHorizontal() )
 				{
 					// Adjust the position to remain constant
-					LayoutData.Offsets.Left = 0;
-					LayoutData.Offsets.Right = 0;
+					LocalLayoutData.Offsets.Left = 0;
+					LocalLayoutData.Offsets.Right = 0;
 				}
-				else if ( LayoutData.Anchors.IsStretchedHorizontal() )
+				else if ( LocalLayoutData.Anchors.IsStretchedHorizontal() )
 				{
 					// Adjust the position to remain constant
-					LayoutData.Offsets.Left = LeftTopDelta.X;
-					LayoutData.Offsets.Right = AnchorPositions.Right - ( AnchorPositions.Left + LayoutData.Offsets.Left + PreEditGeometry.Size.X );
+					LocalLayoutData.Offsets.Left = LeftTopDelta.X;
+					LocalLayoutData.Offsets.Right = AnchorPositions.Right - ( AnchorPositions.Left + LocalLayoutData.Offsets.Left + PreEditGeometry.Size.X );
 				}
 				else
 				{
 					// Adjust the position to remain constant
-					LayoutData.Offsets.Left = LeftTopDelta.X + AlignmentOffset.X;
+					LocalLayoutData.Offsets.Left = LeftTopDelta.X + AlignmentOffset.X;
 				}
 
-				if ( !LayoutData.Anchors.IsStretchedVertical() && PreEditLayoutData.Anchors.IsStretchedVertical() )
+				if ( !LocalLayoutData.Anchors.IsStretchedVertical() && PreEditLayoutData.Anchors.IsStretchedVertical() )
 				{
 					// Adjust the position to remain constant
-					LayoutData.Offsets.Top = LeftTopDelta.Y + AlignmentOffset.Y;
-					LayoutData.Offsets.Bottom = PreEditGeometry.Size.Y;
+					LocalLayoutData.Offsets.Top = LeftTopDelta.Y + AlignmentOffset.Y;
+					LocalLayoutData.Offsets.Bottom = PreEditGeometry.Size.Y;
 				}
-				else if ( !PreserveSize && LayoutData.Anchors.IsStretchedVertical() && !PreEditLayoutData.Anchors.IsStretchedVertical() )
+				else if ( !PreserveSize && LocalLayoutData.Anchors.IsStretchedVertical() && !PreEditLayoutData.Anchors.IsStretchedVertical() )
 				{
 					// Adjust the position to remain constant
-					LayoutData.Offsets.Top = 0;
-					LayoutData.Offsets.Bottom = 0;
+					LocalLayoutData.Offsets.Top = 0;
+					LocalLayoutData.Offsets.Bottom = 0;
 				}
-				else if ( LayoutData.Anchors.IsStretchedVertical() )
+				else if ( LocalLayoutData.Anchors.IsStretchedVertical() )
 				{
 					// Adjust the position to remain constant
-					LayoutData.Offsets.Top = LeftTopDelta.Y;
-					LayoutData.Offsets.Bottom = AnchorPositions.Bottom - ( AnchorPositions.Top + LayoutData.Offsets.Top + PreEditGeometry.Size.Y );
+					LocalLayoutData.Offsets.Top = LeftTopDelta.Y;
+					LocalLayoutData.Offsets.Bottom = AnchorPositions.Bottom - ( AnchorPositions.Top + LocalLayoutData.Offsets.Top + PreEditGeometry.Size.Y );
 				}
 				else
 				{
 					// Adjust the position to remain constant
-					LayoutData.Offsets.Top = LeftTopDelta.Y + AlignmentOffset.Y;
+					LocalLayoutData.Offsets.Top = LeftTopDelta.Y + AlignmentOffset.Y;
 				}
 			}
 			else if ( DesiredPosition.IsSet() )
 			{
 				FVector2D NewLocalPosition = DesiredPosition.GetValue();
 
-				LayoutData.Offsets.Left = NewLocalPosition.X - AnchorPositions.Left;
-				LayoutData.Offsets.Top = NewLocalPosition.Y - AnchorPositions.Top;
+				LocalLayoutData.Offsets.Left = NewLocalPosition.X - AnchorPositions.Left;
+				LocalLayoutData.Offsets.Top = NewLocalPosition.Y - AnchorPositions.Top;
 
-				if ( LayoutData.Anchors.IsStretchedHorizontal() )
+				if ( LocalLayoutData.Anchors.IsStretchedHorizontal() )
 				{
-					LayoutData.Offsets.Right -= LayoutData.Offsets.Left - PreEditLayoutData.Offsets.Left;
+					LocalLayoutData.Offsets.Right -= LocalLayoutData.Offsets.Left - PreEditLayoutData.Offsets.Left;
 				}
 				else
 				{
-					LayoutData.Offsets.Left += AlignmentOffset.X;
+					LocalLayoutData.Offsets.Left += AlignmentOffset.X;
 				}
 
-				if ( LayoutData.Anchors.IsStretchedVertical() )
+				if ( LocalLayoutData.Anchors.IsStretchedVertical() )
 				{
-					LayoutData.Offsets.Bottom -= LayoutData.Offsets.Top - PreEditLayoutData.Offsets.Top;
+					LocalLayoutData.Offsets.Bottom -= LocalLayoutData.Offsets.Top - PreEditLayoutData.Offsets.Top;
 				}
 				else
 				{
-					LayoutData.Offsets.Top += AlignmentOffset.Y;
+					LocalLayoutData.Offsets.Top += AlignmentOffset.Y;
 				}
 
 				DesiredPosition.Reset();
 			}
 			else if ( bMoved )
 			{
-				//LayoutData.Offsets.Left += LeftTopDelta.X;
-				//LayoutData.Offsets.Top += LeftTopDelta.Y;
-				LayoutData.Offsets.Left -= DefaultAnchorPosition.X;
-				LayoutData.Offsets.Top -= DefaultAnchorPosition.Y;
+				LocalLayoutData.Offsets.Left -= DefaultAnchorPosition.X;
+				LocalLayoutData.Offsets.Top -= DefaultAnchorPosition.Y;
 
 				// If the slot is stretched horizontally we need to move the right side as it no longer represents width, but
 				// now represents margin from the right stretched side.
-				if ( LayoutData.Anchors.IsStretchedHorizontal() )
+				if ( LocalLayoutData.Anchors.IsStretchedHorizontal() )
 				{
-					//LayoutData.Offsets.Right = PreEditLayoutData.Offsets.Top;
+					//LocalLayoutData.Offsets.Right = PreEditLayoutData.Offsets.Top;
 				}
 				else
 				{
-					LayoutData.Offsets.Left += AlignmentOffset.X;
+					LocalLayoutData.Offsets.Left += AlignmentOffset.X;
 				}
 
 				// If the slot is stretched vertically we need to move the bottom side as it no longer represents width, but
 				// now represents margin from the bottom stretched side.
-				if ( LayoutData.Anchors.IsStretchedVertical() )
+				if ( LocalLayoutData.Anchors.IsStretchedVertical() )
 				{
-					//LayoutData.Offsets.Bottom -= MoveDelta.Y;
+					//LocalLayoutData.Offsets.Bottom -= MoveDelta.Y;
 				}
 				else
 				{
-					LayoutData.Offsets.Top += AlignmentOffset.Y;
+					LocalLayoutData.Offsets.Top += AlignmentOffset.Y;
 				}
 			}
+			SetLayout(LocalLayoutData);
 		}
 
 		// Apply the changes to the properties.
diff --git a/Engine/Source/Runtime/UMG/Private/Components/GridSlot.cpp b/Engine/Source/Runtime/UMG/Private/Components/GridSlot.cpp
--- a/Engine/Source/Runtime/UMG/Private/Components/GridSlot.cpp
+++ b/Engine/Source/Runtime/UMG/Private/Components/GridSlot.cpp
@@ -10,11 +10,13 @@ UGridSlot::UGridSlot(const FObjectInitializer& ObjectInitializer)
 	: Super(ObjectInitializer)
 	, Slot(nullptr)
 {
+PRAGMA_DISABLE_DEPRECATION_WARNINGS
 	HorizontalAlignment = HAlign_Fill;
 	VerticalAlignment = VAlign_Fill;
 
 	Layer = 0;
 	Nudge = FVector2D(0, 0);
+PRAGMA_ENABLE_DEPRECATION_WARNINGS
 }
 
 void UGridSlot::ReleaseSlateResources(bool bReleaseChildren)
@@ -26,6 +28,7 @@ void UGridSlot::ReleaseSlateResources(bool bReleaseChildren)
 
 void UGridSlot::BuildSlot(TSharedRef<SGridPanel> GridPanel)
 {
+PRAGMA_DISABLE_DEPRECATION_WARNINGS
 	GridPanel->AddSlot(Column, Row, SGridPanel::Layer(Layer))
 		.Expose(Slot)
 		.Padding(Padding)
@@ -37,6 +40,13 @@ void UGridSlot::BuildSlot(TSharedRef<SGridPanel> GridPanel)
 		[
 			Content == nullptr ? SNullWidget::NullWidget : Content->TakeWidget()
 		];
+PRAGMA_ENABLE_DEPRECATION_WARNINGS
+}
+
+PRAGMA_DISABLE_DEPRECATION_WARNINGS
+FMargin UGridSlot::GetPadding() const
+{
+	return Slot ? Slot->GetPadding() : Padding;
 }
 
 void UGridSlot::SetPadding(FMargin InPadding)
@@ -48,6 +58,11 @@ void UGridSlot::SetPadding(FMargin InPadding)
 	}
 }
 
+int32 UGridSlot::GetRow() const
+{
+	return Slot ? Slot->GetRow() : Row;
+}
+
 void UGridSlot::SetRow(int32 InRow)
 {
 	Row = InRow;
@@ -57,6 +72,11 @@ void UGridSlot::SetRow(int32 InRow)
 	}
 }
 
+int32 UGridSlot::GetRowSpan() const
+{
+	return Slot ? Slot->GetRowSpan() : RowSpan;
+}
+
 void UGridSlot::SetRowSpan(int32 InRowSpan)
 {
 	RowSpan = InRowSpan;
@@ -66,6 +86,12 @@ void UGridSlot::SetRowSpan(int32 InRowSpan)
 	}
 }
 
+int32 UGridSlot::GetColumn() const
+{
+	return Slot ? Slot->GetColumn() : Column;
+}
+
+
 void UGridSlot::SetColumn(int32 InColumn)
 {
 	Column = InColumn;
@@ -75,6 +101,11 @@ void UGridSlot::SetColumn(int32 InColumn)
 	}
 }
 
+int32 UGridSlot::GetColumnSpan() const
+{
+	return Slot ? Slot->GetColumnSpan() : ColumnSpan;
+}
+
 void UGridSlot::SetColumnSpan(int32 InColumnSpan)
 {
 	ColumnSpan = InColumnSpan;
@@ -84,6 +115,11 @@ void UGridSlot::SetColumnSpan(int32 InColumnSpan)
 	}
 }
 
+int32 UGridSlot::GetLayer() const
+{
+	return Slot ? Slot->GetLayer() : Layer;
+}
+
 void UGridSlot::SetLayer(int32 InLayer)
 {
 	Layer = InLayer;
@@ -93,6 +129,11 @@ void UGridSlot::SetLayer(int32 InLayer)
 	}
 }
 
+FVector2D UGridSlot::GetNudge() const
+{
+	return Slot ? Slot->GetNudge() : Nudge;
+}
+
 void UGridSlot::SetNudge(FVector2D InNudge)
 {
 	Nudge = InNudge;
@@ -102,6 +143,11 @@ void UGridSlot::SetNudge(FVector2D InNudge)
 	}
 }
 
+EHorizontalAlignment UGridSlot::GetHorizontalAlignment() const
+{
+	return Slot ? Slot->GetHorizontalAlignment() : HorizontalAlignment.GetValue();
+}
+
 void UGridSlot::SetHorizontalAlignment(EHorizontalAlignment InHorizontalAlignment)
 {
 	HorizontalAlignment = InHorizontalAlignment;
@@ -111,6 +157,11 @@ void UGridSlot::SetHorizontalAlignment(EHorizontalAlignment InHorizontalAlignmen
 	}
 }
 
+EVerticalAlignment UGridSlot::GetVerticalAlignment() const
+{
+	return Slot ? Slot->GetVerticalAlignment() : VerticalAlignment.GetValue();
+}
+
 void UGridSlot::SetVerticalAlignment(EVerticalAlignment InVerticalAlignment)
 {
 	VerticalAlignment = InVerticalAlignment;
@@ -119,9 +170,11 @@ void UGridSlot::SetVerticalAlignment(EVerticalAlignment InVerticalAlignment)
 		Slot->SetVerticalAlignment(InVerticalAlignment);
 	}
 }
+PRAGMA_ENABLE_DEPRECATION_WARNINGS
 
 void UGridSlot::SynchronizeProperties()
 {
+PRAGMA_DISABLE_DEPRECATION_WARNINGS
 	SetHorizontalAlignment(HorizontalAlignment);
 	SetVerticalAlignment(VerticalAlignment);
 	SetPadding(Padding);
@@ -133,6 +186,7 @@ void UGridSlot::SynchronizeProperties()
 	SetNudge(Nudge);
 
 	SetLayer(Layer);
+PRAGMA_ENABLE_DEPRECATION_WARNINGS
 }
 
 #if WITH_EDITOR
@@ -140,10 +194,10 @@ void UGridSlot::SynchronizeProperties()
 bool UGridSlot::NudgeByDesigner(const FVector2D& NudgeDirection, const TOptional<int32>& GridSnapSize)
 {
 	const FVector2D ClampedDirection = NudgeDirection.ClampAxes(-1.0f, 1.0f);
-	const int32 NewColumn = Column + ClampedDirection.X;
-	const int32 NewRow = Row + ClampedDirection.Y;
+	const int32 NewColumn = GetColumn() + ClampedDirection.X;
+	const int32 NewRow = GetRow() + ClampedDirection.Y;
 
-	if (NewColumn < 0 || NewRow < 0 || (NewColumn == Column && NewRow == Row))
+	if (NewColumn < 0 || NewRow < 0 || (NewColumn == GetColumn() && NewRow == GetRow()))
 	{
 		return false;
 	}
@@ -159,8 +213,8 @@ bool UGridSlot::NudgeByDesigner(const FVector2D& NudgeDirection, const TOptional
 void UGridSlot::SynchronizeFromTemplate(const UPanelSlot* const TemplateSlot)
 {
 	const ThisClass* const TemplateGridSlot = CastChecked<ThisClass>(TemplateSlot);
-	SetRow(TemplateGridSlot->Row);
-	SetColumn(TemplateGridSlot->Column);
+	SetRow(TemplateGridSlot->GetRow());
+	SetColumn(TemplateGridSlot->GetColumn());
 }
 
 #endif
diff --git a/Engine/Source/Runtime/UMG/Public/Components/ButtonSlot.h b/Engine/Source/Runtime/UMG/Public/Components/ButtonSlot.h
--- a/Engine/Source/Runtime/UMG/Public/Components/ButtonSlot.h
+++ b/Engine/Source/Runtime/UMG/Public/Components/ButtonSlot.h
@@ -21,25 +21,34 @@ class UMG_API UButtonSlot : public UPanelSlot
 public:
 	
 	/** The padding area between the slot and the content it contains. */
-	UPROPERTY(EditAnywhere, BlueprintReadOnly, Category="Layout|Button Slot")
+	UE_DEPRECATED(5.1, "Direct access to Padding is deprecated. Please use the getter or setter.")
+	UPROPERTY(EditAnywhere, BlueprintReadWrite, Getter, Setter, BlueprintSetter="SetPadding", Category="Layout|Button Slot")
 	FMargin Padding;
 
 	/** The alignment of the object horizontally. */
-	UPROPERTY(EditAnywhere, BlueprintReadOnly, Category="Layout|Button Slot")
+	UE_DEPRECATED(5.1, "Direct access to HorizontalAlignment is deprecated. Please use the getter or setter.")
+	UPROPERTY(EditAnywhere, BlueprintReadWrite, Getter, Setter, BlueprintSetter="SetHorizontalAlignment", Category="Layout|Button Slot")
 	TEnumAsByte<EHorizontalAlignment> HorizontalAlignment;
 
 	/** The alignment of the object vertically. */
-	UPROPERTY(EditAnywhere, BlueprintReadOnly, Category="Layout|Button Slot")
+	UE_DEPRECATED(5.1, "Direct access to VerticalAlignment is deprecated. Please use the getter or setter.")
+	UPROPERTY(EditAnywhere, BlueprintReadWrite, Getter, Setter, BlueprintSetter="SetVerticalAlignment", Category="Layout|Button Slot")
 	TEnumAsByte<EVerticalAlignment> VerticalAlignment;
 
 public:
 
+	FMargin GetPadding() const;
+
 	UFUNCTION(BlueprintCallable, Category="Layout|Button Slot")
 	void SetPadding(FMargin InPadding);
 
+	EHorizontalAlignment GetHorizontalAlignment() const;
+
 	UFUNCTION(BlueprintCallable, Category="Layout|Button Slot")
 	void SetHorizontalAlignment(EHorizontalAlignment InHorizontalAlignment);
 
+	EVerticalAlignment GetVerticalAlignment() const;
+
 	UFUNCTION(BlueprintCallable, Category="Layout|Button Slot")
 	void SetVerticalAlignment(EVerticalAlignment InVerticalAlignment);
 
diff --git a/Engine/Source/Runtime/UMG/Public/Components/CanvasPanelSlot.h b/Engine/Source/Runtime/UMG/Public/Components/CanvasPanelSlot.h
--- a/Engine/Source/Runtime/UMG/Public/Components/CanvasPanelSlot.h
+++ b/Engine/Source/Runtime/UMG/Public/Components/CanvasPanelSlot.h
@@ -70,15 +70,18 @@ class UMG_API UCanvasPanelSlot : public UPanelSlot
 public:
 
 	/** The anchoring information for the slot */
-	UPROPERTY(EditAnywhere, BlueprintReadOnly, Category="Layout|Canvas Slot")
+	UE_DEPRECATED(5.1, "Direct access to LayoutData is deprecated. Please use the getter or setter.")
+	UPROPERTY(EditAnywhere, BlueprintReadWrite, Getter="GetLayout", Setter="SetLayout", BlueprintGetter ="GetLayout", BlueprintSetter="SetLayout", Category = "Layout|Canvas Slot")
 	FAnchorData LayoutData;
 
 	/** When AutoSize is true we use the widget's desired size */
-	UPROPERTY(EditAnywhere, BlueprintReadOnly, Category="Layout|Canvas Slot", AdvancedDisplay, meta=( DisplayName="Size To Content" ))
+	UE_DEPRECATED(5.1, "Direct access to bAutoSize is deprecated. Please use the getter or setter.")
+	UPROPERTY(EditAnywhere, BlueprintReadWrite, Getter="GetAutoSize", Setter="SetAutoSize", BlueprintGetter="GetAutoSize", BlueprintSetter="SetAutoSize", Category = "Layout|Canvas Slot", AdvancedDisplay, meta = (DisplayName = "Size To Content"))
 	bool bAutoSize;
 
 	/** The order priority this widget is rendered in.  Higher values are rendered last (and so they will appear to be on top). */
-	UPROPERTY(EditAnywhere, BlueprintReadOnly, Category="Layout|Canvas Slot")
+	UE_DEPRECATED(5.1, "Direct access to ZOrder is deprecated. Please use the getter or setter.")
+	UPROPERTY(EditAnywhere, BlueprintReadWrite, Getter, Setter, BlueprintGetter="GetZOrder", BlueprintSetter="SetZOrder", Category="Layout|Canvas Slot")
 	int32 ZOrder;
 
 public:
diff --git a/Engine/Source/Runtime/UMG/Public/Components/GridSlot.h b/Engine/Source/Runtime/UMG/Public/Components/GridSlot.h
--- a/Engine/Source/Runtime/UMG/Public/Components/GridSlot.h
+++ b/Engine/Source/Runtime/UMG/Public/Components/GridSlot.h
@@ -23,74 +23,109 @@ class UMG_API UGridSlot : public UPanelSlot
 public:
 
 	/** The padding area between the slot and the content it contains. */
-	UPROPERTY(EditAnywhere, BlueprintReadOnly, Category="Layout|Grid Slot")
+	UE_DEPRECATED(5.1, "Direct access to Padding is deprecated. Please use the getter or setter.")
+	UPROPERTY(EditAnywhere, BlueprintReadWrite, Getter, Setter, BlueprintSetter="SetPadding", Category="Layout|Grid Slot")
 	FMargin Padding;
 
 	/** The alignment of the object horizontally. */
-	UPROPERTY(EditAnywhere, BlueprintReadOnly, Category="Layout|Grid Slot")
+	UE_DEPRECATED(5.1, "Direct access to HorizontalAlignment is deprecated. Please use the getter or setter.")
+	UPROPERTY(EditAnywhere, BlueprintReadWrite, Getter, Setter, BlueprintSetter="SetHorizontalAlignment", Category="Layout|Grid Slot")
 	TEnumAsByte<EHorizontalAlignment> HorizontalAlignment;
 
 	/** The alignment of the object vertically. */
-	UPROPERTY(EditAnywhere, BlueprintReadOnly, Category="Layout|Grid Slot")
+	UE_DEPRECATED(5.1, "Direct access to VerticalAlignment is deprecated. Please use the getter or setter.")
+	UPROPERTY(EditAnywhere, BlueprintReadWrite, Getter, Setter, BlueprintSetter="SetVerticalAlignment", Category="Layout|Grid Slot")
 	TEnumAsByte<EVerticalAlignment> VerticalAlignment;
 	
 	/** The row index of the cell this slot is in */
-	UPROPERTY(EditAnywhere, BlueprintReadOnly, meta=( UIMin = "0" ), Category="Layout|Grid Slot")
+	UE_DEPRECATED(5.1, "Direct access to Row is deprecated. Please use the getter or setter.")
+	UPROPERTY(EditAnywhere, BlueprintReadWrite, Getter, Setter, BlueprintSetter="SetRow", meta=( UIMin = "0" ), Category="Layout|Grid Slot")
 	int32 Row;
 	
 	/**  */
-	UPROPERTY(EditAnywhere, BlueprintReadOnly, Category="Layout|Grid Slot")
+	UE_DEPRECATED(5.1, "Direct access to RowSpan is deprecated. Please use the getter or setter.")
+	UPROPERTY(EditAnywhere, BlueprintReadWrite, Getter, Setter, BlueprintSetter="SetRowSpan", Category="Layout|Grid Slot")
 	int32 RowSpan;
 
 	/** The column index of the cell this slot is in */
-	UPROPERTY(EditAnywhere, BlueprintReadOnly, meta=( UIMin = "0" ), Category="Layout|Grid Slot")
+	UE_DEPRECATED(5.1, "Direct access to Column is deprecated. Please use the getter or setter.")
+	UPROPERTY(EditAnywhere, BlueprintReadWrite, Getter, Setter, BlueprintSetter="SetColumn", meta=( UIMin = "0" ), Category="Layout|Grid Slot")
 	int32 Column;
 
 	/**  */
-	UPROPERTY(EditAnywhere, BlueprintReadOnly, Category="Layout|Grid Slot")
+	UE_DEPRECATED(5.1, "Direct access to ColumnSpan is deprecated. Please use the getter or setter.")
+	UPROPERTY(EditAnywhere, BlueprintReadWrite, Getter, Setter, BlueprintSetter="SetColumnSpan", Category="Layout|Grid Slot")
 	int32 ColumnSpan;
 
 	/** Positive values offset this cell to be hit-tested and drawn on top of others. Default is 0; i.e. no offset. */
-	UPROPERTY(EditAnywhere, BlueprintReadOnly, Category="Layout|Grid Slot")
+	UE_DEPRECATED(5.1, "Direct access to Layer is deprecated. Please use the getter or setter.")
+	UPROPERTY(EditAnywhere, BlueprintReadWrite, Getter, Setter, BlueprintSetter="SetLayer", Category="Layout|Grid Slot")
 	int32 Layer;
 
 	/** Offset this slot's content by some amount; positive values offset to lower right */
-	UPROPERTY(EditAnywhere, BlueprintReadOnly, Category="Layout|Grid Slot")
+	UE_DEPRECATED(5.1, "Direct access to Nudge is deprecated. Please use the getter or setter.")
+	UPROPERTY(EditAnywhere, BlueprintReadWrite, Getter, Setter, BlueprintSetter="SetNudge", Category="Layout|Grid Slot")
 	FVector2D Nudge;
 
 public:
 
+	FMargin GetPadding() const;
+
 	UFUNCTION(BlueprintCallable, Category="Layout|Border Slot")
 	void SetPadding(FMargin InPadding);
 
+	/** Gets the row index of the slot */
+	int32 GetRow() const;
+
 	/** Sets the row index of the slot, this determines what cell the slot is in the panel */
 	UFUNCTION(BlueprintCallable, Category="Layout|Grid Slot")
 	void SetRow(int32 InRow);
 
+	/** Gets how many rows this this slot spans over */
+	int32 GetRowSpan() const;
+
 	/** How many rows this this slot spans over */
 	UFUNCTION(BlueprintCallable, Category="Layout|Grid Slot")
 	void SetRowSpan(int32 InRowSpan);
 
+	/** Gets the column index of the slot */
+	int32 GetColumn() const;
+
 	/** Sets the column index of the slot, this determines what cell the slot is in the panel */
 	UFUNCTION(BlueprintCallable, Category="Layout|Grid Slot")
 	void SetColumn(int32 InColumn);
 
+	/** Gets how many columns this slot spans over */
+	int32 GetColumnSpan() const;
+
 	/** How many columns this slot spans over */
 	UFUNCTION(BlueprintCallable, Category="Layout|Grid Slot")
 	void SetColumnSpan(int32 InColumnSpan);
 
+	/** Gets the offset this cell drawn on top of others. */
+	int32 GetLayer() const;
+
 	/** Sets positive values offset this cell to be hit-tested and drawn on top of others. */
 	UFUNCTION(BlueprintCallable, Category = "Layout|Grid Slot")
 	void SetLayer(int32 InLayer);
 
+	/** Gets the offset for this slot's content */
+	FVector2D GetNudge() const;
+
 	/** Sets the offset for this slot's content by some amount; positive values offset to lower right*/
 	UFUNCTION(BlueprintCallable, Category = "Layout|Grid Slot")
 	void SetNudge(FVector2D InNudge);
 
+	/** */
+	EHorizontalAlignment GetHorizontalAlignment() const;
+
 	/**  */
 	UFUNCTION(BlueprintCallable, Category="Layout|Grid Slot")
 	void SetHorizontalAlignment(EHorizontalAlignment InHorizontalAlignment);
 
+	/**  */
+	EVerticalAlignment GetVerticalAlignment() const;
+
 	/**  */
 	UFUNCTION(BlueprintCallable, Category="Layout|Grid Slot")
 	void SetVerticalAlignment(EVerticalAlignment InVerticalAlignment);
@@ -110,8 +145,6 @@ public:
 #endif //WITH_EDITOR
 
 private:
-	//TODO UMG Slots should hold weak or shared refs to slots.
-
 	/** A raw pointer to the slot to allow us to adjust the size, padding...etc at runtime. */
 	SGridPanel::FSlot* Slot;
 };
