UMG: Update ScrollBoxSlot, ScaleBoxSlot with accessors.

diff --git a/Engine/Source/Runtime/UMG/Private/Components/ScaleBoxSlot.cpp b/Engine/Source/Runtime/UMG/Private/Components/ScaleBoxSlot.cpp
--- a/Engine/Source/Runtime/UMG/Private/Components/ScaleBoxSlot.cpp
+++ b/Engine/Source/Runtime/UMG/Private/Components/ScaleBoxSlot.cpp
@@ -11,10 +11,10 @@
 UScaleBoxSlot::UScaleBoxSlot(const FObjectInitializer& ObjectInitializer)
 	: Super(ObjectInitializer)
 {
-	Padding = FMargin(0, 0);
-
+PRAGMA_DISABLE_DEPRECATION_WARNINGS
 	HorizontalAlignment = HAlign_Center;
 	VerticalAlignment = VAlign_Center;
+PRAGMA_DISABLE_DEPRECATION_WARNINGS
 }
 
 void UScaleBoxSlot::ReleaseSlateResources(bool bReleaseChildren)
@@ -28,20 +28,21 @@ void UScaleBoxSlot::BuildSlot(TSharedRef<SScaleBox> InScaleBox)
 {
 	ScaleBox = InScaleBox;
 
-	//ScaleBox->SetPadding(Padding);
-	ScaleBox.Pin()->SetHAlign(HorizontalAlignment);
-	ScaleBox.Pin()->SetVAlign(VerticalAlignment);
-
-	ScaleBox.Pin()->SetContent(Content ? Content->TakeWidget() : SNullWidget::NullWidget);
+PRAGMA_DISABLE_DEPRECATION_WARNINGS
+	InScaleBox->SetHAlign(HorizontalAlignment);
+	InScaleBox->SetVAlign(VerticalAlignment);
+PRAGMA_ENABLE_DEPRECATION_WARNINGS
+	InScaleBox->SetContent(Content ? Content->TakeWidget() : SNullWidget::NullWidget);
 }
 
 void UScaleBoxSlot::SetPadding(FMargin InPadding)
 {
-	Padding = InPadding;
-	if ( ScaleBox.IsValid() )
-	{
-		//ScaleBox.Pin()->SetPadding(InPadding);
-	}
+}
+
+PRAGMA_DISABLE_DEPRECATION_WARNINGS
+EHorizontalAlignment UScaleBoxSlot::GetHorizontalAlignment() const
+{
+	return HorizontalAlignment;
 }
 
 void UScaleBoxSlot::SetHorizontalAlignment(EHorizontalAlignment InHorizontalAlignment)
@@ -53,6 +54,11 @@ void UScaleBoxSlot::SetHorizontalAlignment(EHorizontalAlignment InHorizontalAlig
 	}
 }
 
+EVerticalAlignment UScaleBoxSlot::GetVerticalAlignment() const
+{
+	return VerticalAlignment;
+}
+
 void UScaleBoxSlot::SetVerticalAlignment(EVerticalAlignment InVerticalAlignment)
 {
 	VerticalAlignment = InVerticalAlignment;
@@ -61,10 +67,12 @@ void UScaleBoxSlot::SetVerticalAlignment(EVerticalAlignment InVerticalAlignment)
 		ScaleBox.Pin()->SetVAlign(InVerticalAlignment);
 	}
 }
+PRAGMA_ENABLE_DEPRECATION_WARNINGS
 
 void UScaleBoxSlot::SynchronizeProperties()
 {
-	SetPadding(Padding);
+PRAGMA_DISABLE_DEPRECATION_WARNINGS
 	SetHorizontalAlignment(HorizontalAlignment);
 	SetVerticalAlignment(VerticalAlignment);
+PRAGMA_ENABLE_DEPRECATION_WARNINGS
 }
diff --git a/Engine/Source/Runtime/UMG/Private/Components/ScrollBoxSlot.cpp b/Engine/Source/Runtime/UMG/Private/Components/ScrollBoxSlot.cpp
--- a/Engine/Source/Runtime/UMG/Private/Components/ScrollBoxSlot.cpp
+++ b/Engine/Source/Runtime/UMG/Private/Components/ScrollBoxSlot.cpp
@@ -10,12 +10,15 @@ UScrollBoxSlot::UScrollBoxSlot(const FObjectInitializer& ObjectInitializer)
 	: Super(ObjectInitializer)
 	, Slot(nullptr)
 {
+PRAGMA_DISABLE_DEPRECATION_WARNINGS
 	HorizontalAlignment = HAlign_Fill;
 	VerticalAlignment = VAlign_Fill;
+PRAGMA_ENABLE_DEPRECATION_WARNINGS
 }
 
 void UScrollBoxSlot::BuildSlot(TSharedRef<SScrollBox> ScrollBox)
 {
+PRAGMA_DISABLE_DEPRECATION_WARNINGS
 	ScrollBox->AddSlot()
 		.Padding(Padding)
 		.HAlign(HorizontalAlignment)
@@ -24,6 +27,13 @@ void UScrollBoxSlot::BuildSlot(TSharedRef<SScrollBox> ScrollBox)
 		[
 			Content == nullptr ? SNullWidget::NullWidget : Content->TakeWidget()
 		];
+PRAGMA_ENABLE_DEPRECATION_WARNINGS
+}
+
+PRAGMA_DISABLE_DEPRECATION_WARNINGS
+FMargin UScrollBoxSlot::GetPadding() const
+{
+	return Slot ? Slot->GetPadding() : Padding;
 }
 
 void UScrollBoxSlot::SetPadding(FMargin InPadding)
@@ -35,6 +45,11 @@ void UScrollBoxSlot::SetPadding(FMargin InPadding)
 	}
 }
 
+EHorizontalAlignment UScrollBoxSlot::GetHorizontalAlignment() const
+{
+	return Slot ? Slot->GetHorizontalAlignment() : HorizontalAlignment.GetValue();
+}
+
 void UScrollBoxSlot::SetHorizontalAlignment(EHorizontalAlignment InHorizontalAlignment)
 {
 	HorizontalAlignment = InHorizontalAlignment;
@@ -44,6 +59,11 @@ void UScrollBoxSlot::SetHorizontalAlignment(EHorizontalAlignment InHorizontalAli
 	}
 }
 
+EVerticalAlignment UScrollBoxSlot::GetVerticalAlignment() const
+{
+	return Slot ? Slot->GetVerticalAlignment() : VerticalAlignment.GetValue();
+}
+
 void UScrollBoxSlot::SetVerticalAlignment(EVerticalAlignment InVerticalAlignment)
 {
 	VerticalAlignment = InVerticalAlignment;
@@ -52,12 +72,15 @@ void UScrollBoxSlot::SetVerticalAlignment(EVerticalAlignment InVerticalAlignment
 		Slot->SetVerticalAlignment(InVerticalAlignment);
 	}
 }
+PRAGMA_ENABLE_DEPRECATION_WARNINGS
 
 void UScrollBoxSlot::SynchronizeProperties()
 {
+PRAGMA_DISABLE_DEPRECATION_WARNINGS
 	SetPadding(Padding);
 	SetHorizontalAlignment(HorizontalAlignment);
 	SetVerticalAlignment(VerticalAlignment);
+PRAGMA_ENABLE_DEPRECATION_WARNINGS
 }
 
 void UScrollBoxSlot::ReleaseSlateResources(bool bReleaseChildren)
diff --git a/Engine/Source/Runtime/UMG/Public/Components/ScaleBoxSlot.h b/Engine/Source/Runtime/UMG/Public/Components/ScaleBoxSlot.h
--- a/Engine/Source/Runtime/UMG/Public/Components/ScaleBoxSlot.h
+++ b/Engine/Source/Runtime/UMG/Public/Components/ScaleBoxSlot.h
@@ -19,33 +19,42 @@ class UMG_API UScaleBoxSlot : public UPanelSlot
 	GENERATED_UCLASS_BODY()
 
 public:
-	
+#if WITH_EDITORONLY_DATA
 	/** The padding area between the slot and the content it contains. */
-	UPROPERTY(EditAnywhere, BlueprintReadOnly, Category="Layout|ScaleBox Slot")
-	FMargin Padding;
+	UE_DEPRECATED(5.1, "Padding is not used.")
+	UPROPERTY()
+	FMargin Padding_DEPRECATED;
+#endif
 
 	/** The alignment of the object horizontally. */
-	UPROPERTY(EditAnywhere, BlueprintReadOnly, Category="Layout|ScaleBox Slot")
+	UE_DEPRECATED(5.1, "Direct access to HorizontalAlignment is deprecated. Please use the getter or setter.")
+	UPROPERTY(EditAnywhere, BlueprintReadOnly, Setter, Category="Layout|ScaleBox Slot")
 	TEnumAsByte<EHorizontalAlignment> HorizontalAlignment;
 
 	/** The alignment of the object vertically. */
-	UPROPERTY(EditAnywhere, BlueprintReadOnly, Category="Layout|ScaleBox Slot" )
+	UE_DEPRECATED(5.1, "Direct access to VerticalAlignment is deprecated. Please use the getter or setter.")
+	UPROPERTY(EditAnywhere, BlueprintReadOnly, Setter, Category="Layout|ScaleBox Slot" )
 	TEnumAsByte<EVerticalAlignment> VerticalAlignment;
 
 public:
-
-	UFUNCTION(BlueprintCallable, Category="Layout|ScaleBox Slot")
+	UE_DEPRECATED(5.1, "SetPadding is deprecated")
+	UFUNCTION(BlueprintCallable, Category="Layout|ScaleBox Slot", meta = (DeprecatedFunction, DeprecatedMessage = "The function is deprecated - There is no padding on ScaleBox."))
 	void SetPadding(FMargin InPadding);
 
+	EHorizontalAlignment GetHorizontalAlignment() const;
+
 	UFUNCTION(BlueprintCallable, Category="Layout|ScaleBox Slot")
 	void SetHorizontalAlignment(EHorizontalAlignment InHorizontalAlignment);
 
+
+	EVerticalAlignment GetVerticalAlignment() const;
+
 	UFUNCTION(BlueprintCallable, Category="Layout|ScaleBox Slot")
 	void SetVerticalAlignment(EVerticalAlignment InVerticalAlignment);
 
-	// UPanelSlot interface
+	//~ UPanelSlot interface
 	virtual void SynchronizeProperties() override;
-	// End of UPanelSlot interface
+	//~ End of UPanelSlot interface
 
 	/** Builds the underlying slot for the slate ScaleBox. */
 	void BuildSlot(TSharedRef<SScaleBox> InScaleBox);
@@ -53,7 +62,6 @@ public:
 	virtual void ReleaseSlateResources(bool bReleaseChildren) override;
 
 private:
-
 	/** A pointer to the button to allow us to adjust the size, padding...etc at runtime. */
 	TWeakPtr<SScaleBox> ScaleBox;
 };
diff --git a/Engine/Source/Runtime/UMG/Public/Components/ScrollBoxSlot.h b/Engine/Source/Runtime/UMG/Public/Components/ScrollBoxSlot.h
--- a/Engine/Source/Runtime/UMG/Public/Components/ScrollBoxSlot.h
+++ b/Engine/Source/Runtime/UMG/Public/Components/ScrollBoxSlot.h
@@ -20,33 +20,42 @@ class UMG_API UScrollBoxSlot : public UPanelSlot
 public:
 	
 	/** The padding area between the slot and the content it contains. */
+	UE_DEPRECATED(5.1, "Direct access to Padding is deprecated. Please use the getter or setter.")
 	UPROPERTY(EditAnywhere, BlueprintReadOnly, Category="Layout|ScrollBox Slot")
 	FMargin Padding;
 
 	/** The alignment of the object horizontally. */
+	UE_DEPRECATED(5.1, "Direct access to HorizontalAlignment is deprecated. Please use the getter or setter.")
 	UPROPERTY(EditAnywhere, BlueprintReadOnly, Category="Layout|ScrollBox Slot")
 	TEnumAsByte<EHorizontalAlignment> HorizontalAlignment;
 
 	/** The alignment of the object vertically. */
+	UE_DEPRECATED(5.1, "Direct access to VerticalAlignment is deprecated. Please use the getter or setter.")
 	UPROPERTY(EditAnywhere, BlueprintReadOnly, Category = "Layout|ScrollBox Slot")
 	TEnumAsByte<EVerticalAlignment> VerticalAlignment;
 
 public:
 
+	FMargin GetPadding() const;
+
 	UFUNCTION(BlueprintCallable, Category="Layout|ScrollBox Slot")
 	void SetPadding(FMargin InPadding);
 
+	EHorizontalAlignment GetHorizontalAlignment() const;
+
 	UFUNCTION(BlueprintCallable, Category="Layout|ScrollBox Slot")
 	void SetHorizontalAlignment(EHorizontalAlignment InHorizontalAlignment);
 
+	EVerticalAlignment GetVerticalAlignment() const;
+
 	UFUNCTION(BlueprintCallable, Category = "Layout|ScrollBox Slot")
 	void SetVerticalAlignment(EVerticalAlignment InVerticalAlignment);
 
 public:
 
-	// UPanelSlot interface
+	//~ UPanelSlot interface
 	virtual void SynchronizeProperties() override;
-	// End of UPanelSlot interface
+	//~ End of UPanelSlot interface
 
 	/** Builds the underlying FSlot for the Slate layout panel. */
 	void BuildSlot(TSharedRef<SScrollBox> ScrollBox);
@@ -54,7 +63,6 @@ public:
 	virtual void ReleaseSlateResources(bool bReleaseChildren) override;
 
 private:
-	//TODO UMG Slots should hold weak or shared refs to slots.
 
 	/** A raw pointer to the slot to allow us to adjust the size, padding...etc at runtime. */
 	SScrollBox::FSlot* Slot;
