Prevent null entries from being added to a ListView via SetListItems

diff --git a/Engine/Source/Runtime/UMG/Public/Components/ListView.h b/Engine/Source/Runtime/UMG/Public/Components/ListView.h
--- a/Engine/Source/Runtime/UMG/Public/Components/ListView.h
+++ b/Engine/Source/Runtime/UMG/Public/Components/ListView.h
@@ -39,7 +39,13 @@ public:
 	void SetListItems(const TArray<ItemObjectT, AllocatorType>& InListItems)
 	{
 		ClearListItems();
-		ListItems.Append(InListItems);
+		for (const ItemObjectT& ListItem : InListItems)
+		{
+			if (ListItem != nullptr)
+			{
+				ListItems.Add(ListItem);
+			}
+		}
 
 		OnItemsChanged(ListItems, TArray<UObject*>());
 
