UListView::SetListItems optimizations

diff --git a/Engine/Source/Runtime/UMG/Public/Components/ListView.h b/Engine/Source/Runtime/UMG/Public/Components/ListView.h
--- a/Engine/Source/Runtime/UMG/Public/Components/ListView.h
+++ b/Engine/Source/Runtime/UMG/Public/Components/ListView.h
@@ -39,7 +39,8 @@ public:
 	void SetListItems(const TArray<ItemObjectT, AllocatorType>& InListItems)
 	{
 		ClearListItems();
-		for (const ItemObjectT& ListItem : InListItems)
+		ListItems.Reserve(InListItems.Num());
+		for (const ItemObjectT ListItem : InListItems)
 		{
 			if (ListItem != nullptr)
 			{
