Add accessors to OverlaySlot

diff --git a/Engine/Plugins/Experimental/CommonUI/Source/CommonUI/Private/CommonVisibilitySwitcherSlot.cpp b/Engine/Plugins/Experimental/CommonUI/Source/CommonUI/Private/CommonVisibilitySwitcherSlot.cpp
--- a/Engine/Plugins/Experimental/CommonUI/Source/CommonUI/Private/CommonVisibilitySwitcherSlot.cpp
+++ b/Engine/Plugins/Experimental/CommonUI/Source/CommonUI/Private/CommonVisibilitySwitcherSlot.cpp
@@ -8,17 +8,17 @@
 UCommonVisibilitySwitcherSlot::UCommonVisibilitySwitcherSlot(const FObjectInitializer& Initializer)
 	: Super(Initializer)
 {
-	HorizontalAlignment = EHorizontalAlignment::HAlign_Fill;
-	VerticalAlignment = EVerticalAlignment::VAlign_Fill;
+	SetHorizontalAlignment(EHorizontalAlignment::HAlign_Fill);
+	SetVerticalAlignment(EVerticalAlignment::VAlign_Fill);
 }
 
 void UCommonVisibilitySwitcherSlot::BuildSlot(TSharedRef<SOverlay> Overlay)
 {
 	Overlay->AddSlot()
 		.Expose(Slot)
-		.Padding(Padding)
-		.HAlign(HorizontalAlignment)
-		.VAlign(VerticalAlignment)
+		.Padding(GetPadding())
+		.HAlign(GetHorizontalAlignment())
+		.VAlign(GetVerticalAlignment())
 		[
 			SAssignNew(VisibilityBox, SBox)
 			[
diff --git a/Engine/Source/Runtime/UMG/Private/Components/OverlaySlot.cpp b/Engine/Source/Runtime/UMG/Private/Components/OverlaySlot.cpp
--- a/Engine/Source/Runtime/UMG/Private/Components/OverlaySlot.cpp
+++ b/Engine/Source/Runtime/UMG/Private/Components/OverlaySlot.cpp
@@ -6,6 +6,7 @@
 /////////////////////////////////////////////////////
 // UOverlaySlot
 
+PRAGMA_DISABLE_DEPRECATION_WARNINGS
 UOverlaySlot::UOverlaySlot(const FObjectInitializer& ObjectInitializer)
 	: Super(ObjectInitializer)
 {
@@ -42,6 +43,11 @@ void UOverlaySlot::SetPadding(FMargin InPadding)
 	}
 }
 
+FMargin UOverlaySlot::GetPadding() const
+{
+	return Padding;
+}
+
 void UOverlaySlot::SetHorizontalAlignment(EHorizontalAlignment InHorizontalAlignment)
 {
 	HorizontalAlignment = InHorizontalAlignment;
@@ -51,6 +57,11 @@ void UOverlaySlot::SetHorizontalAlignment(EHorizontalAlignment InHorizontalAlign
 	}
 }
 
+EHorizontalAlignment UOverlaySlot::GetHorizontalAlignment() const
+{
+	return HorizontalAlignment;
+}
+
 void UOverlaySlot::SetVerticalAlignment(EVerticalAlignment InVerticalAlignment)
 {
 	VerticalAlignment = InVerticalAlignment;
@@ -60,9 +71,16 @@ void UOverlaySlot::SetVerticalAlignment(EVerticalAlignment InVerticalAlignment)
 	}
 }
 
+EVerticalAlignment UOverlaySlot::GetVerticalAlignment() const
+{
+	return VerticalAlignment;
+}
+
 void UOverlaySlot::SynchronizeProperties()
 {
 	SetPadding(Padding);
 	SetHorizontalAlignment(HorizontalAlignment);
 	SetVerticalAlignment(VerticalAlignment);
 }
+
+PRAGMA_ENABLE_DEPRECATION_WARNINGS
diff --git a/Engine/Source/Runtime/UMG/Public/Components/OverlaySlot.h b/Engine/Source/Runtime/UMG/Public/Components/OverlaySlot.h
--- a/Engine/Source/Runtime/UMG/Public/Components/OverlaySlot.h
+++ b/Engine/Source/Runtime/UMG/Public/Components/OverlaySlot.h
@@ -27,28 +27,43 @@ protected:
 public:
 	
 	/** The padding area between the slot and the content it contains. */
-	UPROPERTY(EditAnywhere, BlueprintReadOnly, Category="Layout|Overlay Slot")
+	UE_DEPRECATED(5.1, "Direct access to Padding is deprecated. Please use the getter or setter.")
+	UPROPERTY(EditAnywhere, BlueprintReadWrite, Getter, Setter, BlueprintSetter = SetPadding, Category="Layout|Overlay Slot")
 	FMargin Padding;
 
 	/** The alignment of the object horizontally. */
-	UPROPERTY(EditAnywhere, BlueprintReadOnly, Category="Layout|Overlay Slot")
+	UE_DEPRECATED(5.1, "Direct access to HorizontalAlignment is deprecated. Please use the getter or setter.")
+	UPROPERTY(EditAnywhere, BlueprintReadWrite, Getter, Setter, BlueprintSetter = SetHorizontalAlignment, Category="Layout|Overlay Slot")
 	TEnumAsByte<EHorizontalAlignment> HorizontalAlignment;
 
 	/** The alignment of the object vertically. */
-	UPROPERTY(EditAnywhere, BlueprintReadOnly, Category="Layout|Overlay Slot")
+	UE_DEPRECATED(5.1, "Direct access to VerticalAlignment is deprecated. Please use the getter or setter.")
+	UPROPERTY(EditAnywhere, BlueprintReadWrite, Getter, Setter, BlueprintSetter = SetVerticalAlignment, Category="Layout|Overlay Slot")
 	TEnumAsByte<EVerticalAlignment> VerticalAlignment;
 
 public:
 
+	/** Set padding area between the slot and the content it contains. */
 	UFUNCTION(BlueprintCallable, Category="Layout|Overlay Slot")
 	void SetPadding(FMargin InPadding);
 
+	/** Get padding area between the slot and the content it contains. */
+	FMargin GetPadding() const;
+
+	/** Set the alignment of the object horizontally. */
 	UFUNCTION(BlueprintCallable, Category="Layout|Overlay Slot")
 	void SetHorizontalAlignment(EHorizontalAlignment InHorizontalAlignment);
+	
+	/** Get the alignment of the object horizontally. */
+	EHorizontalAlignment GetHorizontalAlignment() const;
 
+	/** Set the alignment of the object vertically. */
 	UFUNCTION(BlueprintCallable, Category="Layout|Overlay Slot")
 	void SetVerticalAlignment(EVerticalAlignment InVerticalAlignment);
 
+	/** Get the alignment of the object vertically. */
+	EVerticalAlignment GetVerticalAlignment() const;
+
 public:
 
 	// UPanelSlot interface
