Add accessors to BackgroundBlur Widget

diff --git a/Engine/Source/Runtime/UMG/Private/Components/BackgroundBlur.cpp b/Engine/Source/Runtime/UMG/Private/Components/BackgroundBlur.cpp
--- a/Engine/Source/Runtime/UMG/Private/Components/BackgroundBlur.cpp
+++ b/Engine/Source/Runtime/UMG/Private/Components/BackgroundBlur.cpp
@@ -52,6 +52,8 @@ TSharedRef<SWidget> UBackgroundBlur::RebuildWidget()
 	return MyBackgroundBlur.ToSharedRef();
 }
 
+PRAGMA_DISABLE_DEPRECATION_WARNINGS
+
 void UBackgroundBlur::SynchronizeProperties()
 {
 	Super::SynchronizeProperties();
@@ -103,6 +105,11 @@ void UBackgroundBlur::SetPadding(FMargin InPadding)
 	}
 }
 
+FMargin UBackgroundBlur::GetPadding() const
+{
+	return Padding;
+}
+
 void UBackgroundBlur::SetHorizontalAlignment(EHorizontalAlignment InHorizontalAlignment)
 {
 	HorizontalAlignment = InHorizontalAlignment;
@@ -112,6 +119,11 @@ void UBackgroundBlur::SetHorizontalAlignment(EHorizontalAlignment InHorizontalAl
 	}
 }
 
+EHorizontalAlignment UBackgroundBlur::GetHorizontalAlignment() const
+{
+	return HorizontalAlignment;
+}
+
 void UBackgroundBlur::SetVerticalAlignment(EVerticalAlignment InVerticalAlignment)
 {
 	VerticalAlignment = InVerticalAlignment;
@@ -121,6 +133,11 @@ void UBackgroundBlur::SetVerticalAlignment(EVerticalAlignment InVerticalAlignmen
 	}
 }
 
+EVerticalAlignment UBackgroundBlur::GetVerticalAlignment() const
+{
+	return VerticalAlignment;
+}
+
 void UBackgroundBlur::SetApplyAlphaToBlur(bool bInApplyAlphaToBlur)
 {
 	bApplyAlphaToBlur = bInApplyAlphaToBlur;
@@ -130,15 +147,41 @@ void UBackgroundBlur::SetApplyAlphaToBlur(bool bInApplyAlphaToBlur)
 	}
 }
 
+bool UBackgroundBlur::GetApplyAlphaToBlur() const
+{
+	return bApplyAlphaToBlur;
+}
+
+void UBackgroundBlur::SetOverrideAutoRadiusCalculation(bool InOverrideAutoRadiusCalculation)
+{
+	bOverrideAutoRadiusCalculation = InOverrideAutoRadiusCalculation;
+	if (MyBackgroundBlur.IsValid())
+	{
+		// When set to false, it needs to reset the optional value and it will use the Blur strength to calculate the blur radius
+		MyBackgroundBlur->SetBlurRadius(bOverrideAutoRadiusCalculation ? BlurRadius : TOptional<int32>());
+	}	
+}
+
+bool UBackgroundBlur::GetOverrideAutoRadiusCalculation() const
+{
+	return bOverrideAutoRadiusCalculation;
+}
+
 void UBackgroundBlur::SetBlurRadius(int32 InBlurRadius)
 {
 	BlurRadius = InBlurRadius;
 	if (MyBackgroundBlur.IsValid())
 	{
 		MyBackgroundBlur->SetBlurRadius(InBlurRadius);
+		bOverrideAutoRadiusCalculation = true;
 	}
 }
 
+int32 UBackgroundBlur::GetBlurRadius() const
+{
+	return BlurRadius;
+}
+
 void UBackgroundBlur::SetBlurStrength(float InStrength)
 {
 	BlurStrength = InStrength;
@@ -148,6 +191,11 @@ void UBackgroundBlur::SetBlurStrength(float InStrength)
 	}
 }
 
+float UBackgroundBlur::GetBlurStrength() const
+{
+	return BlurStrength;
+}
+
 void UBackgroundBlur::SetCornerRadius(FVector4 InCornerRadius)
 {
 	CornerRadius = InCornerRadius;
@@ -157,6 +205,11 @@ void UBackgroundBlur::SetCornerRadius(FVector4 InCornerRadius)
 	}
 }
 
+FVector4 UBackgroundBlur::GetCornerRadius() const
+{
+	return CornerRadius;
+}
+
 void UBackgroundBlur::SetLowQualityFallbackBrush(const FSlateBrush& InBrush)
 {
 	LowQualityFallbackBrush = InBrush;
@@ -166,6 +219,13 @@ void UBackgroundBlur::SetLowQualityFallbackBrush(const FSlateBrush& InBrush)
 	}
 }
 
+FSlateBrush UBackgroundBlur::GetLowQualityFallbackBrush() const
+{
+	return LowQualityFallbackBrush;
+}
+
+PRAGMA_ENABLE_DEPRECATION_WARNINGS
+
 void UBackgroundBlur::Serialize(FArchive& Ar)
 {
 	Super::Serialize(Ar);
diff --git a/Engine/Source/Runtime/UMG/Public/Components/BackgroundBlur.h b/Engine/Source/Runtime/UMG/Public/Components/BackgroundBlur.h
--- a/Engine/Source/Runtime/UMG/Public/Components/BackgroundBlur.h
+++ b/Engine/Source/Runtime/UMG/Public/Components/BackgroundBlur.h
@@ -19,43 +19,51 @@ class UMG_API UBackgroundBlur : public UContentWidget
 
 public:
 	/** The padding area between the slot and the content it contains. */
-	UPROPERTY(EditAnywhere, BlueprintReadOnly, Category="Content")
+	UE_DEPRECATED(5.1, "Direct access to Padding is deprecated. Please use the getter or setter.")
+	UPROPERTY(EditAnywhere, BlueprintReadWrite, Getter, Setter, BlueprintSetter=SetPadding, Category = Content)
 	FMargin Padding;
 
 	/** The alignment of the content horizontally. */
-	UPROPERTY(EditAnywhere, BlueprintReadOnly, Category="Content")
+	UE_DEPRECATED(5.1, "Direct access to HorizontalAlignment is deprecated. Please use the getter or setter.")
+	UPROPERTY(EditAnywhere, BlueprintReadWrite, Getter, Setter, BlueprintSetter=SetHorizontalAlignment, Category = Content)
 	TEnumAsByte<EHorizontalAlignment> HorizontalAlignment;
 
 	/** The alignment of the content vertically. */
-	UPROPERTY(EditAnywhere, BlueprintReadOnly, Category="Content")
+	UE_DEPRECATED(5.1, "Direct access to VerticalAlignment is deprecated. Please use the getter or setter.")
+	UPROPERTY(EditAnywhere, BlueprintReadWrite, Getter, Setter, BlueprintSetter=SetVerticalAlignment, Category = Content)
 	TEnumAsByte<EVerticalAlignment> VerticalAlignment;
 
 	/** True to modulate the strength of the blur based on the widget alpha. */
-	UPROPERTY(EditAnywhere, BlueprintReadOnly, Category=Appearance)
+	UE_DEPRECATED(5.1, "Direct access to bApplyAlphaToBlur is deprecated. Please use the getter or setter.")
+	UPROPERTY(EditAnywhere, BlueprintReadWrite, Getter=GetApplyAlphaToBlur, Setter=SetApplyAlphaToBlur, BlueprintSetter=SetApplyAlphaToBlur, Category = Content)
 	bool bApplyAlphaToBlur;
 
 	/**
-	 * How blurry the background is.  Larger numbers mean more blurry but will result in larger runtime cost on the gpu.
+	 * How blurry the background is.  Larger numbers mean more blurry but will result in larger runtime cost on the GPU.
 	 */
-	UPROPERTY(EditAnywhere, BlueprintReadOnly, Category=Appearance, meta=(ClampMin=0,ClampMax=100))
+	UE_DEPRECATED(5.1, "Direct access to BlurStrength is deprecated. Please use the getter or setter.")
+	UPROPERTY(EditAnywhere, BlueprintReadWrite, Getter, Setter, BlueprintSetter=SetBlurStrength, Category = Appearance, meta = (ClampMin = 0, ClampMax = 100))
 	float BlurStrength;
 
-	/** Whether or not the radius should be computed automatically or if it should use the radius */
-	UPROPERTY()
+	/** When OverrideAutoRadiusCalculation is set to true, BlurRadius is used for the radius of the blur. When false, it's automatically calculated using the BlurStength value. */
+	UE_DEPRECATED(5.1, "Direct access to bOverrideAutoRadiusCalculation is deprecated. Please use the getter or setter.")
+	UPROPERTY(Getter = GetOverrideAutoRadiusCalculation, Setter = SetOverrideAutoRadiusCalculation)
 	bool bOverrideAutoRadiusCalculation;
 
 	/**
 	 * This is the number of pixels which will be weighted in each direction from any given pixel when computing the blur
-	 * A larger value is more costly but allows for stronger blurs.  
+	 * A larger value is more costly but allows for stronger blurs.
 	 */
-	UPROPERTY(EditAnywhere, AdvancedDisplay, BlueprintReadOnly, Category=Appearance, meta=(ClampMin=0, ClampMax=255, EditCondition="bOverrideAutoRadiusCalculation"))
+	UE_DEPRECATED(5.1, "Direct access to BlurRadius is deprecated. Please use the getter or setter.")
+	UPROPERTY(EditAnywhere, AdvancedDisplay, BlueprintReadWrite, Getter, Setter, BlueprintSetter=SetBlurRadius, Category = Appearance, meta = (ClampMin = 0, ClampMax = 255, EditCondition = "bOverrideAutoRadiusCalculation"))
 	int32 BlurRadius;
 
-		/**
+	/**
 	 * This is the number of pixels which will be weighted in each direction from any given pixel when computing the blur
 	 * A larger value is more costly but allows for stronger blurs.  
 	 */
-	UPROPERTY(EditAnywhere, AdvancedDisplay, BlueprintReadOnly, Category = Appearance)
+	UE_DEPRECATED(5.1, "Direct access to CornerRadius is deprecated. Please use the getter or setter.")
+	UPROPERTY(EditAnywhere, AdvancedDisplay, BlueprintReadWrite, Getter, Setter, BlueprintSetter=SetCornerRadius, Category = Appearance)
 	FVector4 CornerRadius;
 
 	/**
@@ -63,7 +71,8 @@ public:
 	 * You can enable low quality mode for background blurs by setting the cvar Slate.ForceBackgroundBlurLowQualityOverride to 1. 
 	 * This is usually done in the project's scalability settings
 	 */
-	UPROPERTY(EditAnywhere, BlueprintReadOnly, Category=Appearance)
+	UE_DEPRECATED(5.1, "Direct access to LowQualityFallbackBrush is deprecated. Please use the getter or setter.")
+	UPROPERTY(EditAnywhere, BlueprintReadWrite, Getter, Setter, BlueprintSetter=SetLowQualityFallbackBrush, Category = Appearance)
 	FSlateBrush LowQualityFallbackBrush;
 
 public:
@@ -77,27 +86,47 @@ public:
 	UFUNCTION(BlueprintCallable, Category="Appearance")
 	void SetPadding(FMargin InPadding);
 
+	FMargin GetPadding() const;
+
 	UFUNCTION(BlueprintCallable, Category="Appearance")
 	void SetHorizontalAlignment(EHorizontalAlignment InHorizontalAlignment);
 
+	EHorizontalAlignment GetHorizontalAlignment() const;
+
 	UFUNCTION(BlueprintCallable, Category="Appearance")
 	void SetVerticalAlignment(EVerticalAlignment InVerticalAlignment);
 
+	EVerticalAlignment GetVerticalAlignment() const;
+
 	UFUNCTION(BlueprintCallable, Category = "Appearance")
 	void SetApplyAlphaToBlur(bool bInApplyAlphaToBlur);
 
+	bool GetApplyAlphaToBlur() const;
+
+	void SetOverrideAutoRadiusCalculation(bool InOverrideAutoRadiusCalculation);
+
+	bool GetOverrideAutoRadiusCalculation() const;
+
 	UFUNCTION(BlueprintCallable, Category = "Appearance")
 	void SetBlurRadius(int32 InBlurRadius);
 
+	int32 GetBlurRadius() const;
+
 	UFUNCTION(BlueprintCallable, Category = "Appearance")
 	virtual void SetBlurStrength(float InStrength);
 
+	float GetBlurStrength() const;
+
 	UFUNCTION(BlueprintCallable, Category = "Appearance")
 	virtual void SetCornerRadius(FVector4 InCornerRadius);
+	
+	FVector4 GetCornerRadius() const;
 
 	UFUNCTION(BlueprintCallable, Category = "Appearance")
 	void SetLowQualityFallbackBrush(const FSlateBrush& InBrush);
 
+	FSlateBrush GetLowQualityFallbackBrush() const;
+
 	/** UObject interface */
 	virtual void Serialize(FArchive& Ar) override;
 	virtual void PostLoad() override;
