Adds validity check when dereferencing PanelWidget in DynamicEntryBoxBase in AddEntryChild.

diff --git a/Engine/Source/Runtime/UMG/Private/Components/DynamicEntryBoxBase.cpp b/Engine/Source/Runtime/UMG/Private/Components/DynamicEntryBoxBase.cpp
--- a/Engine/Source/Runtime/UMG/Private/Components/DynamicEntryBoxBase.cpp
+++ b/Engine/Source/Runtime/UMG/Private/Components/DynamicEntryBoxBase.cpp
@@ -384,67 +384,73 @@ void UDynamicEntryBoxBase::AddEntryChild(UUserWidget& ChildWidget)
 	}
 	else if (EntryBoxType == EDynamicBoxType::Overlay)
 	{
-		const int32 ChildIdx = MyPanelWidget->GetChildren()->Num();
+		if (MyPanelWidget.IsValid())
+		{
+			const int32 ChildIdx = MyPanelWidget->GetChildren()->Num();
 
-		EHorizontalAlignment HAlign = EntryHorizontalAlignment;
-		EVerticalAlignment VAlign = EntryVerticalAlignment;
+			EHorizontalAlignment HAlign = EntryHorizontalAlignment;
+			EVerticalAlignment VAlign = EntryVerticalAlignment;
 
-		FVector2D TargetSpacing = FVector2D::ZeroVector;
-		if (SpacingPattern.Num() > 0)
-		{
-			for (int32 CountIdx = 0; CountIdx < ChildIdx; ++CountIdx)
+			FVector2D TargetSpacing = FVector2D::ZeroVector;
+			if (SpacingPattern.Num() > 0)
 			{
-				const int32 PatternIdx = CountIdx % SpacingPattern.Num();
-				TargetSpacing += SpacingPattern[PatternIdx];
+				for (int32 CountIdx = 0; CountIdx < ChildIdx; ++CountIdx)
+				{
+					const int32 PatternIdx = CountIdx % SpacingPattern.Num();
+					TargetSpacing += SpacingPattern[PatternIdx];
+				}
+			}
+			else
+			{
+				TargetSpacing = EntrySpacing * ChildIdx;
+				HAlign = EntrySpacing.X >= 0.f ? EHorizontalAlignment::HAlign_Left : EHorizontalAlignment::HAlign_Right;
+				VAlign = EntrySpacing.Y >= 0.f ? EVerticalAlignment::VAlign_Top : EVerticalAlignment::VAlign_Bottom;
 			}
-		}
-		else
-		{
-			TargetSpacing = EntrySpacing * ChildIdx;
-			HAlign = EntrySpacing.X >= 0.f ? EHorizontalAlignment::HAlign_Left : EHorizontalAlignment::HAlign_Right;
-			VAlign = EntrySpacing.Y >= 0.f ? EVerticalAlignment::VAlign_Top : EVerticalAlignment::VAlign_Bottom;
-		}
 		
-		StaticCastSharedPtr<SOverlay>(MyPanelWidget)->AddSlot()
-			.HAlign(HAlign)
-			.VAlign(VAlign)
-			.Padding(BuildEntryPadding(TargetSpacing))
-			[
-				ChildWidget.TakeWidget()
-			];
-	}
-	else
-	{
-		const bool bIsHBox = EntryBoxType == EDynamicBoxType::Horizontal;
-		const bool bIsFirstChild = MyPanelWidget->GetChildren()->Num() == 0;
-
-		FMargin Padding;
-		Padding.Top = bIsHBox || bIsFirstChild ? 0.f : EntrySpacing.Y;
-		Padding.Left = bIsHBox && !bIsFirstChild ? EntrySpacing.X : 0.f;
-
-		if (bIsHBox)
-		{
-			StaticCastSharedPtr<SHorizontalBox>(MyPanelWidget)->AddSlot()
-				.MaxWidth(MaxElementSize)
-				.HAlign(EntryHorizontalAlignment)
-				.VAlign(EntryVerticalAlignment)
-				.SizeParam(UWidget::ConvertSerializedSizeParamToRuntime(EntrySizeRule))
-				.Padding(Padding)
+			StaticCastSharedPtr<SOverlay>(MyPanelWidget)->AddSlot()
+				.HAlign(HAlign)
+				.VAlign(VAlign)
+				.Padding(BuildEntryPadding(TargetSpacing))
 				[
 					ChildWidget.TakeWidget()
 				];
 		}
-		else
+	}
+	else
+	{
+		if (MyPanelWidget.IsValid())
 		{
-			StaticCastSharedPtr<SVerticalBox>(MyPanelWidget)->AddSlot()
-				.MaxHeight(MaxElementSize)
-				.HAlign(EntryHorizontalAlignment)
-				.VAlign(EntryVerticalAlignment)
-				.SizeParam(UWidget::ConvertSerializedSizeParamToRuntime(EntrySizeRule))
-				.Padding(Padding)
-				[
-					ChildWidget.TakeWidget()
-				];
+			const bool bIsHBox = EntryBoxType == EDynamicBoxType::Horizontal;
+			const bool bIsFirstChild = MyPanelWidget->GetChildren()->Num() == 0;
+
+			FMargin Padding;
+			Padding.Top = bIsHBox || bIsFirstChild ? 0.f : EntrySpacing.Y;
+			Padding.Left = bIsHBox && !bIsFirstChild ? EntrySpacing.X : 0.f;
+
+			if (bIsHBox)
+			{
+				StaticCastSharedPtr<SHorizontalBox>(MyPanelWidget)->AddSlot()
+					.MaxWidth(MaxElementSize)
+					.HAlign(EntryHorizontalAlignment)
+					.VAlign(EntryVerticalAlignment)
+					.SizeParam(UWidget::ConvertSerializedSizeParamToRuntime(EntrySizeRule))
+					.Padding(Padding)
+					[
+						ChildWidget.TakeWidget()
+					];
+			}
+			else
+			{
+				StaticCastSharedPtr<SVerticalBox>(MyPanelWidget)->AddSlot()
+					.MaxHeight(MaxElementSize)
+					.HAlign(EntryHorizontalAlignment)
+					.VAlign(EntryVerticalAlignment)
+					.SizeParam(UWidget::ConvertSerializedSizeParamToRuntime(EntrySizeRule))
+					.Padding(Padding)
+					[
+						ChildWidget.TakeWidget()
+					];
+			}
 		}
 	}
 }
