Text Blocks:  add way to set material in blueprint

diff --git a/Engine/Source/Runtime/UMG/Private/Components/RichTextBlock.cpp b/Engine/Source/Runtime/UMG/Private/Components/RichTextBlock.cpp
--- a/Engine/Source/Runtime/UMG/Private/Components/RichTextBlock.cpp
+++ b/Engine/Source/Runtime/UMG/Private/Components/RichTextBlock.cpp
@@ -281,6 +281,13 @@ void URichTextBlock::SetDefaultTextStyle(const FTextBlockStyle& InDefaultTextSty
 	ApplyUpdatedDefaultTextStyle();
 }
 
+void URichTextBlock::SetDefaultMaterial(UMaterialInterface* InMaterial)
+{
+	BeginDefaultStyleOverride();
+	DefaultTextStyleOverride.Font.FontMaterial = InMaterial;
+	ApplyUpdatedDefaultTextStyle();
+}
+
 void URichTextBlock::ClearAllDefaultStyleOverrides()
 {
 	if (bOverrideDefaultStyle)
diff --git a/Engine/Source/Runtime/UMG/Private/Components/TextBlock.cpp b/Engine/Source/Runtime/UMG/Private/Components/TextBlock.cpp
--- a/Engine/Source/Runtime/UMG/Private/Components/TextBlock.cpp
+++ b/Engine/Source/Runtime/UMG/Private/Components/TextBlock.cpp
@@ -197,6 +197,18 @@ void UTextBlock::SetTextOverflowPolicy(ETextOverflowPolicy InOverflowPolicy)
 	SynchronizeProperties();
 }
 
+void UTextBlock::SetFontMaterial(UMaterialInterface* InMaterial)
+{
+	Font.FontMaterial = InMaterial;
+	SetFont(Font);
+}
+
+void UTextBlock::SetFontOutlineMaterial(UMaterialInterface* InMaterial)
+{
+	Font.OutlineSettings.OutlineMaterial = InMaterial;
+	SetFont(Font);
+}
+
 UMaterialInstanceDynamic* UTextBlock::GetDynamicFontMaterial()
 {
 	if (Font.FontMaterial)
diff --git a/Engine/Source/Runtime/UMG/Public/Components/RichTextBlock.h b/Engine/Source/Runtime/UMG/Public/Components/RichTextBlock.h
--- a/Engine/Source/Runtime/UMG/Public/Components/RichTextBlock.h
+++ b/Engine/Source/Runtime/UMG/Public/Components/RichTextBlock.h
@@ -14,6 +14,7 @@
 class SRichTextBlock;
 class UDataTable;
 class UMaterialInstanceDynamic;
+class UMaterialInterface;
 class URichTextBlockDecorator;
 
 /** Simple struct for rich text styles */
@@ -111,6 +112,9 @@ public:
 	UFUNCTION(BlueprintCallable, Category = Appearance)
 	void SetDefaultTextStyle(const FTextBlockStyle& InDefaultTextStyle);
 
+	UFUNCTION(BlueprintCallable, Category = Appearance)
+	void SetDefaultMaterial(UMaterialInterface* InMaterial);
+
 	/** Remove all overrides made to the default text style and return to the style specified in the style set data table */
 	UFUNCTION()
 	void ClearAllDefaultStyleOverrides();
diff --git a/Engine/Source/Runtime/UMG/Public/Components/TextBlock.h b/Engine/Source/Runtime/UMG/Public/Components/TextBlock.h
--- a/Engine/Source/Runtime/UMG/Public/Components/TextBlock.h
+++ b/Engine/Source/Runtime/UMG/Public/Components/TextBlock.h
@@ -11,6 +11,7 @@
 #include "TextBlock.generated.h"
 
 class STextBlock;
+class UMaterialInterface;
 
 /**
  * A simple static text widget.
@@ -206,6 +207,12 @@ public:
 	UFUNCTION(BlueprintCallable, Category="Appearance")
 	void SetTextOverflowPolicy(ETextOverflowPolicy InOverflowPolicy);
 
+	UFUNCTION(BlueprintCallable, Category="Appearance")
+	void SetFontMaterial(UMaterialInterface* InMaterial);
+
+	UFUNCTION(BlueprintCallable, Category="Appearance")
+	void SetFontOutlineMaterial(UMaterialInterface* InMaterial);
+
 	/**  */
 	UFUNCTION(BlueprintCallable, Category="Appearance")
 	UMaterialInstanceDynamic* GetDynamicFontMaterial();
