UMG: Add the widget name to a warning in UBorder.
#rnx

diff --git a/Engine/Source/Runtime/UMG/Private/Components/Border.cpp b/Engine/Source/Runtime/UMG/Private/Components/Border.cpp
--- a/Engine/Source/Runtime/UMG/Private/Components/Border.cpp
+++ b/Engine/Source/Runtime/UMG/Private/Components/Border.cpp
@@ -5,6 +5,7 @@
 #include "Materials/MaterialInterface.h"
 #include "Engine/Texture2D.h"
 #include "Materials/MaterialInstanceDynamic.h"
+#include "Types/ReflectionMetadata.h"
 #include "Widgets/SNullWidget.h"
 #include "Widgets/DeclarativeSyntaxSupport.h"
 #include "Widgets/Layout/SBorder.h"
@@ -264,7 +265,7 @@ void UBorder::SetBrushFromMaterial(UMaterialInterface* Material)
 {
 	if (!Material)
 	{
-		UE_LOG(LogSlate, Log, TEXT("UBorder::SetBrushFromMaterial.  Incoming material is null"));
+		UE_LOG(LogSlate, Log, TEXT("UBorder::SetBrushFromMaterial. Incoming material is null. %s"), *GetPathName());
 	}
 
 	Background.SetResourceObject(Material);
