UMG: Remove obsolete warning in debug build when the UserWidget is constructed.

diff --git a/Engine/Source/Runtime/UMG/Private/WidgetBlueprintGeneratedClass.cpp b/Engine/Source/Runtime/UMG/Private/WidgetBlueprintGeneratedClass.cpp
--- a/Engine/Source/Runtime/UMG/Private/WidgetBlueprintGeneratedClass.cpp
+++ b/Engine/Source/Runtime/UMG/Private/WidgetBlueprintGeneratedClass.cpp
@@ -218,10 +218,6 @@ void UWidgetBlueprintGeneratedClass::InitializeWidgetStatic(UUserWidget* UserWid
 		CreatedWidgetTree = UserWidget->WidgetTree;
 	}
 
-#if !WITH_EDITOR && UE_BUILD_DEBUG
-	UE_LOG(LogUMG, Warning, TEXT("Widget Class %s - Slow Static Duplicate Object."), *InClass->GetName());
-#endif
-
 #if WITH_EDITOR
 	UserWidget->WidgetGeneratedBy = InClass->ClassGeneratedBy;
 #endif
