Editable Text Boxes:  fix missing includes

diff --git a/Engine/Source/Runtime/UMG/Private/Components/EditableText.cpp b/Engine/Source/Runtime/UMG/Private/Components/EditableText.cpp
--- a/Engine/Source/Runtime/UMG/Private/Components/EditableText.cpp
+++ b/Engine/Source/Runtime/UMG/Private/Components/EditableText.cpp
@@ -7,6 +7,7 @@
 #include "Widgets/Input/SEditableText.h"
 #include "Slate/SlateBrushAsset.h"
 #include "Styling/UMGCoreStyle.h"
+#include "Materials/MaterialInterface.h"
 
 #define LOCTEXT_NAMESPACE "UMG"
 
diff --git a/Engine/Source/Runtime/UMG/Private/Components/MultiLineEditableText.cpp b/Engine/Source/Runtime/UMG/Private/Components/MultiLineEditableText.cpp
--- a/Engine/Source/Runtime/UMG/Private/Components/MultiLineEditableText.cpp
+++ b/Engine/Source/Runtime/UMG/Private/Components/MultiLineEditableText.cpp
@@ -6,6 +6,7 @@
 #include "Widgets/DeclarativeSyntaxSupport.h"
 #include "Widgets/Text/SMultiLineEditableText.h"
 #include "Styling/UMGCoreStyle.h"
+#include "Materials/MaterialInterface.h"
 
 #define LOCTEXT_NAMESPACE "UMG"
 
diff --git a/Engine/Source/Runtime/UMG/Public/Components/EditableText.h b/Engine/Source/Runtime/UMG/Public/Components/EditableText.h
--- a/Engine/Source/Runtime/UMG/Public/Components/EditableText.h
+++ b/Engine/Source/Runtime/UMG/Public/Components/EditableText.h
@@ -14,6 +14,7 @@
 #include "EditableText.generated.h"
 
 class SEditableText;
+class UMaterialInterface;
 class USlateBrushAsset;
 class USlateWidgetStyleAsset;
 
diff --git a/Engine/Source/Runtime/UMG/Public/Components/MultiLineEditableText.h b/Engine/Source/Runtime/UMG/Public/Components/MultiLineEditableText.h
--- a/Engine/Source/Runtime/UMG/Public/Components/MultiLineEditableText.h
+++ b/Engine/Source/Runtime/UMG/Public/Components/MultiLineEditableText.h
@@ -11,6 +11,7 @@
 #include "Widgets/Text/ISlateEditableTextWidget.h"
 #include "MultiLineEditableText.generated.h"
 
+class UMaterialInterface;
 class SMultiLineEditableText;
 
 /**
