UMG: UUMGSequencePlayer::Stop no longer forcibly flushes the runner, and waits until it has been flushed before Finishing the animation

  - This is currently enabled by default, but controllable through the UMG.AsyncAnimationControlFlow cvar (false disables this new flow)
  - When an animation would usually be ticked by the animation tick manager, it now schedules an update on Stop, and queues a latent action to finish stopping before it completes the operation

diff --git a/Engine/Source/Runtime/MovieScene/Private/Evaluation/MovieSceneEvaluationTemplateInstance.cpp b/Engine/Source/Runtime/MovieScene/Private/Evaluation/MovieSceneEvaluationTemplateInstance.cpp
--- a/Engine/Source/Runtime/MovieScene/Private/Evaluation/MovieSceneEvaluationTemplateInstance.cpp
+++ b/Engine/Source/Runtime/MovieScene/Private/Evaluation/MovieSceneEvaluationTemplateInstance.cpp
@@ -155,6 +155,11 @@ void FMovieSceneRootEvaluationTemplateInstance::Finish(IMovieScenePlayer& Player
 		EntitySystemRunner.FinishInstance(RootInstanceHandle);
 	}
 
+	OnFinish();
+}
+
+void FMovieSceneRootEvaluationTemplateInstance::OnFinish()
+{
 	DirectorInstances.Reset();
 }
 
diff --git a/Engine/Source/Runtime/MovieScene/Public/Evaluation/MovieSceneEvaluationTemplateInstance.h b/Engine/Source/Runtime/MovieScene/Public/Evaluation/MovieSceneEvaluationTemplateInstance.h
--- a/Engine/Source/Runtime/MovieScene/Public/Evaluation/MovieSceneEvaluationTemplateInstance.h
+++ b/Engine/Source/Runtime/MovieScene/Public/Evaluation/MovieSceneEvaluationTemplateInstance.h
@@ -63,6 +63,11 @@ public:
 	 */
 	MOVIESCENE_API void Finish(IMovieScenePlayer& Player);
 
+	/**
+	 * Called when this template has finished evaluating to perform cleanup
+	 */
+	MOVIESCENE_API void OnFinish();
+
 	bool IsValid() const
 	{
 		return CompiledDataManager && WeakRootSequence.Get();
diff --git a/Engine/Source/Runtime/UMG/Private/Animation/UMGSequencePlayer.cpp b/Engine/Source/Runtime/UMG/Private/Animation/UMGSequencePlayer.cpp
--- a/Engine/Source/Runtime/UMG/Private/Animation/UMGSequencePlayer.cpp
+++ b/Engine/Source/Runtime/UMG/Private/Animation/UMGSequencePlayer.cpp
@@ -11,6 +11,18 @@
 
 extern TAutoConsoleVariable<bool> CVarUserWidgetUseParallelAnimation;
 
+namespace UE::UMG
+{
+
+	bool GAsyncAnimationControlFlow = true;
+	FAutoConsoleVariableRef CVarAsyncAnimationControlFlow(
+		TEXT("UMG.AsyncAnimationControlFlow"),
+		GAsyncAnimationControlFlow,
+		TEXT("(Default: true) Whether to perform animation control flow functions (Play, Pause, Stop etc) asynchronously.")
+	);
+
+} // namespace UE::UMG
+
 UUMGSequencePlayer::UUMGSequencePlayer(const FObjectInitializer& ObjectInitializer)
 	: Super(ObjectInitializer)
 {
@@ -20,6 +32,7 @@ UUMGSequencePlayer::UUMGSequencePlayer(const FObjectInitializer& ObjectInitializ
 	bRestoreState = false;
 	Animation = nullptr;
 	bIsEvaluating = false;
+	bIsStopping = false;
 	bCompleteOnPostEvaluation = false;
 	UserTag = NAME_None;
 }
@@ -58,6 +71,11 @@ UMovieSceneEntitySystemLinker* UUMGSequencePlayer::ConstructEntitySystemLinker()
 
 void UUMGSequencePlayer::Tick(float DeltaTime)
 {
+	if (bIsStopping)
+	{
+		return;
+	}
+
 	if ( PlayerStatus == EMovieScenePlayerStatus::Playing )
 	{
 		FFrameTime DeltaFrameTime = (bIsPlayingForward ? DeltaTime * PlaybackSpeed : -DeltaTime * PlaybackSpeed) * AnimationResolution;
@@ -328,35 +346,75 @@ void UUMGSequencePlayer::Stop()
 
 	PlayerStatus = EMovieScenePlayerStatus::Stopped;
 
+	UMovieSceneSequence* MovieSceneSequence = RootTemplateInstance.GetSequence(MovieSceneSequenceID::Root);
+	const bool bIsRunningWithTickManager = CVarUserWidgetUseParallelAnimation.GetValueOnGameThread();
+	const bool bIsSequenceBlocking = EnumHasAnyFlags(MovieSceneSequence->GetFlags(), EMovieSceneSequenceFlags::BlockingEvaluation);
+	const bool bIsAsync = bIsRunningWithTickManager && !bIsSequenceBlocking && UE::UMG::GAsyncAnimationControlFlow == true;
+
 	UUserWidget* Widget = UserWidget.Get();
 	UUMGSequenceTickManager* TickManager = Widget ? ToRawPtr(Widget->AnimationTickManager) : nullptr;
 
-	if (TickManager && RootTemplateInstance.IsValid())
+	TimeCursorPosition = FFrameTime(0);
+
+	if (!TickManager || !RootTemplateInstance.IsValid())
 	{
-		if (RootTemplateInstance.HasEverUpdated())
-		{
-			const FMovieSceneContext Context(FMovieSceneEvaluationRange(AbsolutePlaybackStart, AnimationResolution), PlayerStatus);
-			RootTemplateInstance.Evaluate(Context, *this);
-		}
-		else
+		HandleLatentStop();
+		return;
+	}
+
+	if (!RootTemplateInstance.HasEverUpdated())
+	{
+		TickManager->ClearLatentActions(this);
+		LatentActions.Empty();
+		HandleLatentStop();
+		return;
+	}
+
+	const FMovieSceneContext Context(FMovieSceneEvaluationRange(AbsolutePlaybackStart, AnimationResolution), PlayerStatus);
+
+	// Prevent any other updates
+	bIsStopping = true;
+
+	if (bIsAsync)
+	{
+
+		TickManager->AddLatentAction(FMovieSceneSequenceLatentActionDelegate::CreateUObject(this, &UUMGSequencePlayer::HandleLatentStop));
+		TickManager->GetRunner().QueueUpdate(Context, RootTemplateInstance.GetRootInstanceHandle());
+	}
+	else
+	{
+		RootTemplateInstance.Evaluate(Context, *this);
+		HandleLatentStop();
+	}
+}
+
+void UUMGSequencePlayer::HandleLatentStop()
+{
+	// Linker may not necessarily be valid here if we have already torn down the sequence
+	if (UMovieSceneEntitySystemLinker* Linker = RootTemplateInstance.GetEntitySystemLinker())
+	{
+		UE::MovieScene::FSequenceInstance* RootInstance = RootTemplateInstance.FindInstance(MovieSceneSequenceID::Root);
+		// Finish the root instance
+		if (RootInstance)
 		{
-			TickManager->ClearLatentActions(this);
-			LatentActions.Empty();
+			RootInstance->Finish(Linker);
 		}
-		RootTemplateInstance.Finish(*this);
 	}
 
+	RootTemplateInstance.OnFinish();
+
 	if (bRestoreState)
 	{
 		RestorePreAnimatedState();
 	}
 
+	bIsStopping = false;
+
 	UserWidget->OnAnimationFinishedPlaying(*this);
 	OnSequenceFinishedPlayingEvent.Broadcast(*this);
-
-	TimeCursorPosition = FFrameTime(0);
 }
 
+
 void UUMGSequencePlayer::SetNumLoopsToPlay(int32 InNumLoopsToPlay)
 {
 	if (PlayMode == EUMGSequencePlayMode::PingPong)
@@ -413,19 +471,7 @@ void UUMGSequencePlayer::PostEvaluation(const FMovieSceneContext& Context)
 		bCompleteOnPostEvaluation = false;
 
 		PlayerStatus = EMovieScenePlayerStatus::Stopped;
-		
-		if (RootTemplateInstance.IsValid())
-		{
-			RootTemplateInstance.Finish(*this);
-		}
-
-		if (bRestoreState)
-		{
-			RestorePreAnimatedState();
-		}
-
-		UserWidget->OnAnimationFinishedPlaying(*this);
-		OnSequenceFinishedPlayingEvent.Broadcast(*this);
+		HandleLatentStop();
 	}
 }
 
diff --git a/Engine/Source/Runtime/UMG/Public/Animation/UMGSequencePlayer.h b/Engine/Source/Runtime/UMG/Public/Animation/UMGSequencePlayer.h
--- a/Engine/Source/Runtime/UMG/Public/Animation/UMGSequencePlayer.h
+++ b/Engine/Source/Runtime/UMG/Public/Animation/UMGSequencePlayer.h
@@ -101,6 +101,8 @@ private:
 	void QueueLatentAction(FMovieSceneSequenceLatentActionDelegate Delegate);
 	void ApplyLatentActions();
 
+	void HandleLatentStop();
+
 	/** Animation being played */
 	UPROPERTY()
 	TObjectPtr<UWidgetAnimation> Animation;
@@ -160,6 +162,9 @@ private:
 	/** Set to true while evaluating to prevent reentrancy */
 	bool bIsEvaluating : 1;
 
+	/** Set to true when this player is in the process of being stopped */
+	bool bIsStopping : 1;
+
 	/** Set to true if we need to run the finishing logic in post-evaluation */
 	bool bCompleteOnPostEvaluation : 1;
 
