Added an option to UListViewBase to override FinishGeneratingEntry which enables fixing bugs when using list views in editor utlity widgets

diff --git a/Engine/Source/Runtime/UMG/Public/Components/ListViewBase.h b/Engine/Source/Runtime/UMG/Public/Components/ListViewBase.h
--- a/Engine/Source/Runtime/UMG/Public/Components/ListViewBase.h
+++ b/Engine/Source/Runtime/UMG/Public/Components/ListViewBase.h
@@ -563,7 +563,19 @@ protected:
 	//@todo DanH: Should probably have the events for native & BP built in up here - need to update existing binds to UListView's version
 	virtual void HandleListEntryHovered(UUserWidget& EntryWidget) {}
 	virtual void HandleListEntryUnhovered(UUserWidget& EntryWidget) {}
+	virtual	void FinishGeneratingEntry(UUserWidget& GeneratedEntry);
+   
+    /** Called when a row widget is generated for a list item */
+    UPROPERTY(BlueprintAssignable, Category = Events, meta = (DisplayName = "On Entry Generated"))
+    FOnListEntryGeneratedDynamic BP_OnEntryGenerated;
 
+	/**
+	* Normally these are processed by UListViewBase::FinishGeneratingEntry which uses World->GetTimerManager() to generate entries next frame
+	* However when for example using a listview in editor utility widgets a world there is not reliable and an alternative is to use
+	* FTSTicker::GetCoreTicker()
+	*/
+	TArray<TWeakObjectPtr<UUserWidget>> GeneratedEntriesToAnnounce;
+	
 	template <typename WidgetEntryT = UUserWidget, typename ObjectTableRowT = SObjectTableRow<UObject*>>
 	WidgetEntryT& GenerateTypedEntry(TSubclassOf<WidgetEntryT> WidgetClass, const TSharedRef<STableViewBase>& OwnerTable)
 	{
@@ -677,13 +689,7 @@ protected:
 	bool bAllowDragging = true;
 
 private:
-	void FinishGeneratingEntry(UUserWidget& GeneratedEntry);
-	void HandleAnnounceGeneratedEntries();
-
-private:
-	/** Called when a row widget is generated for a list item */
-	UPROPERTY(BlueprintAssignable, Category = Events, meta = (DisplayName = "On Entry Generated"))
-	FOnListEntryGeneratedDynamic BP_OnEntryGenerated;
+	virtual void HandleAnnounceGeneratedEntries();
 
 	/** Called when a row widget is released by the list (i.e. when it no longer represents a list item) */
 	UPROPERTY(BlueprintAssignable, Category = Events, meta = (DisplayName = "On Entry Released"))
@@ -701,8 +707,7 @@ private:
 	FUserWidgetPool EntryWidgetPool;
 
 	FTimerHandle EntryGenAnnouncementTimerHandle;
-	TArray<TWeakObjectPtr<UUserWidget>> GeneratedEntriesToAnnounce;
-
+	
 	FOnListEntryGenerated OnListEntryGeneratedEvent;
 	FOnEntryWidgetReleased OnEntryWidgetReleasedEvent;
 
