PR #9461: fixed order of parameters in KeyEvent (Contributed by [REDACTED_PERSON])

diff --git a/Engine/Source/Runtime/UMG/Private/Components/WidgetInteractionComponent.cpp b/Engine/Source/Runtime/UMG/Private/Components/WidgetInteractionComponent.cpp
--- a/Engine/Source/Runtime/UMG/Private/Components/WidgetInteractionComponent.cpp
+++ b/Engine/Source/Runtime/UMG/Private/Components/WidgetInteractionComponent.cpp
@@ -534,7 +534,7 @@ bool UWidgetInteractionComponent::PressKey(FKey Key, bool bRepeat)
 	uint32 KeyCode, CharCode;
 	GetKeyAndCharCodes(Key, bHasKeyCode, KeyCode, bHasCharCode, CharCode);
 
-	FKeyEvent KeyEvent(Key, ModifierKeys, VirtualUser->GetUserIndex(), bRepeat, KeyCode, CharCode);
+	FKeyEvent KeyEvent(Key, ModifierKeys, VirtualUser->GetUserIndex(), bRepeat, CharCode, KeyCode);
 	bool bDownResult = FSlateApplication::Get().ProcessKeyDownEvent(KeyEvent);
 
 	bool bKeyCharResult = false;
@@ -558,7 +558,7 @@ bool UWidgetInteractionComponent::ReleaseKey(FKey Key)
 	uint32 KeyCode, CharCode;
 	GetKeyAndCharCodes(Key, bHasKeyCode, KeyCode, bHasCharCode, CharCode);
 
-	FKeyEvent KeyEvent(Key, ModifierKeys, VirtualUser->GetUserIndex(), false, KeyCode, CharCode);
+	FKeyEvent KeyEvent(Key, ModifierKeys, VirtualUser->GetUserIndex(), false, CharCode, KeyCode);
 	return FSlateApplication::Get().ProcessKeyUpEvent(KeyEvent);
 }
 
