Correction to BP node Play Animation with Finished Event to correctly trigger the finish event on the next frame and prevent from triggering it if the widget no longer exist.

diff --git a/Engine/Source/Runtime/UMG/Private/Animation/WidgetAnimationPlayCallbackProxy.cpp b/Engine/Source/Runtime/UMG/Private/Animation/WidgetAnimationPlayCallbackProxy.cpp
--- a/Engine/Source/Runtime/UMG/Private/Animation/WidgetAnimationPlayCallbackProxy.cpp
+++ b/Engine/Source/Runtime/UMG/Private/Animation/WidgetAnimationPlayCallbackProxy.cpp
@@ -4,6 +4,7 @@
 #include "Animation/UMGSequencePlayer.h"
 #include "Engine/World.h"
 #include "TimerManager.h"
+#include "Types/WidgetActiveTimerDelegate.h"
 
 #define LOCTEXT_NAMESPACE "UMG"
 
@@ -35,12 +36,10 @@ class UUMGSequencePlayer* UWidgetAnimationPlayCallbackProxy::ExecutePlayAnimatio
 		return nullptr;
 	}
 
-	WorldPtr = Widget->GetWorld();
-
 	UUMGSequencePlayer* Player = Widget->PlayAnimation(InAnimation, StartAtTime, NumLoopsToPlay, PlayMode, PlaybackSpeed);
 	if (Player)
 	{
-		Player->OnSequenceFinishedPlaying().AddUObject(this, &UWidgetAnimationPlayCallbackProxy::OnFinished);
+		Player->OnSequenceFinishedPlaying().AddUObject(this, &UWidgetAnimationPlayCallbackProxy::OnSequenceFinished);
 	}
 
 	return Player;
@@ -53,33 +52,29 @@ class UUMGSequencePlayer* UWidgetAnimationPlayCallbackProxy::ExecutePlayAnimatio
 		return nullptr;
 	}
 
-	WorldPtr = Widget->GetWorld();
-
 	UUMGSequencePlayer* Player = Widget->PlayAnimationTimeRange(InAnimation, StartAtTime, EndAtTime, NumLoopsToPlay, PlayMode, PlaybackSpeed);
 	if (Player)
 	{
-		OnFinishedHandle = Player->OnSequenceFinishedPlaying().AddUObject(this, &UWidgetAnimationPlayCallbackProxy::OnFinished);
+		OnFinishedHandle = Player->OnSequenceFinishedPlaying().AddUObject(this, &UWidgetAnimationPlayCallbackProxy::OnSequenceFinished);
 	}
 
 	return Player;
 }
 
-void UWidgetAnimationPlayCallbackProxy::OnFinished(class UUMGSequencePlayer& Player)
+void UWidgetAnimationPlayCallbackProxy::OnSequenceFinished(class UUMGSequencePlayer& Player)
 {
 	Player.OnSequenceFinishedPlaying().Remove(OnFinishedHandle);
 
-	// We delay the Finish trigger to next frame.
-	if (UWorld* World = WorldPtr.Get())
-	{
-		// Use a dummy timer handle as we don't need to store it for later but we don't need to look for something to clear
-		FTimerHandle TimerHandle;
-		World->GetTimerManager().SetTimer(TimerHandle, this, &UWidgetAnimationPlayCallbackProxy::OnFinished_Delayed, 0.001f, false);
-	}
+	// We delay the Finish broadcast to next frame.
+	FTSTicker::FDelegateHandle TickerHandle = FTSTicker::GetCoreTicker().AddTicker(FTickerDelegate::CreateUObject(this, &UWidgetAnimationPlayCallbackProxy::OnAnimationFinished));
 }
 
-void UWidgetAnimationPlayCallbackProxy::OnFinished_Delayed()
+
+bool UWidgetAnimationPlayCallbackProxy::OnAnimationFinished(float /*DeltaTime*/)
 {
 	Finished.Broadcast();
-}
 
+	// Returning false, disable the ticker.
+	return false;
+}
 #undef LOCTEXT_NAMESPACE
diff --git a/Engine/Source/Runtime/UMG/Public/Animation/WidgetAnimationPlayCallbackProxy.h b/Engine/Source/Runtime/UMG/Public/Animation/WidgetAnimationPlayCallbackProxy.h
--- a/Engine/Source/Runtime/UMG/Public/Animation/WidgetAnimationPlayCallbackProxy.h
+++ b/Engine/Source/Runtime/UMG/Public/Animation/WidgetAnimationPlayCallbackProxy.h
@@ -25,11 +25,8 @@ class UWidgetAnimationPlayCallbackProxy : public UObject
 private:
 	class UUMGSequencePlayer* ExecutePlayAnimation(class UUserWidget* Widget, class UWidgetAnimation* InAnimation, float StartAtTime, int32 NumLoopsToPlay, EUMGSequencePlayMode::Type PlayMode, float PlaybackSpeed);
 	class UUMGSequencePlayer* ExecutePlayAnimationTimeRange(class UUserWidget* Widget, class UWidgetAnimation* InAnimation, float StartAtTime, float EndAtTime, int32 NumLoopsToPlay, EUMGSequencePlayMode::Type PlayMode, float PlaybackSpeed);
-	void OnFinished(class UUMGSequencePlayer& Player);
-	void OnFinished_Delayed();
-
-	// Pointer to the world, needed to delay the results slightly
-	TWeakObjectPtr<UWorld> WorldPtr;
+	void OnSequenceFinished(class UUMGSequencePlayer& Player);
+	bool OnAnimationFinished(float DeltaTime);
 
 	FDelegateHandle OnFinishedHandle;
 };
