PR #9445: Fix SetRowFill RowIndex parameter name. (Contributed by [REDACTED_PERSON])

diff --git a/Engine/Source/Runtime/UMG/Private/Components/GridPanel.cpp b/Engine/Source/Runtime/UMG/Private/Components/GridPanel.cpp
--- a/Engine/Source/Runtime/UMG/Private/Components/GridPanel.cpp
+++ b/Engine/Source/Runtime/UMG/Private/Components/GridPanel.cpp
@@ -111,17 +111,17 @@ void UGridPanel::SetColumnFill(int32 ColumnIndex, float Coefficient)
 	}
 }
 
-void UGridPanel::SetRowFill(int32 ColumnIndex, float Coefficient)
+void UGridPanel::SetRowFill(int32 RowIndex, float Coefficient)
 {
-	while (RowFill.Num() <= ColumnIndex)
+	while (RowFill.Num() <= RowIndex)
 	{
 		RowFill.Emplace(0);
 	}
-	RowFill[ColumnIndex] = Coefficient;
+	RowFill[RowIndex] = Coefficient;
 
 	if (MyGridPanel.IsValid())
 	{
-		MyGridPanel->SetRowFill(ColumnIndex, Coefficient);
+		MyGridPanel->SetRowFill(RowIndex, Coefficient);
 	}
 }
 
diff --git a/Engine/Source/Runtime/UMG/Public/Components/GridPanel.h b/Engine/Source/Runtime/UMG/Public/Components/GridPanel.h
--- a/Engine/Source/Runtime/UMG/Public/Components/GridPanel.h
+++ b/Engine/Source/Runtime/UMG/Public/Components/GridPanel.h
@@ -41,7 +41,7 @@ public:
 	void SetColumnFill(int32 ColumnIndex, float Coefficient);
 
 	UFUNCTION(BlueprintCallable, Category = "Widget")
-	void SetRowFill(int32 ColumnIndex, float Coefficient);
+	void SetRowFill(int32 RowIndex, float Coefficient);
 
 public:
 
