Unity Build fix: Missing include

diff --git a/Engine/Source/Runtime/UMG/Private/Animation/WidgetAnimationPlayCallbackProxy.cpp b/Engine/Source/Runtime/UMG/Private/Animation/WidgetAnimationPlayCallbackProxy.cpp
--- a/Engine/Source/Runtime/UMG/Private/Animation/WidgetAnimationPlayCallbackProxy.cpp
+++ b/Engine/Source/Runtime/UMG/Private/Animation/WidgetAnimationPlayCallbackProxy.cpp
@@ -5,6 +5,7 @@
 #include "Engine/World.h"
 #include "TimerManager.h"
 #include "Types/WidgetActiveTimerDelegate.h"
+#include "Containers/Ticker.h"
 
 #define LOCTEXT_NAMESPACE "UMG"
 
