Fix to WidgetComponent to prevent recreation of Widget on OnRegister/OnUnregister of component in editor. This was causing a performance issue with PostEditPropertyChanged executed on the Actor.

diff --git a/Engine/Source/Runtime/UMG/Private/Components/WidgetComponent.cpp b/Engine/Source/Runtime/UMG/Private/Components/WidgetComponent.cpp
--- a/Engine/Source/Runtime/UMG/Private/Components/WidgetComponent.cpp
+++ b/Engine/Source/Runtime/UMG/Private/Components/WidgetComponent.cpp
@@ -903,13 +903,6 @@ void UWidgetComponent::OnRegister()
 		}
 
 		BodySetup = nullptr;
-
-#if WITH_EDITOR
-		if (!bIsGameWorld)
-		{
-			InitWidget();
-		}
-#endif
 	}
 #endif // !UE_SERVER
 }
@@ -1066,14 +1059,6 @@ void UWidgetComponent::OnUnregister()
 		}
 	}
 #endif
-
-#if WITH_EDITOR
-	if (!GetWorld()->IsGameWorld())
-	{
-		ReleaseResources();
-	}
-#endif
-
 	Super::OnUnregister();
 }
 
