Allow user widgets to remove themselves from the player layers in the viewport even if the associated player controller has been destroyed.

diff --git a/Engine/Source/Runtime/Engine/Classes/Engine/GameViewportClient.h b/Engine/Source/Runtime/Engine/Classes/Engine/GameViewportClient.h
--- a/Engine/Source/Runtime/Engine/Classes/Engine/GameViewportClient.h
+++ b/Engine/Source/Runtime/Engine/Classes/Engine/GameViewportClient.h
@@ -273,7 +273,7 @@ public:
 	/**
 	 * Removes a previously-added widget from the Slate viewport, in the player's section.
 	 *
-	 * @param	Player The player to remove the widget's viewport from.
+	 * @param	Player The player to remove the widget's viewport from. Null will remove the widget regardless of which player it was added for.
 	 * @param	ViewportContent  The widget to remove.  Must be valid.
 	 */
 	virtual void RemoveViewportWidgetForPlayer(ULocalPlayer* Player, TSharedRef<SWidget> ViewportContent);
diff --git a/Engine/Source/Runtime/UMG/Private/UserWidget.cpp b/Engine/Source/Runtime/UMG/Private/UserWidget.cpp
--- a/Engine/Source/Runtime/UMG/Private/UserWidget.cpp
+++ b/Engine/Source/Runtime/UMG/Private/UserWidget.cpp
@@ -1030,11 +1030,11 @@ void UUserWidget::RemoveFromParent()
 					TSharedRef<SWidget> WidgetHostRef = WidgetHost.ToSharedRef();
 
 					ViewportClient->RemoveViewportWidgetContent(WidgetHostRef);
-
-					if (ULocalPlayer* LocalPlayer = GetOwningLocalPlayer())
-					{
-						ViewportClient->RemoveViewportWidgetForPlayer(LocalPlayer, WidgetHostRef);
-					}
+					
+					// We may no longer have access to our owning player if the player controller was destroyed
+					// Passing nullptr to RemoveViewportWidgetForPlayer will search all player layers for this widget
+					ULocalPlayer* MaybeOwningLocalPlayer = GetOwningLocalPlayer();
+					ViewportClient->RemoveViewportWidgetForPlayer(MaybeOwningLocalPlayer, WidgetHostRef);
 
 					FWorldDelegates::LevelRemovedFromWorld.RemoveAll(this);
 				}
