UMG: Speculative fix for crash in HandleLatentStop

diff --git a/Engine/Source/Runtime/UMG/Private/Animation/UMGSequencePlayer.cpp b/Engine/Source/Runtime/UMG/Private/Animation/UMGSequencePlayer.cpp
--- a/Engine/Source/Runtime/UMG/Private/Animation/UMGSequencePlayer.cpp
+++ b/Engine/Source/Runtime/UMG/Private/Animation/UMGSequencePlayer.cpp
@@ -410,7 +410,11 @@ void UUMGSequencePlayer::HandleLatentStop()
 
 	bIsStopping = false;
 
-	UserWidget->OnAnimationFinishedPlaying(*this);
+	if (UserWidget.IsValid())
+	{
+		UserWidget->OnAnimationFinishedPlaying(*this);
+	}
+
 	OnSequenceFinishedPlayingEvent.Broadcast(*this);
 }
 
