Speculative fix for crashes in latent actions.

It might be possible that delegates tied to destroyed objects might be left in the latent actions list.

diff --git a/Engine/Source/Runtime/MovieScene/Private/MovieSceneSequenceTickManager.cpp b/Engine/Source/Runtime/MovieScene/Private/MovieSceneSequenceTickManager.cpp
--- a/Engine/Source/Runtime/MovieScene/Private/MovieSceneSequenceTickManager.cpp
+++ b/Engine/Source/Runtime/MovieScene/Private/MovieSceneSequenceTickManager.cpp
@@ -226,6 +226,13 @@ void UMovieSceneSequenceTickManager::UnregisterTickClient(TScriptInterface<IMovi
 		return;
 	}
 
+	// Remove any latent actions tied to the given client.
+	UObject* Object = InTickInterface.GetObject();
+	if (Object)
+	{
+		ClearLatentActions(Object);
+	}
+
 	const int32 ClientIndex = Algo::IndexOfBy(TickableClients, ClientInterface, &FTickableClientData::Interface);
 	if (!TickableClients.IsValidIndex(ClientIndex))
 	{
diff --git a/Engine/Source/Runtime/UMG/Private/Animation/UMGSequencePlayer.cpp b/Engine/Source/Runtime/UMG/Private/Animation/UMGSequencePlayer.cpp
--- a/Engine/Source/Runtime/UMG/Private/Animation/UMGSequencePlayer.cpp
+++ b/Engine/Source/Runtime/UMG/Private/Animation/UMGSequencePlayer.cpp
@@ -547,6 +547,21 @@ void UUMGSequencePlayer::BeginDestroy()
 {
 	RootTemplateInstance.BeginDestroy();
 
+	// Remove any latent actions added by this player.
+	if (CVarUserWidgetUseParallelAnimation.GetValueOnGameThread())
+	{
+		UUserWidget* Widget = UserWidget.Get();
+		UUMGSequenceTickManager* TickManager = Widget ? ToRawPtr(Widget->AnimationTickManager) : nullptr;
+		if (TickManager)
+		{
+			TickManager->ClearLatentActions(this);
+		}
+	}
+	else
+	{
+		LatentActions.Empty();
+	}
+
 	Super::BeginDestroy();
 }
 
diff --git a/Engine/Source/Runtime/UMG/Private/Animation/UMGSequenceTickManager.cpp b/Engine/Source/Runtime/UMG/Private/Animation/UMGSequenceTickManager.cpp
--- a/Engine/Source/Runtime/UMG/Private/Animation/UMGSequenceTickManager.cpp
+++ b/Engine/Source/Runtime/UMG/Private/Animation/UMGSequenceTickManager.cpp
@@ -41,6 +41,7 @@ void UUMGSequenceTickManager::AddWidget(UUserWidget* InWidget)
 
 void UUMGSequenceTickManager::RemoveWidget(UUserWidget* InWidget)
 {
+	ClearLatentActions(InWidget);
 	TWeakObjectPtr<UUserWidget> WeakWidget = InWidget;
 	WeakUserWidgetData.Remove(WeakWidget);
 }
