[Backout] - CL21389818 - This is not the proper fix, backout. Will reinvestigate the performance issue.
Original CL Desc
-----------------------------------------------------------------
Fix to WidgetComponent to prevent recreation of Widget on OnRegister/OnUnregister of component in editor. This was causing a performance issue with PostEditPropertyChanged executed on the Actor.

diff --git a/Engine/Source/Runtime/UMG/Private/Components/WidgetComponent.cpp b/Engine/Source/Runtime/UMG/Private/Components/WidgetComponent.cpp
--- a/Engine/Source/Runtime/UMG/Private/Components/WidgetComponent.cpp
+++ b/Engine/Source/Runtime/UMG/Private/Components/WidgetComponent.cpp
@@ -903,6 +903,13 @@ void UWidgetComponent::OnRegister()
 		}
 
 		BodySetup = nullptr;
+
+#if WITH_EDITOR
+		if (!bIsGameWorld)
+		{
+			InitWidget();
+		}
+#endif
 	}
 #endif // !UE_SERVER
 }
@@ -1059,6 +1066,14 @@ void UWidgetComponent::OnUnregister()
 		}
 	}
 #endif
+
+#if WITH_EDITOR
+	if (!GetWorld()->IsGameWorld())
+	{
+		ReleaseResources();
+	}
+#endif
+
 	Super::OnUnregister();
 }
 
