Add a flag for flushing input when you change the UI mode on a player controller. This will ensure that any input actions get properly cancelled and have their key state reset if desired. This should be an option on each function instead of handling it more internally because you may not always want to do it depending on the behavior of the UI you want. For example, a third person game that brings up some new widgets may want to continue passing input through the player controller while the UI is up.

diff --git a/Engine/Source/Runtime/UMG/Private/WidgetBlueprintLibrary.cpp b/Engine/Source/Runtime/UMG/Private/WidgetBlueprintLibrary.cpp
--- a/Engine/Source/Runtime/UMG/Private/WidgetBlueprintLibrary.cpp
+++ b/Engine/Source/Runtime/UMG/Private/WidgetBlueprintLibrary.cpp
@@ -76,7 +76,7 @@ UDragDropOperation* UWidgetBlueprintLibrary::CreateDragDropOperation(TSubclassOf
 	return DragDropOperation;
 }
 
-void UWidgetBlueprintLibrary::SetInputMode_UIOnlyEx(APlayerController* PlayerController, UWidget* InWidgetToFocus, EMouseLockMode InMouseLockMode)
+void UWidgetBlueprintLibrary::SetInputMode_UIOnlyEx(APlayerController* PlayerController, UWidget* InWidgetToFocus, EMouseLockMode InMouseLockMode, const bool bFlushInput /* = false */)
 {
 	if (PlayerController != nullptr)
 	{
@@ -88,6 +88,11 @@ void UWidgetBlueprintLibrary::SetInputMode_UIOnlyEx(APlayerController* PlayerCon
 			InputMode.SetWidgetToFocus(InWidgetToFocus->TakeWidget());
 		}
 		PlayerController->SetInputMode(InputMode);
+
+		if (bFlushInput)
+		{
+			PlayerController->FlushPressedKeys();
+		}
 	}
 	#if WITH_EDITOR 
 	else
@@ -97,7 +102,7 @@ void UWidgetBlueprintLibrary::SetInputMode_UIOnlyEx(APlayerController* PlayerCon
 	#endif // WITH_EDITOR
 }
 
-void UWidgetBlueprintLibrary::SetInputMode_GameAndUIEx(APlayerController* PlayerController, UWidget* InWidgetToFocus, EMouseLockMode InMouseLockMode, bool bHideCursorDuringCapture)
+void UWidgetBlueprintLibrary::SetInputMode_GameAndUIEx(APlayerController* PlayerController, UWidget* InWidgetToFocus, EMouseLockMode InMouseLockMode, bool bHideCursorDuringCapture, const bool bFlushInput /* = false */)
 {
 	if (PlayerController != nullptr)
 	{
@@ -110,6 +115,11 @@ void UWidgetBlueprintLibrary::SetInputMode_GameAndUIEx(APlayerController* Player
 			InputMode.SetWidgetToFocus(InWidgetToFocus->TakeWidget());
 		}
 		PlayerController->SetInputMode(InputMode);
+
+		if (bFlushInput)
+		{
+			PlayerController->FlushPressedKeys();
+		}
 	}
 	#if WITH_EDITOR 
 	else
@@ -120,12 +130,17 @@ void UWidgetBlueprintLibrary::SetInputMode_GameAndUIEx(APlayerController* Player
 	
 }
 
-void UWidgetBlueprintLibrary::SetInputMode_GameOnly(APlayerController* PlayerController)
+void UWidgetBlueprintLibrary::SetInputMode_GameOnly(APlayerController* PlayerController, const bool bFlushInput /* = false */)
 {
 	if (PlayerController != nullptr)
 	{
 		FInputModeGameOnly InputMode;
 		PlayerController->SetInputMode(InputMode);
+		
+		if (bFlushInput)
+		{
+			PlayerController->FlushPressedKeys();
+		}
 	}
 	#if WITH_EDITOR 
 	else
diff --git a/Engine/Source/Runtime/UMG/Public/Blueprint/WidgetBlueprintLibrary.h b/Engine/Source/Runtime/UMG/Public/Blueprint/WidgetBlueprintLibrary.h
--- a/Engine/Source/Runtime/UMG/Public/Blueprint/WidgetBlueprintLibrary.h
+++ b/Engine/Source/Runtime/UMG/Public/Blueprint/WidgetBlueprintLibrary.h
@@ -47,14 +47,14 @@ public:
 	 * Note: This means that any bound Input Events in the widget will not be called!
 	 */
 	UFUNCTION(BlueprintCallable, BlueprintCosmetic, Category="Input", meta=(DisplayName="Set Input Mode UI Only"))
-	static void SetInputMode_UIOnlyEx(APlayerController* PlayerController, UWidget* InWidgetToFocus = nullptr, EMouseLockMode InMouseLockMode = EMouseLockMode::DoNotLock);
+	static void SetInputMode_UIOnlyEx(APlayerController* PlayerController, UWidget* InWidgetToFocus = nullptr, EMouseLockMode InMouseLockMode = EMouseLockMode::DoNotLock, const bool bFlushInput = false);
 	/**
 	 * Setup an input mode that allows only the UI to respond to user input, and if the UI doesn't handle it player input / player controller gets a chance.
 	 * 
 	 * Note: This means that any bound Input events in the widget will be called.
 	 */
 	UFUNCTION(BlueprintCallable, BlueprintCosmetic, Category="Input", meta = (DisplayName = "Set Input Mode Game And UI"))
-	static void SetInputMode_GameAndUIEx(APlayerController* PlayerController, UWidget* InWidgetToFocus = nullptr, EMouseLockMode InMouseLockMode = EMouseLockMode::DoNotLock, bool bHideCursorDuringCapture = true);
+	static void SetInputMode_GameAndUIEx(APlayerController* PlayerController, UWidget* InWidgetToFocus = nullptr, EMouseLockMode InMouseLockMode = EMouseLockMode::DoNotLock, bool bHideCursorDuringCapture = true, const bool bFlushInput = false);
 
 	/**
 	 * Setup an input mode that allows only player input / player controller to respond to user input.
@@ -62,7 +62,7 @@ public:
 	 * Note: Any bound Input Events in this widget will be called.
 	 */
 	UFUNCTION(BlueprintCallable, BlueprintCosmetic, Category="Input")
-	static void SetInputMode_GameOnly(APlayerController* PlayerController);
+	static void SetInputMode_GameOnly(APlayerController* PlayerController, const bool bFlushInput = false);
 
 	/** */
 	UFUNCTION(BlueprintCallable, BlueprintCosmetic, Category="Focus")
