UMG: Add PreConstruct and Construct call back to extension. Do not remove extension in the destruction call. They maybe be added out of the construct/destruct loop.

diff --git a/Engine/Plugins/Runtime/ModelViewViewModel/Source/ModelViewViewModel/Private/View/MVVMView.cpp b/Engine/Plugins/Runtime/ModelViewViewModel/Source/ModelViewViewModel/Private/View/MVVMView.cpp
--- a/Engine/Plugins/Runtime/ModelViewViewModel/Source/ModelViewViewModel/Private/View/MVVMView.cpp
+++ b/Engine/Plugins/Runtime/ModelViewViewModel/Source/ModelViewViewModel/Private/View/MVVMView.cpp
@@ -23,9 +23,10 @@ void UMVVMView::ConstructView(const UMVVMViewClass* InClassExtension)
 }
 
 
-void UMVVMView::Initialize()
+void UMVVMView::Construct()
 {
 	check(ClassExtension);
+	check(bConstructed == false);
 
 	// Init ViewModel instances
 	for (const FMVVMViewClass_SourceCreator& Item : ClassExtension->GetViewModelCreators())
@@ -42,11 +43,16 @@ void UMVVMView::Initialize()
 			EnableLibraryBinding(Binding, Index);
 		}
 	}
+
+	bConstructed = true;
 }
 
 
 void UMVVMView::Destruct()
 {
+	check(bConstructed == true);
+	bConstructed = false;
+
 	for (FRegisteredSource& Source : AllSources)
 	{
 		UObject* SourceObject = Source.Source.Get();
diff --git a/Engine/Plugins/Runtime/ModelViewViewModel/Source/ModelViewViewModel/Private/View/MVVMViewClass.cpp b/Engine/Plugins/Runtime/ModelViewViewModel/Source/ModelViewViewModel/Private/View/MVVMViewClass.cpp
--- a/Engine/Plugins/Runtime/ModelViewViewModel/Source/ModelViewViewModel/Private/View/MVVMViewClass.cpp
+++ b/Engine/Plugins/Runtime/ModelViewViewModel/Source/ModelViewViewModel/Private/View/MVVMViewClass.cpp
@@ -165,7 +165,7 @@ FString FMVVMViewClass_CompiledBinding::ToString() const
 // 
 ///////////////////////////////////////////////////////////////////////
 
-void UMVVMViewClass::Initialize(UUserWidget* UserWidget)
+void UMVVMViewClass::Construct(UUserWidget* UserWidget)
 {
 	check(UserWidget->GetExtension<UMVVMView>() == nullptr);
 
@@ -182,4 +182,10 @@ void UMVVMViewClass::Initialize(UUserWidget* UserWidget)
 	}
 }
 
+
+void UMVVMViewClass::Destruct(UUserWidget* UserWidget)
+{
+	UserWidget->RemoveExtensions<UMVVMView>();
+}
+
 #undef LOCTEXT_NAMESPACE
diff --git a/Engine/Plugins/Runtime/ModelViewViewModel/Source/ModelViewViewModel/Public/View/MVVMView.h b/Engine/Plugins/Runtime/ModelViewViewModel/Source/ModelViewViewModel/Public/View/MVVMView.h
--- a/Engine/Plugins/Runtime/ModelViewViewModel/Source/ModelViewViewModel/Public/View/MVVMView.h
+++ b/Engine/Plugins/Runtime/ModelViewViewModel/Source/ModelViewViewModel/Public/View/MVVMView.h
@@ -31,7 +31,7 @@ public:
 	void ConstructView(const UMVVMViewClass* ClassExtension);
 
 	//~ Begin UUserWidgetExtension implementation
-	virtual void Initialize() override;
+	virtual void Construct() override;
 	virtual void Destruct() override;
 
 // todo needed when we will support different type of update
@@ -84,6 +84,9 @@ private:
 	/** The binding that are enabled for the instance. */
 	TBitArray<> EnabledLibraryBindings;
 
-	/** Should log when a binding is executed */
+	/** Should log when a binding is executed. */
 	bool bLogBinding = false;
+
+	/** Is the Construct method was called. */
+	bool bConstructed = false;
 };
diff --git a/Engine/Plugins/Runtime/ModelViewViewModel/Source/ModelViewViewModel/Public/View/MVVMViewClass.h b/Engine/Plugins/Runtime/ModelViewViewModel/Source/ModelViewViewModel/Public/View/MVVMViewClass.h
--- a/Engine/Plugins/Runtime/ModelViewViewModel/Source/ModelViewViewModel/Public/View/MVVMViewClass.h
+++ b/Engine/Plugins/Runtime/ModelViewViewModel/Source/ModelViewViewModel/Public/View/MVVMViewClass.h
@@ -193,7 +193,10 @@ class MODELVIEWVIEWMODEL_API UMVVMViewClass : public UWidgetBlueprintGeneratedCl
 	friend UE::MVVM::Private::FMVVMViewBlueprintCompiler;
 
 public:
-	virtual void Initialize(UUserWidget* UserWidget) override;
+	//~ Begin UWidgetBlueprintGeneratedClassExtension
+	virtual void Construct(UUserWidget* UserWidget) override;
+	virtual void Destruct(UUserWidget* UserWidget) override;
+	//~ End UWidgetBlueprintGeneratedClassExtension
 
 public:
 	/** Get the list of the needed ViewModel. */
diff --git a/Engine/Source/Runtime/UMG/Private/UserWidget.cpp b/Engine/Source/Runtime/UMG/Private/UserWidget.cpp
--- a/Engine/Source/Runtime/UMG/Private/UserWidget.cpp
+++ b/Engine/Source/Runtime/UMG/Private/UserWidget.cpp
@@ -74,6 +74,7 @@ UUserWidget::UUserWidget(const FObjectInitializer& ObjectInitializer)
 	, bHasScriptImplementedTick(true)
 	, bHasScriptImplementedPaint(true)
 	, bInitialized(false)
+	, bAreExtensionsConstructed(false)
 	, bStoppingAllAnimations(false)
 	, TickFrequency(EWidgetTickFrequency::Auto)
 {
@@ -1407,6 +1408,7 @@ void UUserWidget::NativePreConstruct()
 void UUserWidget::NativeConstruct()
 {
 	LLM_SCOPE_BYTAG(UI_UMG);
+
 	if (UWidgetBlueprintGeneratedClass* BPClass = Cast<UWidgetBlueprintGeneratedClass>(GetClass()))
 	{
 		BPClass->ForEachExtension([this](UWidgetBlueprintGeneratedClassExtension* Extension)
@@ -1415,6 +1417,20 @@ void UUserWidget::NativeConstruct()
 			});
 	}
 
+	// Extension can add other extensions.
+	check(bAreExtensionsConstructed == false);
+	bAreExtensionsConstructed = true;
+	if (Extensions.Num() > 0)
+	{
+		TArray<UUserWidgetExtension*, TInlineAllocator<32>> LocalExtensions;
+		LocalExtensions.Append(Extensions);
+		for (UUserWidgetExtension* Extension : LocalExtensions)
+		{
+			check(Extension);
+			Extension->Construct();
+		}
+	}
+
 	Construct();
 	UpdateCanTick();
 }
@@ -1424,15 +1440,15 @@ void UUserWidget::NativeDestruct()
 	StopListeningForAllInputActions();
 	Destruct();
 
-	if (bInitialized)
+	// Extension can remove other extensions.
+	bAreExtensionsConstructed = false; // To prevent calling Destruct on the same extension if it's removed by another extension.
+	if (Extensions.Num() > 0)
 	{
-		// Extension can remove other extensions.
-		while (Extensions.Num() > 0)
+		TArray<UUserWidgetExtension*, TInlineAllocator<32>> LocalExtensions;
+		LocalExtensions.Append(Extensions);
+		for (UUserWidgetExtension* Extension : LocalExtensions)
 		{
-			int32 Index = Extensions.Num() - 1;
-			UUserWidgetExtension* Extension = Extensions[Index];
 			check(Extension);
-			Extensions.RemoveAtSwap(Index);
 			Extension->Destruct();
 		}
 	}
@@ -1692,7 +1708,7 @@ void UUserWidget::UpdateCanTick()
 			bCanTick |= World->GetLatentActionManager().GetNumActionsForObject(this) != 0;
 			bCanTick |= ActiveSequencePlayers.Num() > 0;
 
-			if (!bCanTick)
+			if (!bCanTick && bAreExtensionsConstructed)
 			{
 				for(UUserWidgetExtension* Extension : Extensions)
 				{
@@ -2156,6 +2172,14 @@ UUserWidgetExtension* UUserWidget::AddExtension(TSubclassOf<UUserWidgetExtension
 	{
 		Extension->Initialize();
 	}
+	if (bAreExtensionsConstructed)
+	{
+		Extension->Construct();
+		if (Extension->RequiresTick())
+		{
+			UpdateCanTick();
+		}
+	}
 	return Extension;
 }
 
@@ -2165,14 +2189,49 @@ void UUserWidget::RemoveExtension(UUserWidgetExtension* InExtension)
 	{
 		if (Extensions.RemoveSingleSwap(InExtension))
 		{
-			if (bInitialized)
+			if (bAreExtensionsConstructed)
 			{
+				bool bUpdateTick = InExtension->RequiresTick();
 				InExtension->Destruct();
+				if (bUpdateTick)
+				{
+					UpdateCanTick();
+				}
 			}
 		}
 	}
 }
 
+void UUserWidget::RemoveExtensions(TSubclassOf<UUserWidgetExtension> InExtensionType)
+{
+
+	TArray<UUserWidgetExtension*, TInlineAllocator<32>> LocalExtensions;
+	for (int32 Index = Extensions.Num() - 1; Index >= 0; --Index)
+	{
+		UUserWidgetExtension* Extension = Extensions[Index];
+		if (Extension->IsA(InExtensionType))
+		{
+			LocalExtensions.Add(Extension);
+			Extensions.RemoveAtSwap(Index);
+
+		}
+	}
+
+	if (bAreExtensionsConstructed)
+	{
+		bool bUpdateTick = false;
+		for (UUserWidgetExtension* Extension : LocalExtensions)
+		{
+			bUpdateTick = bUpdateTick || Extension->RequiresTick();
+			Extension->Destruct();
+		}
+		if (bUpdateTick)
+		{
+			UpdateCanTick();
+		}
+	}
+}
+
 
 /////////////////////////////////////////////////////
 
diff --git a/Engine/Source/Runtime/UMG/Public/Blueprint/UserWidget.h b/Engine/Source/Runtime/UMG/Public/Blueprint/UserWidget.h
--- a/Engine/Source/Runtime/UMG/Public/Blueprint/UserWidget.h
+++ b/Engine/Source/Runtime/UMG/Public/Blueprint/UserWidget.h
@@ -1096,6 +1096,17 @@ public:
 	UFUNCTION(BlueprintCallable, Category = "User Interface|Extension")
 	void RemoveExtension(UUserWidgetExtension* InExtension);
 
+	/** Remove all extensions of the requested type. */
+	template<typename ExtensionType>
+	void RemoveExtensions()
+	{
+		return RemoveExtensions(ExtensionType::StaticClass());
+	}
+
+	/** Remove all extensions of the requested type. */
+	UFUNCTION(BlueprintCallable, Category = "User Interface|Extension")
+	void RemoveExtensions(TSubclassOf<UUserWidgetExtension> InExtensionType);
+
 	/**
 	 * Plays a sound through the UI
 	 *
@@ -1263,6 +1274,9 @@ private:
 	/** Has this widget been initialized by its class yet? */
 	uint8 bInitialized : 1;
 
+	/** Has this widget been constructed and we need to call Construct on new extension. */
+	uint8 bAreExtensionsConstructed : 1;
+
 	/** If we're stopping all animations, don't allow new animations to be created as side-effects. */
 	uint8 bStoppingAllAnimations : 1;
 
diff --git a/Engine/Source/Runtime/UMG/Public/Extensions/UserWidgetExtension.h b/Engine/Source/Runtime/UMG/Public/Extensions/UserWidgetExtension.h
--- a/Engine/Source/Runtime/UMG/Public/Extensions/UserWidgetExtension.h
+++ b/Engine/Source/Runtime/UMG/Public/Extensions/UserWidgetExtension.h
@@ -26,6 +26,11 @@ public:
 	{
 	}
 
+	/** Extend the UUserWidget::Construct function */
+	virtual void Construct()
+	{
+	}
+
 	/** Extend the UUserWidget::Destruct function */
 	virtual void Destruct()
 	{
