UScrollBox: expose "GetViewFraction" from underlying slate widget

diff --git a/Engine/Source/Runtime/UMG/Private/Components/ScrollBox.cpp b/Engine/Source/Runtime/UMG/Private/Components/ScrollBox.cpp
--- a/Engine/Source/Runtime/UMG/Private/Components/ScrollBox.cpp
+++ b/Engine/Source/Runtime/UMG/Private/Components/ScrollBox.cpp
@@ -115,7 +115,7 @@ void UScrollBox::OnSlotRemoved(UPanelSlot* InSlot)
 	// Remove the widget from the live slot if it exists.
 	if ( MyScrollBox.IsValid() && InSlot->Content)
 	{
-		TSharedPtr<SWidget> Widget = InSlot->Content->GetCachedWidget();
+		const TSharedPtr<SWidget> Widget = InSlot->Content->GetCachedWidget();
 		if ( Widget.IsValid() )
 		{
 			MyScrollBox->RemoveSlot(Widget.ToSharedRef());
@@ -189,6 +189,16 @@ float UScrollBox::GetScrollOffsetOfEnd() const
 	return 0;
 }
 
+float UScrollBox::GetViewFraction() const
+{
+	if ( MyScrollBox.IsValid() )
+	{
+		return MyScrollBox->GetViewFraction();
+	}
+
+	return 0;
+}
+
 float UScrollBox::GetViewOffsetFraction() const
 {
 	if ( MyScrollBox.IsValid() )
diff --git a/Engine/Source/Runtime/UMG/Public/Components/ScrollBox.h b/Engine/Source/Runtime/UMG/Public/Components/ScrollBox.h
--- a/Engine/Source/Runtime/UMG/Public/Components/ScrollBox.h
+++ b/Engine/Source/Runtime/UMG/Public/Components/ScrollBox.h
@@ -159,6 +159,10 @@ public:
 	UFUNCTION(BlueprintCallable, Category = "Widget")
 	float GetScrollOffsetOfEnd() const;
 
+	/** Gets the fraction currently visible in the scrollbox */
+	UFUNCTION(BlueprintCallable, Category="Widget")
+	float GetViewFraction() const;
+
 	UFUNCTION(BlueprintCallable, Category="Widget")
 	float GetViewOffsetFraction() const;
 
