PR #9470: Fix Crash when floatbinding with Property. (Contributed by [REDACTED_PERSON])

diff --git a/Engine/Source/Runtime/UMG/Private/Binding/FloatBinding.cpp b/Engine/Source/Runtime/UMG/Private/Binding/FloatBinding.cpp
--- a/Engine/Source/Runtime/UMG/Private/Binding/FloatBinding.cpp
+++ b/Engine/Source/Runtime/UMG/Private/Binding/FloatBinding.cpp
@@ -29,31 +29,18 @@ float UFloatBinding::GetValue() const
 
 		float FloatValue = 0.0f;
 
-		SourcePath.Resolve(Source);
-		if (FProperty* Property = SourcePath.GetFProperty())
+		if (SourcePath.Resolve(Source))
 		{
 			double DoubleValue = 0.0;
-			if (Property->IsA<FFloatProperty>() && SourcePath.GetValue<float>(Source, FloatValue))
+			if (SourcePath.GetValue<float>(Source, FloatValue))
 			{
 				return FloatValue;
 			}
-			else if (Property->IsA<FDoubleProperty>() && SourcePath.GetValue<double>(Source, DoubleValue))
+			else if (SourcePath.GetValue<double>(Source, DoubleValue))
 			{
 				FloatValue = static_cast<float>(DoubleValue);
 				return FloatValue;
 			}
-			else
-			{
-				checkf(false, TEXT("Unexpected property type: '%s'! Float bindings must use either a float or double property."), *Property->GetCPPType());
-			}
-		}
-		else
-		{
-			check(SourcePath.GetCachedFunction());
-			if (SourcePath.GetValue<float>(Source, FloatValue))
-			{
-				return FloatValue;
-			}
 		}
 	}
 
diff --git a/Engine/Source/Runtime/UMG/Private/Binding/Int32Binding.cpp b/Engine/Source/Runtime/UMG/Private/Binding/Int32Binding.cpp
--- a/Engine/Source/Runtime/UMG/Private/Binding/Int32Binding.cpp
+++ b/Engine/Source/Runtime/UMG/Private/Binding/Int32Binding.cpp
@@ -31,7 +31,7 @@ int32 UInt32Binding::GetValue() const
 		}
 	}
 
-	return false;
+	return 0;
 }
 
 #undef LOCTEXT_NAMESPACE
