Set ForceNewLine in WrapBox to the actual boolean value (previously it was always set to false)
#9086

diff --git a/Engine/Source/Runtime/UMG/Private/Components/WrapBoxSlot.cpp b/Engine/Source/Runtime/UMG/Private/Components/WrapBoxSlot.cpp
--- a/Engine/Source/Runtime/UMG/Private/Components/WrapBoxSlot.cpp
+++ b/Engine/Source/Runtime/UMG/Private/Components/WrapBoxSlot.cpp
@@ -36,7 +36,7 @@ PRAGMA_DISABLE_DEPRECATION_WARNINGS
 		.VAlign(VerticalAlignment)
 		.FillEmptySpace(bFillEmptySpace)
 		.FillLineWhenSizeLessThan(FillSpanWhenLessThan == 0 ? TOptional<float>() : TOptional<float>(FillSpanWhenLessThan))
-		.ForceNewLine(false)
+		.ForceNewLine(bForceNewLine)
 		.Expose(Slot)
 		[
 			Content == nullptr ? SNullWidget::NullWidget : Content->TakeWidget()
