Cleaned up build.cs files by removing any include paths that were already being added by UBT. This was done to help identify how include paths are being added and to help with future refactoring.

diff --git a/Engine/Source/Runtime/SlateNullRenderer/SlateNullRenderer.Build.cs b/Engine/Source/Runtime/SlateNullRenderer/SlateNullRenderer.Build.cs
--- a/Engine/Source/Runtime/SlateNullRenderer/SlateNullRenderer.Build.cs
+++ b/Engine/Source/Runtime/SlateNullRenderer/SlateNullRenderer.Build.cs
@@ -6,8 +6,6 @@ public class SlateNullRenderer : ModuleRules
 {
 	public SlateNullRenderer(ReadOnlyTargetRules Target) : base(Target)
 	{
-        PrivateIncludePaths.Add("Runtime/SlateNullRenderer/Private");
-
 		PrivateDependencyModuleNames.AddRange(
 			new string[] {
 				"Core",
diff --git a/Engine/Source/Runtime/SlateRHIRenderer/SlateRHIRenderer.Build.cs b/Engine/Source/Runtime/SlateRHIRenderer/SlateRHIRenderer.Build.cs
--- a/Engine/Source/Runtime/SlateRHIRenderer/SlateRHIRenderer.Build.cs
+++ b/Engine/Source/Runtime/SlateRHIRenderer/SlateRHIRenderer.Build.cs
@@ -6,8 +6,6 @@ public class SlateRHIRenderer : ModuleRules
 {
     public SlateRHIRenderer(ReadOnlyTargetRules Target) : base(Target)
 	{
-        PrivateIncludePaths.Add("Runtime/SlateRHIRenderer/Private");
-
 		PrivateDependencyModuleNames.AddRange(
 			new string[] {
 				"Core",
diff --git a/Engine/Source/Runtime/Sockets/Sockets.Build.cs b/Engine/Source/Runtime/Sockets/Sockets.Build.cs
--- a/Engine/Source/Runtime/Sockets/Sockets.Build.cs
+++ b/Engine/Source/Runtime/Sockets/Sockets.Build.cs
@@ -6,8 +6,6 @@ public class Sockets : ModuleRules
 {
 	public Sockets(ReadOnlyTargetRules Target) : base(Target)
 	{
-		PrivateIncludePaths.Add("Runtime/Sockets/Private");
-
 		PublicIncludePathModuleNames.Add("NetCommon");
 
 		PrivateDependencyModuleNames.AddRange(
diff --git a/Engine/Source/Runtime/StaticMeshDescription/StaticMeshDescription.Build.cs b/Engine/Source/Runtime/StaticMeshDescription/StaticMeshDescription.Build.cs
--- a/Engine/Source/Runtime/StaticMeshDescription/StaticMeshDescription.Build.cs
+++ b/Engine/Source/Runtime/StaticMeshDescription/StaticMeshDescription.Build.cs
@@ -6,9 +6,6 @@ namespace UnrealBuildTool.Rules
 	{
 		public StaticMeshDescription(ReadOnlyTargetRules Target) : base(Target)
 		{
-            PrivateIncludePaths.Add("Runtime/StaticMeshDescription/Private");
-            PublicIncludePaths.Add("Runtime/StaticMeshDescription/Public");
-
 			PublicDependencyModuleNames.AddRange(
 				new string[]
 				{
diff --git a/Engine/Source/Runtime/StreamingPauseRendering/StreamingPauseRendering.Build.cs b/Engine/Source/Runtime/StreamingPauseRendering/StreamingPauseRendering.Build.cs
--- a/Engine/Source/Runtime/StreamingPauseRendering/StreamingPauseRendering.Build.cs
+++ b/Engine/Source/Runtime/StreamingPauseRendering/StreamingPauseRendering.Build.cs
@@ -6,8 +6,6 @@ public class StreamingPauseRendering : ModuleRules
 {
     public StreamingPauseRendering(ReadOnlyTargetRules Target) : base(Target)
 	{
-        PrivateIncludePaths.Add("Runtime/StreamingPauseRendering/Private");
-
 		PublicDependencyModuleNames.AddRange(
 			new string[] {
 					"Engine",
diff --git a/Engine/Source/Runtime/SynthBenchmark/SynthBenchmark.Build.cs b/Engine/Source/Runtime/SynthBenchmark/SynthBenchmark.Build.cs
--- a/Engine/Source/Runtime/SynthBenchmark/SynthBenchmark.Build.cs
+++ b/Engine/Source/Runtime/SynthBenchmark/SynthBenchmark.Build.cs
@@ -8,19 +8,6 @@ namespace UnrealBuildTool.Rules
 		{
             PublicIncludePathModuleNames.Add("Renderer");
 
-			PublicIncludePaths.AddRange(
-				new string[] {
-					// ... add public include paths required here ...
-				}
-				);
-
-			PrivateIncludePaths.AddRange(
-				new string[] {
-					"Runtime/SynthBenchmark/Private",
-					// ... add other private include paths required here ...
-				}
-				);
-
 			PublicDependencyModuleNames.AddRange(
 				new string[]
 				{
diff --git a/Engine/Source/Runtime/UMG/UMG.Build.cs b/Engine/Source/Runtime/UMG/UMG.Build.cs
--- a/Engine/Source/Runtime/UMG/UMG.Build.cs
+++ b/Engine/Source/Runtime/UMG/UMG.Build.cs
@@ -6,12 +6,6 @@ public class UMG : ModuleRules
 {
 	public UMG(ReadOnlyTargetRules Target) : base(Target)
 	{
-        PrivateIncludePaths.AddRange(
-            new string[] {
-                "Runtime/UMG/Private" // For PCH includes (because they don't work with relative paths, yet)
-            })
-		;
-
 		PrivateDependencyModuleNames.AddRange(
 			new string[] {
 				"Core",
@@ -21,7 +15,7 @@ public class UMG : ModuleRules
 				"Slate",
 				"SlateCore",
 				"RenderCore",
-				"RHI",
+				"RHI"
 			}
 		);
 
@@ -37,19 +31,14 @@ public class UMG : ModuleRules
 
         PrivateIncludePathModuleNames.AddRange(
             new string[] {
+				"SlateRHIRenderer",
 				"ImageWrapper",
-                 "TargetPlatform",
-            }
+                "TargetPlatform",
+			}
         );
 
 		if (Target.Type != TargetType.Server)
 		{
-			PrivateIncludePathModuleNames.AddRange(
-				new string[] {
-					"SlateRHIRenderer",
-				}
-			);
-
             DynamicallyLoadedModuleNames.AddRange(
                 new string[] {
 				    "ImageWrapper",
diff --git a/Engine/Source/Runtime/WebBrowser/WebBrowser.Build.cs b/Engine/Source/Runtime/WebBrowser/WebBrowser.Build.cs
--- a/Engine/Source/Runtime/WebBrowser/WebBrowser.Build.cs
+++ b/Engine/Source/Runtime/WebBrowser/WebBrowser.Build.cs
@@ -7,8 +7,6 @@ public class WebBrowser : ModuleRules
 {
 	public WebBrowser(ReadOnlyTargetRules Target) : base(Target)
 	{
-		PublicIncludePaths.Add("Runtime/WebBrowser/Public");
-		PrivateIncludePaths.Add("Runtime/WebBrowser/Private");
 		PrivateDependencyModuleNames.AddRange(
 			new string[]
 			{
diff --git a/Engine/Source/Runtime/WebBrowserTexture/WebBrowserTexture.Build.cs b/Engine/Source/Runtime/WebBrowserTexture/WebBrowserTexture.Build.cs
--- a/Engine/Source/Runtime/WebBrowserTexture/WebBrowserTexture.Build.cs
+++ b/Engine/Source/Runtime/WebBrowserTexture/WebBrowserTexture.Build.cs
@@ -13,10 +13,7 @@ public class WebBrowserTexture : ModuleRules
             Target.Platform == UnrealTargetPlatform.IOS ||
             Target.Platform == UnrealTargetPlatform.TVOS ||
             Target.bBuildEditor == true)
-		{
-			PublicIncludePaths.Add("Runtime/WebBrowserTexture/Public");
-			PrivateIncludePaths.Add("Runtime/WebBrowserTexture/Private");
-			
+		{			
 			// Needed for external texture support
 			PublicIncludePaths.AddRange(
 				new string[] {
diff --git a/Engine/Source/Runtime/WidgetCarousel/WidgetCarousel.Build.cs b/Engine/Source/Runtime/WidgetCarousel/WidgetCarousel.Build.cs
--- a/Engine/Source/Runtime/WidgetCarousel/WidgetCarousel.Build.cs
+++ b/Engine/Source/Runtime/WidgetCarousel/WidgetCarousel.Build.cs
@@ -16,10 +16,5 @@ public class WidgetCarousel : ModuleRules
 				"CoreUObject"
 			}
 		);
-
-		PrivateIncludePaths.AddRange(
-			new string[] {
-			}
-		);
 	}
 }
diff --git a/Engine/Source/Runtime/Windows/AudioMixerXAudio2/AudioMixerXAudio2.Build.cs b/Engine/Source/Runtime/Windows/AudioMixerXAudio2/AudioMixerXAudio2.Build.cs
--- a/Engine/Source/Runtime/Windows/AudioMixerXAudio2/AudioMixerXAudio2.Build.cs
+++ b/Engine/Source/Runtime/Windows/AudioMixerXAudio2/AudioMixerXAudio2.Build.cs
@@ -7,8 +7,6 @@ public class AudioMixerXAudio2 : ModuleRules
 	public AudioMixerXAudio2(ReadOnlyTargetRules Target) : base(Target)
 	{
 		PrivateIncludePathModuleNames.Add("TargetPlatform");
-		PublicIncludePaths.Add("Runtime/AudioMixer/Public");
-		PrivateIncludePaths.Add("Runtime/AudioMixer/Private");
 
 		if (Target.bCompileAgainstEngine)
         {
