UMG: Update the Throbber with getter and setter

diff --git a/Engine/Source/Runtime/UMG/Private/Components/Throbber.cpp b/Engine/Source/Runtime/UMG/Private/Components/Throbber.cpp
--- a/Engine/Source/Runtime/UMG/Private/Components/Throbber.cpp
+++ b/Engine/Source/Runtime/UMG/Private/Components/Throbber.cpp
@@ -2,6 +2,7 @@
 
 #include "Components/Throbber.h"
 #include "SlateFwd.h"
+#include "SlateGlobals.h"
 #include "Slate/SlateBrushAsset.h"
 #include "Styling/UMGCoreStyle.h"
 
@@ -19,6 +20,7 @@ static FSlateBrush* EditorThrobberBrush = nullptr;
 UThrobber::UThrobber(const FObjectInitializer& ObjectInitializer)
 	: Super(ObjectInitializer)
 {
+PRAGMA_DISABLE_DEPRECATION_WARNINGS
 	NumberOfPieces = 3;
 
 	bAnimateVertically = true;
@@ -52,6 +54,7 @@ UThrobber::UThrobber(const FObjectInitializer& ObjectInitializer)
 		PostEditChange();
 	}
 #endif // WITH_EDITOR
+PRAGMA_ENABLE_DEPRECATION_WARNINGS
 }
 
 void UThrobber::ReleaseSlateResources(bool bReleaseChildren)
@@ -63,11 +66,12 @@ void UThrobber::ReleaseSlateResources(bool bReleaseChildren)
 
 TSharedRef<SWidget> UThrobber::RebuildWidget()
 {
+PRAGMA_DISABLE_DEPRECATION_WARNINGS
 	MyThrobber = SNew(SThrobber)
 		.PieceImage(&Image)
 		.NumPieces(FMath::Clamp(NumberOfPieces, 1, 25))
 		.Animate(GetAnimation());
-
+PRAGMA_ENABLE_DEPRECATION_WARNINGS
 	return MyThrobber.ToSharedRef();
 }
 
@@ -75,10 +79,14 @@ void UThrobber::SynchronizeProperties()
 {
 	Super::SynchronizeProperties();
 
+PRAGMA_DISABLE_DEPRECATION_WARNINGS
+	MyThrobber->SetPieceImage(&Image);
 	MyThrobber->SetNumPieces(FMath::Clamp(NumberOfPieces, 1, 25));
 	MyThrobber->SetAnimate(GetAnimation());
+PRAGMA_ENABLE_DEPRECATION_WARNINGS
 }
 
+PRAGMA_DISABLE_DEPRECATION_WARNINGS
 SThrobber::EAnimation UThrobber::GetAnimation() const
 {
 	const int32 AnimationParams = (bAnimateVertically ? SThrobber::Vertical : 0) |
@@ -90,13 +98,24 @@ SThrobber::EAnimation UThrobber::GetAnimation() const
 
 void UThrobber::SetNumberOfPieces(int32 InNumberOfPieces)
 {
-	NumberOfPieces = InNumberOfPieces;
+	int32 NewNumberOfPieces = FMath::Clamp(InNumberOfPieces, 1, 25);
+	if (NewNumberOfPieces != InNumberOfPieces)
+	{
+		UE_LOG(LogSlate, Warning, TEXT("The number of Pieces was clamped between 1 and 25"));
+	}
+
+	NumberOfPieces = NewNumberOfPieces;
 	if (MyThrobber.IsValid())
 	{
-		MyThrobber->SetNumPieces(FMath::Clamp(NumberOfPieces, 1, 25));
+		MyThrobber->SetNumPieces(NumberOfPieces);
 	}
 }
 
+int32 UThrobber::GetNumberOfPieces() const
+{
+	return NumberOfPieces;
+}
+
 void UThrobber::SetAnimateHorizontally(bool bInAnimateHorizontally)
 {
 	bAnimateHorizontally = bInAnimateHorizontally;
@@ -106,6 +125,11 @@ void UThrobber::SetAnimateHorizontally(bool bInAnimateHorizontally)
 	}
 }
 
+bool UThrobber::IsAnimateHorizontally() const
+{
+	return bAnimateHorizontally;
+}
+
 void UThrobber::SetAnimateVertically(bool bInAnimateVertically)
 {
 	bAnimateVertically = bInAnimateVertically;
@@ -115,7 +139,12 @@ void UThrobber::SetAnimateVertically(bool bInAnimateVertically)
 	}
 }
 
-void UThrobber:: SetAnimateOpacity(bool bInAnimateOpacity)
+bool UThrobber::IsAnimateVertically() const
+{
+	return bAnimateVertically;
+}
+
+void UThrobber::SetAnimateOpacity(bool bInAnimateOpacity)
 {
 	bAnimateOpacity = bInAnimateOpacity;
 	if (MyThrobber.IsValid())
@@ -124,6 +153,28 @@ void UThrobber:: SetAnimateOpacity(bool bInAnimateOpacity)
 	}
 }
 
+bool UThrobber::IsAnimateOpacity() const
+{
+	return bAnimateOpacity;
+}
+
+void UThrobber::SetImage(const FSlateBrush& Brush)
+{
+	Image = Brush;
+	if (MyThrobber.IsValid())
+	{
+		// to force a new invalidation
+		MyThrobber->SetPieceImage(nullptr);
+		MyThrobber->SetPieceImage(&Image);
+	}
+}
+
+const FSlateBrush& UThrobber::GetImage() const
+{
+	return Image;
+}
+PRAGMA_ENABLE_DEPRECATION_WARNINGS
+
 #if WITH_EDITOR
 
 const FText UThrobber::GetPaletteCategory()
diff --git a/Engine/Source/Runtime/UMG/Public/Components/Throbber.h b/Engine/Source/Runtime/UMG/Public/Components/Throbber.h
--- a/Engine/Source/Runtime/UMG/Public/Components/Throbber.h
+++ b/Engine/Source/Runtime/UMG/Public/Components/Throbber.h
@@ -23,22 +23,28 @@ class UMG_API UThrobber : public UWidget
 public:
 	
 	/** How many pieces there are */
-	UPROPERTY(EditAnywhere, BlueprintReadOnly, Category=Appearance, meta=( ClampMin="1", ClampMax="25", UIMin="1", UIMax="25" ))
+	UE_DEPRECATED(5.2, "Direct access to NumberOfPieces is deprecated. Please use the getter or setter.")
+	UPROPERTY(EditAnywhere, BlueprintReadWrite, Getter, Setter, BlueprintSetter="SetNumberOfPieces", Category=Appearance, meta=( ClampMin="1", ClampMax="25", UIMin="1", UIMax="25" ))
 	int32 NumberOfPieces;
 
 	/** Should the pieces animate horizontally? */
-	UPROPERTY(EditAnywhere, BlueprintReadOnly, Category=Appearance)
+	UE_DEPRECATED(5.2, "Direct access to bAnimateHorizontally is deprecated. Please use the getter or setter.")
+	UPROPERTY(EditAnywhere, BlueprintReadWrite, Getter="IsAnimateHorizontally", Setter="SetAnimateHorizontally", BlueprintSetter = "SetAnimateHorizontally", Category = Appearance)
 	bool bAnimateHorizontally;
 
 	/** Should the pieces animate vertically? */
-	UPROPERTY(EditAnywhere, BlueprintReadOnly, Category=Appearance)
+	UE_DEPRECATED(5.2, "Direct access to bAnimateVertically is deprecated. Please use the getter or setter.")
+	UPROPERTY(EditAnywhere, BlueprintReadWrite, Getter="IsAnimateVertically", Setter="SetAnimateVertically", BlueprintSetter="SetAnimateVertically", Category = Appearance)
 	bool bAnimateVertically;
 
 	/** Should the pieces animate their opacity? */
-	UPROPERTY(EditAnywhere, BlueprintReadOnly, Category=Appearance)
+	UE_DEPRECATED(5.2, "Direct access to bAnimateOpacity is deprecated. Please use the getter or setter.")
+	UPROPERTY(EditAnywhere, BlueprintReadWrite, Getter="IsAnimateOpacity", Setter="SetAnimateOpacity", BlueprintSetter="SetAnimateOpacity", Category = Appearance)
 	bool bAnimateOpacity;
 
-	UPROPERTY(EditAnywhere, BlueprintReadOnly, Category=Appearance)
+	/** The animated pieces. */
+	UE_DEPRECATED(5.2, "Direct access to Image is deprecated. Please use the getter or setter.")
+	UPROPERTY(EditAnywhere, BlueprintReadWrite, Getter, Setter, Category=Appearance)
 	FSlateBrush Image;
 
 public:
@@ -46,18 +52,25 @@ public:
 	/** Sets how many pieces there are */
 	UFUNCTION(BlueprintCallable, Category="Appearance")
 	void SetNumberOfPieces(int32 InNumberOfPieces);
+	int32 GetNumberOfPieces() const;
 
 	/** Sets whether the pieces animate horizontally. */
 	UFUNCTION(BlueprintCallable, Category="Appearance")
 	void SetAnimateHorizontally(bool bInAnimateHorizontally);
+	bool IsAnimateHorizontally() const;
 
 	/** Sets whether the pieces animate vertically. */
 	UFUNCTION(BlueprintCallable, Category="Appearance")
 	void SetAnimateVertically(bool bInAnimateVertically);
+	bool IsAnimateVertically() const;
 
 	/** Sets whether the pieces animate their opacity. */
 	UFUNCTION(BlueprintCallable, Category = "Appearance")
 	void SetAnimateOpacity(bool bInAnimateOpacity);
+	bool IsAnimateOpacity() const;
+
+	void SetImage(const FSlateBrush& Brush);
+	const FSlateBrush& GetImage() const;
 
 	//~ Begin UWidget Interface
 	virtual void SynchronizeProperties() override;
@@ -77,7 +90,6 @@ protected:
 	//~ End UWidget Interface
 
 private:
-
 	/** Gets the combined value of the animation properties as a single SThrobber::EAnimation value. */
 	SThrobber::EAnimation GetAnimation() const;
 
