UMG: Add GameViewportSubsystem. It manage the UWidget that are fullscreen or on a player screen. Previously, the logic was only on UserWidget.

diff --git a/Engine/Source/Runtime/LevelSequence/Private/LevelSequenceActor.cpp b/Engine/Source/Runtime/LevelSequence/Private/LevelSequenceActor.cpp
--- a/Engine/Source/Runtime/LevelSequence/Private/LevelSequenceActor.cpp
+++ b/Engine/Source/Runtime/LevelSequence/Private/LevelSequenceActor.cpp
@@ -302,7 +302,7 @@ void ALevelSequenceActor::RefreshBurnIn()
 {
 	if (BurnInInstance)
 	{
-		BurnInInstance->RemoveFromViewport();
+		BurnInInstance->RemoveFromParent();
 		BurnInInstance = nullptr;
 	}
 	
diff --git a/Engine/Source/Runtime/UMG/Private/Blueprint/GameViewportSubsystem.cpp b/Engine/Source/Runtime/UMG/Private/Blueprint/GameViewportSubsystem.cpp
new file mode 100644
--- /dev/null
+++ b/Engine/Source/Runtime/UMG/Private/Blueprint/GameViewportSubsystem.cpp
@@ -0,0 +1,274 @@
+// Copyright Epic Games, Inc. All Rights Reserved.
+
+#include "Blueprint/GameViewportSubsystem.h"
+
+#include "Blueprint/UserWidget.h"
+#include "Blueprint/WidgetLayoutLibrary.h"
+#include "Components/Widget.h"
+#include "Engine/GameViewportClient.h"
+#include "Engine/LocalPlayer.h"
+#include "Engine/World.h"
+
+/*
+ *
+ */
+namespace UE::UMG::Private
+{
+	TPair<FMargin, bool> CalculateOffsetArgument(const FGameViewportWidgetSlot& Slot)
+	{
+		// If the size is zero, and we're not stretched, then use the desired size.
+		FVector2D FinalSize = FVector2D(Slot.Offsets.Right, Slot.Offsets.Bottom);
+		bool bUseAutoSize = FinalSize.IsZero() && !Slot.Anchors.IsStretchedVertical() && !Slot.Anchors.IsStretchedHorizontal();
+		if (bUseAutoSize)
+		{
+			return TPair<FMargin, bool>(FinalSize, bUseAutoSize);
+		}
+		else
+		{
+			return TPair<FMargin, bool>(FMargin(Slot.Offsets.Left, Slot.Offsets.Top, FinalSize.X, FinalSize.Y), bUseAutoSize);
+		}
+	}
+}
+ 
+ /*
+ *
+ */
+void UGameViewportSubsystem::Initialize(FSubsystemCollectionBase& Collection)
+{
+	Super::Initialize(Collection);
+	FWorldDelegates::LevelRemovedFromWorld.AddUObject(this, &ThisClass::OnLevelRemovedFromWorld);
+}
+
+void UGameViewportSubsystem::Deinitialize()
+{
+	FWorldDelegates::LevelRemovedFromWorld.RemoveAll(this);
+	Super::Deinitialize();
+}
+
+UGameViewportSubsystem* UGameViewportSubsystem::Get(UWorld* World)
+{
+	if (World)
+	{
+		if (UGameInstance* GameInstance = World->GetGameInstance())
+		{
+			return GameInstance->GetSubsystem<UGameViewportSubsystem>();
+		}
+	}
+	return nullptr;
+}
+
+bool UGameViewportSubsystem::IsWidgetAdded(const UWidget* Widget) const
+{
+	if (!Widget)
+	{
+		FFrame::KismetExecutionMessage(TEXT("The Widget is invalid."), ELogVerbosity::Warning);
+		return false;
+	}
+
+	if (Widget->bIsManagedByGameViewportSubsystem)
+	{
+		TWeakObjectPtr<UWidget> TempObject = const_cast<UWidget*>(Widget);
+		if (const FSlotInfo* FoundSlot = ViewportWidgets.Find(TempObject))
+		{
+			return FoundSlot->FullScreenWidget.IsValid();
+		}
+	}
+
+	return false;
+}
+
+void UGameViewportSubsystem::AddWidget(UWidget* Widget, FGameViewportWidgetSlot Slot)
+{
+	AddToScreen(Widget, nullptr, Slot);
+}
+
+void UGameViewportSubsystem::AddWidgetForPlayer(UWidget* Widget, ULocalPlayer* Player, FGameViewportWidgetSlot Slot)
+{
+	if (!Player)
+	{
+		FFrame::KismetExecutionMessage(TEXT("The Player is invalid."), ELogVerbosity::Warning);
+		return;
+	}
+	AddToScreen(Widget, Player, Slot);
+}
+
+void UGameViewportSubsystem::AddToScreen(UWidget* Widget, ULocalPlayer* Player, FGameViewportWidgetSlot& Slot)
+{
+	if (!Widget)
+	{
+		FFrame::KismetExecutionMessage(TEXT("The Widget is invalid."), ELogVerbosity::Warning);
+		return;
+	}
+
+	if (UPanelWidget* ParentPanel = Widget->GetParent())
+	{
+		FFrame::KismetExecutionMessage(*FString::Printf(TEXT("The widget '%s' already has a parent widget. It can't also be added to the viewport!"), *Widget->GetName()), ELogVerbosity::Warning);
+		return;
+	}
+
+	// Add the widget to the current worlds viewport.
+	UWorld* World = GetGameInstance()->GetWorld();
+	if (!World || !World->IsGameWorld())
+	{
+		FFrame::KismetExecutionMessage(*FString::Printf(TEXT("The widget '%s' does not have a World."), *Widget->GetName()), ELogVerbosity::Warning);
+		return;
+	}
+
+	UGameViewportClient* ViewportClient = World->GetGameViewport();
+	if (!ViewportClient)
+	{
+		FFrame::KismetExecutionMessage(*FString::Printf(TEXT("No game viewport was found."), *Widget->GetName()), ELogVerbosity::Warning);
+		return;
+	}
+
+	FSlotInfo& SlotInfo = ViewportWidgets.FindOrAdd(Widget);
+	Widget->bIsManagedByGameViewportSubsystem = true;
+
+	if (SlotInfo.FullScreenWidget.IsValid())
+	{
+		FFrame::KismetExecutionMessage(*FString::Printf(TEXT("The widget '%s' was already added to the screen."), *Widget->GetName()), ELogVerbosity::Warning);
+		return;
+	}
+
+	TPair<FMargin, bool> OffsetArgument = UE::UMG::Private::CalculateOffsetArgument(Slot);
+	SConstraintCanvas::FSlot* RawSlot = nullptr;
+	TSharedRef<SConstraintCanvas> FullScreenCanvas = SNew(SConstraintCanvas)
+		+ SConstraintCanvas::Slot()
+		.Offset(OffsetArgument.Get<0>())
+		.AutoSize(OffsetArgument.Get<1>())
+		.Anchors(Slot.Anchors)
+		.Alignment(Slot.Alignment)
+		.Expose(RawSlot)
+		[
+			Widget->TakeWidget()
+		];
+
+	SlotInfo.Slot = Slot;
+	SlotInfo.LocalPlayer = Player;
+	SlotInfo.FullScreenWidget = FullScreenCanvas;
+	SlotInfo.FullScreenWidgetSlot = RawSlot;
+
+	if (Player)
+	{
+		ViewportClient->AddViewportWidgetForPlayer(Player, FullScreenCanvas, Slot.ZOrder);
+	}
+	else
+	{
+		// We add 10 to the zorder when adding to the viewport to avoid 
+		// displaying below any built-in controls, like the virtual joysticks on mobile builds.
+		ViewportClient->AddViewportWidgetContent(FullScreenCanvas, Slot.ZOrder + 10);
+	}
+}
+
+void UGameViewportSubsystem::RemoveWidget(UWidget* Widget)
+{
+	if (Widget && Widget->bIsManagedByGameViewportSubsystem)
+	{
+		FSlotInfo SlotInfo;
+		ViewportWidgets.RemoveAndCopyValue(Widget, SlotInfo);
+		Widget->bIsManagedByGameViewportSubsystem = false;
+		if (TSharedPtr<SWidget> WidgetHost = SlotInfo.FullScreenWidget.Pin())
+		{
+			// If this is a game world remove the widget from the current world's viewport.
+			UWorld* World = GetGameInstance()->GetWorld();
+			if (World && World->IsGameWorld())
+			{
+				if (UGameViewportClient* ViewportClient = World->GetGameViewport())
+				{
+					TSharedRef<SWidget> WidgetHostRef = WidgetHost.ToSharedRef();
+					ViewportClient->RemoveViewportWidgetContent(WidgetHostRef);
+
+					// We may no longer have access to our owning player if the player controller was destroyed
+					// Passing nullptr to RemoveViewportWidgetForPlayer will search all player layers for this widget
+					ViewportClient->RemoveViewportWidgetForPlayer(SlotInfo.LocalPlayer.Get(), WidgetHostRef);
+				}
+			}
+		}
+	}
+}
+
+FGameViewportWidgetSlot UGameViewportSubsystem::GetWidgetSlot(const UWidget* Widget) const
+{
+	if (Widget && Widget->bIsManagedByGameViewportSubsystem)
+	{
+		TWeakObjectPtr<UWidget> TempObject = const_cast<UWidget*>(Widget);
+		const FSlotInfo* SlotInfo = ViewportWidgets.Find(TempObject);
+		ensureMsgf(SlotInfo != nullptr, TEXT("The bIsManagedByGameViewportSubsystem should matches the ViewportWidgets state."));
+		if (SlotInfo)
+		{
+			return SlotInfo->Slot;
+		}
+	}
+	return FGameViewportWidgetSlot();
+}
+
+void UGameViewportSubsystem::SetWidgetSlot(UWidget* Widget, FGameViewportWidgetSlot Slot)
+{
+	if (Widget && !Widget->HasAnyFlags(RF_BeginDestroyed))
+	{
+		TWeakObjectPtr<UWidget> TempObject = const_cast<UWidget*>(Widget);
+		FSlotInfo& SlotInfo = ViewportWidgets.FindOrAdd(TempObject);
+		Widget->bIsManagedByGameViewportSubsystem = true;
+		SlotInfo.Slot = Slot;
+		if (TSharedPtr<SConstraintCanvas> WidgetHost = SlotInfo.FullScreenWidget.Pin())
+		{
+			check(SlotInfo.FullScreenWidgetSlot);
+			TPair<FMargin, bool> OffsetArgument = UE::UMG::Private::CalculateOffsetArgument(Slot);
+			SlotInfo.FullScreenWidgetSlot->SetOffset(OffsetArgument.Get<0>());
+			SlotInfo.FullScreenWidgetSlot->SetAutoSize(OffsetArgument.Get<1>());
+			SlotInfo.FullScreenWidgetSlot->SetAnchors(Slot.Anchors);
+			SlotInfo.FullScreenWidgetSlot->SetAlignment(Slot.Alignment);
+			WidgetHost->Invalidate(EInvalidateWidgetReason::Layout);
+			//todo set the zorder
+		}
+	}
+}
+
+
+FGameViewportWidgetSlot UGameViewportSubsystem::SetWidgetSlotPosition(FGameViewportWidgetSlot Slot, const UWidget* Widget, FVector2D Position, bool bRemoveDPIScale)
+{
+	if (bRemoveDPIScale)
+	{
+		const float Scale = UWidgetLayoutLibrary::GetViewportScale(Widget);
+		Position /= Scale;
+	}
+
+	Slot.Offsets.Left = Position.X;
+	Slot.Offsets.Top = Position.Y;
+	Slot.Anchors = FAnchors(0.f, 0.f);
+
+	return Slot;
+}
+
+FGameViewportWidgetSlot UGameViewportSubsystem::SetWidgetSlotDesiredSize(FGameViewportWidgetSlot Slot, FVector2D Size)
+{
+	Slot.Offsets.Right = Size.X;
+	Slot.Offsets.Bottom = Size.Y;
+	Slot.Anchors = FAnchors(0.f, 0.f);
+	return Slot;
+}
+
+void UGameViewportSubsystem::OnLevelRemovedFromWorld(ULevel* InLevel, UWorld* InWorld)
+{
+	// If the InLevel is null, it's a signal that the entire world is about to disappear, so
+	// go ahead and remove this widget from the viewport, it could be holding onto too many
+	// dangerous actor references that won't carry over into the next world.
+	if (InLevel == nullptr)
+	{
+		for (FViewportWidget::TIterator Itt = ViewportWidgets.CreateIterator(); Itt; ++Itt)
+		{
+			if (UWidget* Widget = Itt.Key().Get())
+			{
+				if (InWorld == Widget->GetWorld())
+				{
+					Widget->RemoveFromParent();
+				}
+			}
+			else
+			{
+				Itt.RemoveCurrent();
+			}
+		}
+	}
+}
+
diff --git a/Engine/Source/Runtime/UMG/Private/Components/Widget.cpp b/Engine/Source/Runtime/UMG/Private/Components/Widget.cpp
--- a/Engine/Source/Runtime/UMG/Private/Components/Widget.cpp
+++ b/Engine/Source/Runtime/UMG/Private/Components/Widget.cpp
@@ -21,6 +21,7 @@
 #include "Binding/PropertyBinding.h"
 #include "Binding/WidgetFieldNotificationExtension.h"
 #include "Logging/MessageLog.h"
+#include "Blueprint/GameViewportSubsystem.h"
 #include "Blueprint/UserWidget.h"
 #include "Blueprint/UserWidgetBlueprint.h"
 #include "Blueprint/WidgetBlueprintGeneratedClass.h"
@@ -179,6 +180,7 @@ UWidget::UWidget(const FObjectInitializer& ObjectInitializer)
 PRAGMA_DISABLE_DEPRECATION_WARNINGS
 	bIsEnabled = true;
 	bIsVariable = true;
+	bIsManagedByGameViewportSubsystem = false;
 #if WITH_EDITOR
 	DesignerFlags = static_cast<uint8>(EWidgetDesignFlags::None);
 #endif
@@ -313,6 +315,18 @@ void UWidget::SetIsEnabled(bool bInIsEnabled)
 	}
 }
 
+bool UWidget::IsInViewport() const
+{
+	if (bIsManagedByGameViewportSubsystem)
+	{
+		if (UGameViewportSubsystem* Subsystem = UGameViewportSubsystem::Get(GetWorld()))
+		{
+			return Subsystem->IsWidgetAdded(this);
+		}
+	}
+	return false;
+}
+
 EMouseCursor::Type UWidget::GetCursor() const
 {
 	return Cursor;
@@ -830,8 +844,14 @@ void UWidget::RemoveFromParent()
 {
 	if (!HasAnyFlags(RF_BeginDestroyed))
 	{
-		UPanelWidget* CurrentParent = GetParent();
-		if (CurrentParent)
+		if (bIsManagedByGameViewportSubsystem)
+		{
+			if (UGameViewportSubsystem* Subsystem = UGameViewportSubsystem::Get(GetWorld()))
+			{
+				Subsystem->RemoveWidget(this);
+			}
+		}
+		else if (UPanelWidget* CurrentParent = GetParent())
 		{
 			CurrentParent->RemoveChild(this);
 		}
@@ -1506,6 +1526,18 @@ UWorld* UWidget::GetWorld() const
 	return nullptr;
 }
 
+void UWidget::BeginDestroy()
+{
+	if (bIsManagedByGameViewportSubsystem)
+	{
+		if (UGameViewportSubsystem* Subsystem = UGameViewportSubsystem::Get(GetWorld()))
+		{
+			Subsystem->RemoveWidget(this);
+		}
+	}
+	Super::BeginDestroy();
+}
+
 void UWidget::FinishDestroy()
 {
 	Super::FinishDestroy();
diff --git a/Engine/Source/Runtime/UMG/Private/UserWidget.cpp b/Engine/Source/Runtime/UMG/Private/UserWidget.cpp
--- a/Engine/Source/Runtime/UMG/Private/UserWidget.cpp
+++ b/Engine/Source/Runtime/UMG/Private/UserWidget.cpp
@@ -22,8 +22,8 @@
 #include "Animation/WidgetAnimation.h"
 #include "MovieScene.h"
 #include "Interfaces/ITargetPlatform.h"
+#include "Blueprint/GameViewportSubsystem.h"
 #include "Blueprint/WidgetBlueprintLibrary.h"
-#include "Blueprint/WidgetLayoutLibrary.h"
 #include "UObject/EditorObjectVersion.h"
 #include "UMGPrivate.h"
 #include "UObject/ObjectSaveContext.h"
@@ -78,7 +78,6 @@ UUserWidget::UUserWidget(const FObjectInitializer& ObjectInitializer)
 	, bStoppingAllAnimations(false)
 	, TickFrequency(EWidgetTickFrequency::Auto)
 {
-	ViewportAnchors = FAnchors(0, 0, 1, 1);
 	SetVisibilityInternal(ESlateVisibility::SelfHitTestInvisible);
 
 	bIsFocusable = false;
@@ -486,14 +485,6 @@ void UUserWidget::Invalidate(EInvalidateWidgetReason InvalidateReason)
 	}
 }
 
-void UUserWidget::InvalidateFullScreenWidget(EInvalidateWidgetReason InvalidateReason)
-{
-	if (TSharedPtr<SWidget> FullScreenWidgetPinned = FullScreenWidget.Pin())
-	{
-		FullScreenWidgetPinned->Invalidate(InvalidateReason);
-	}
-}
-
 UUMGSequencePlayer* UUserWidget::PlayAnimation(UWidgetAnimation* InAnimation, float StartAtTime, int32 NumberOfLoops, EUMGSequencePlayMode::Type PlayMode, float PlaybackSpeed, bool bRestoreState)
 {
 	SCOPED_NAMED_EVENT_TEXT("Widget::PlayAnimation", FColor::Emerald);
@@ -927,132 +918,49 @@ UWidget* UUserWidget::GetRootWidget() const
 
 void UUserWidget::AddToViewport(int32 ZOrder)
 {
-	AddToScreen(nullptr, ZOrder);
-}
-
-bool UUserWidget::AddToPlayerScreen(int32 ZOrder)
-{
-	if ( ULocalPlayer* LocalPlayer = GetOwningLocalPlayer() )
+	if (UGameViewportSubsystem* Subsystem = UGameViewportSubsystem::Get(GetWorld()))
 	{
-		AddToScreen(LocalPlayer, ZOrder);
-		return true;
-	}
-
-	FMessageLog("PIE").Error(LOCTEXT("AddToPlayerScreen_NoPlayer", "AddToPlayerScreen Failed.  No Owning Player!"));
-	return false;
-}
-
-void UUserWidget::AddToScreen(ULocalPlayer* Player, int32 ZOrder)
-{
-	if ( !FullScreenWidget.IsValid() )
-	{
-		if ( UPanelWidget* ParentPanel = GetParent() )
+		FGameViewportWidgetSlot ViewportSlot;
+		if (bIsManagedByGameViewportSubsystem)
 		{
-			FMessageLog("PIE").Error(FText::Format(LOCTEXT("WidgetAlreadyHasParent", "The widget '{0}' already has a parent widget.  It can't also be added to the viewport!"),
-				FText::FromString(GetClass()->GetName())));
-			return;
+			ViewportSlot = Subsystem->GetWidgetSlot(this);
 		}
-
-		// First create and initialize the variable so that users calling this function twice don't
-		// attempt to add the widget to the viewport again.
-		TSharedRef<SConstraintCanvas> FullScreenCanvas = SNew(SConstraintCanvas);
-		FullScreenWidget = FullScreenCanvas;
-
-		TSharedRef<SWidget> UserSlateWidget = TakeWidget();
-
-		FullScreenCanvas->AddSlot()
-			.Offset(BIND_UOBJECT_ATTRIBUTE(FMargin, GetFullScreenOffset))
-			.Anchors(BIND_UOBJECT_ATTRIBUTE(FAnchors, GetAnchorsInViewport))
-			.Alignment(BIND_UOBJECT_ATTRIBUTE(FVector2D, GetAlignmentInViewport))
-			[
-				UserSlateWidget
-			];
-
-		// If this is a game world add the widget to the current worlds viewport.
-		UWorld* World = GetWorld();
-		if ( World && World->IsGameWorld() )
-		{
-			if ( UGameViewportClient* ViewportClient = World->GetGameViewport() )
-			{
-				if ( Player )
-				{
-					ViewportClient->AddViewportWidgetForPlayer(Player, FullScreenCanvas, ZOrder);
-				}
-				else
-				{
-					// We add 10 to the zorder when adding to the viewport to avoid 
-					// displaying below any built-in controls, like the virtual joysticks on mobile builds.
-					ViewportClient->AddViewportWidgetContent(FullScreenCanvas, ZOrder + 10);
-				}
-
-				// Just in case we already hooked this delegate, remove the handler.
-				FWorldDelegates::LevelRemovedFromWorld.RemoveAll(this);
-
-				// Widgets added to the viewport are automatically removed if the persistent level is unloaded.
-				FWorldDelegates::LevelRemovedFromWorld.AddUObject(this, &UUserWidget::OnLevelRemovedFromWorld);
-			}
-		}
-	}
-	else
-	{
-		FMessageLog("PIE").Warning(FText::Format(LOCTEXT("WidgetAlreadyOnScreen", "The widget '{0}' was already added to the screen."),
-			FText::FromString(GetClass()->GetName())));
-	}
-}
-
-void UUserWidget::OnLevelRemovedFromWorld(ULevel* InLevel, UWorld* InWorld)
-{
-	// If the InLevel is null, it's a signal that the entire world is about to disappear, so
-	// go ahead and remove this widget from the viewport, it could be holding onto too many
-	// dangerous actor references that won't carry over into the next world.
-	if ( InLevel == nullptr && InWorld == GetWorld() )
-	{
-		RemoveFromParent();
+		ViewportSlot.ZOrder = ZOrder;
+		Subsystem->AddWidget(this, ViewportSlot);
 	}
 }
 
-void UUserWidget::RemoveFromViewport()
-{
-	RemoveFromParent();
-}
-
-void UUserWidget::RemoveFromParent()
+bool UUserWidget::AddToPlayerScreen(int32 ZOrder)
 {
-	if (!HasAnyFlags(RF_BeginDestroyed))
+	if (UGameViewportSubsystem* Subsystem = UGameViewportSubsystem::Get(GetWorld()))
 	{
-		if (FullScreenWidget.IsValid())
+		if (ULocalPlayer* LocalPlayer = GetOwningLocalPlayer())
 		{
-			TSharedPtr<SWidget> WidgetHost = FullScreenWidget.Pin();
-
-			// If this is a game world remove the widget from the current world's viewport.
-			UWorld* World = GetWorld();
-			if (World && World->IsGameWorld())
+			FGameViewportWidgetSlot ViewportSlot;
+			if (bIsManagedByGameViewportSubsystem)
 			{
-				if (UGameViewportClient* ViewportClient = World->GetGameViewport())
-				{
-					TSharedRef<SWidget> WidgetHostRef = WidgetHost.ToSharedRef();
-
-					ViewportClient->RemoveViewportWidgetContent(WidgetHostRef);
-					
-					// We may no longer have access to our owning player if the player controller was destroyed
-					// Passing nullptr to RemoveViewportWidgetForPlayer will search all player layers for this widget
-					ULocalPlayer* MaybeOwningLocalPlayer = GetOwningLocalPlayer();
-					ViewportClient->RemoveViewportWidgetForPlayer(MaybeOwningLocalPlayer, WidgetHostRef);
-
-					FWorldDelegates::LevelRemovedFromWorld.RemoveAll(this);
-				}
+				ViewportSlot = Subsystem->GetWidgetSlot(this);
 			}
+			ViewportSlot.ZOrder = ZOrder;
+			Subsystem->AddWidgetForPlayer(this, GetOwningLocalPlayer(), ViewportSlot);
+			return true;
 		}
 		else
 		{
-			Super::RemoveFromParent();
+			FMessageLog("PIE").Error(LOCTEXT("AddToPlayerScreen_NoPlayer", "AddToPlayerScreen Failed.  No Owning Player!"));
 		}
 	}
+	return false;
+}
+
+void UUserWidget::RemoveFromViewport()
+{
+	RemoveFromParent();
 }
 
 bool UUserWidget::GetIsVisible() const
 {
-	return FullScreenWidget.IsValid();
+	return IsInViewport();
 }
 
 void UUserWidget::SetVisibility(ESlateVisibility InVisibility)
@@ -1068,11 +976,6 @@ void UUserWidget::SetVisibility(ESlateVisibility InVisibility)
 	}
 }
 
-bool UUserWidget::IsInViewport() const
-{
-	return FullScreenWidget.IsValid();
-}
-
 void UUserWidget::SetPlayerContext(const FLocalPlayerContext& InPlayerContext)
 {
 	PlayerContext = InPlayerContext;
@@ -1146,82 +1049,120 @@ APlayerCameraManager* UUserWidget::GetOwningPlayerCameraManager() const
 	return nullptr;
 }
 
-void UUserWidget::SetPositionInViewport(FVector2D Position, bool bRemoveDPIScale )
+void UUserWidget::SetPositionInViewport(FVector2D Position, bool bRemoveDPIScale)
 {
-	if (bRemoveDPIScale)
-	{
-		const float Scale = UWidgetLayoutLibrary::GetViewportScale(this);
-		Position /= Scale;
-	}
-
-	FAnchors Zero{ 0.f, 0.f };
-	if (ViewportOffsets.Left != Position.X
-		|| ViewportOffsets.Top != Position.Y
-		|| ViewportAnchors != Zero)
+	if (UGameViewportSubsystem* Subsystem = UGameViewportSubsystem::Get(GetWorld()))
 	{
-		ViewportOffsets.Left = Position.X;
-		ViewportOffsets.Top = Position.Y;
-		ViewportAnchors = Zero;
-		InvalidateFullScreenWidget(EInvalidateWidgetReason::Layout);
+		if (bIsManagedByGameViewportSubsystem)
+		{
+			FGameViewportWidgetSlot ViewportSlot = Subsystem->GetWidgetSlot(this);
+			ViewportSlot = UGameViewportSubsystem::SetWidgetSlotPosition(ViewportSlot, this, Position, bRemoveDPIScale);
+			Subsystem->SetWidgetSlot(this, ViewportSlot);
+		}
+		else
+		{
+			FGameViewportWidgetSlot ViewportSlot = UGameViewportSubsystem::SetWidgetSlotPosition(FGameViewportWidgetSlot(), this, Position, bRemoveDPIScale);
+			Subsystem->SetWidgetSlot(this, ViewportSlot);
+		}
 	}
 }
 
 void UUserWidget::SetDesiredSizeInViewport(FVector2D DesiredSize)
 {
-	FAnchors Zero{0.f, 0.f};
-	if (ViewportOffsets.Right != DesiredSize.X
-		|| ViewportOffsets.Bottom != DesiredSize.Y
-		|| ViewportAnchors != Zero)
+	if (UGameViewportSubsystem* Subsystem = UGameViewportSubsystem::Get(GetWorld()))
 	{
-		ViewportOffsets.Right = DesiredSize.X;
-		ViewportOffsets.Bottom = DesiredSize.Y;
-		ViewportAnchors = Zero;
-		InvalidateFullScreenWidget(EInvalidateWidgetReason::Layout);
+		if (bIsManagedByGameViewportSubsystem)
+		{
+			FGameViewportWidgetSlot ViewportSlot = Subsystem->GetWidgetSlot(this);
+			ViewportSlot = UGameViewportSubsystem::SetWidgetSlotDesiredSize(ViewportSlot, DesiredSize);
+			Subsystem->SetWidgetSlot(this, ViewportSlot);
+		}
+		else
+		{
+			FGameViewportWidgetSlot ViewportSlot = UGameViewportSubsystem::SetWidgetSlotDesiredSize(FGameViewportWidgetSlot(), DesiredSize);
+			Subsystem->SetWidgetSlot(this, ViewportSlot);
+		}
 	}
-
 }
 
 void UUserWidget::SetAnchorsInViewport(FAnchors Anchors)
 {
-	if (ViewportAnchors != Anchors)
+	if (UGameViewportSubsystem* Subsystem = UGameViewportSubsystem::Get(GetWorld()))
 	{
-		ViewportAnchors = Anchors;
-		InvalidateFullScreenWidget(EInvalidateWidgetReason::Layout);
+		if (bIsManagedByGameViewportSubsystem)
+		{
+			FGameViewportWidgetSlot ViewportSlot = Subsystem->GetWidgetSlot(this);
+			if (ViewportSlot.Anchors != Anchors)
+			{
+				ViewportSlot.Anchors = Anchors;
+				Subsystem->SetWidgetSlot(this, ViewportSlot);
+			}
+		}
+		else
+		{
+			FGameViewportWidgetSlot ViewportSlot;
+			ViewportSlot.Anchors = Anchors;
+			Subsystem->SetWidgetSlot(this, ViewportSlot);
+		}
 	}
 }
 
 void UUserWidget::SetAlignmentInViewport(FVector2D Alignment)
 {
-	if (ViewportAlignment != Alignment)
+	if (UGameViewportSubsystem* Subsystem = UGameViewportSubsystem::Get(GetWorld()))
 	{
-		ViewportAlignment = Alignment;
-		InvalidateFullScreenWidget(EInvalidateWidgetReason::Layout);
+		if (bIsManagedByGameViewportSubsystem)
+		{
+			FGameViewportWidgetSlot ViewportSlot = Subsystem->GetWidgetSlot(this);
+			if (ViewportSlot.Alignment != Alignment)
+			{
+				ViewportSlot.Alignment = Alignment;
+				Subsystem->SetWidgetSlot(this, ViewportSlot);
+			}
+		}
+		else
+		{
+			FGameViewportWidgetSlot ViewportSlot;
+			ViewportSlot.Alignment = Alignment;
+			Subsystem->SetWidgetSlot(this, ViewportSlot);
+		}
 	}
 }
 
 FMargin UUserWidget::GetFullScreenOffset() const
 {
-	// If the size is zero, and we're not stretched, then use the desired size.
-	FVector2D FinalSize = FVector2D(ViewportOffsets.Right, ViewportOffsets.Bottom);
-	if ( FinalSize.IsZero() && !ViewportAnchors.IsStretchedVertical() && !ViewportAnchors.IsStretchedHorizontal() )
+	if (bIsManagedByGameViewportSubsystem)
 	{
-		if (TSharedPtr<SWidget> CachedWidget = GetCachedWidget())
+		if (UGameViewportSubsystem* Subsystem = UGameViewportSubsystem::Get(GetWorld()))
 		{
-			FinalSize = CachedWidget->GetDesiredSize();
+			return Subsystem->GetWidgetSlot(this).Offsets;
 		}
 	}
-
-	return FMargin(ViewportOffsets.Left, ViewportOffsets.Top, FinalSize.X, FinalSize.Y);
+	return FGameViewportWidgetSlot().Offsets;
 }
 
 FAnchors UUserWidget::GetAnchorsInViewport() const
 {
-	return ViewportAnchors;
+	if (bIsManagedByGameViewportSubsystem)
+	{
+		if (UGameViewportSubsystem* Subsystem = UGameViewportSubsystem::Get(GetWorld()))
+		{
+			return Subsystem->GetWidgetSlot(this).Anchors;
+		}
+	}
+	return FGameViewportWidgetSlot().Anchors;
 }
 
 FVector2D UUserWidget::GetAlignmentInViewport() const
 {
-	return ViewportAlignment;
+	if (bIsManagedByGameViewportSubsystem)
+	{
+		if (UGameViewportSubsystem* Subsystem = UGameViewportSubsystem::Get(GetWorld()))
+		{
+			return Subsystem->GetWidgetSlot(this).Alignment;
+		}
+	}
+	return FGameViewportWidgetSlot().Alignment;
 }
 
 void UUserWidget::RemoveObsoleteBindings(const TArray<FName>& NamedSlots)
diff --git a/Engine/Source/Runtime/UMG/Private/WidgetLayoutLibrary.cpp b/Engine/Source/Runtime/UMG/Private/WidgetLayoutLibrary.cpp
--- a/Engine/Source/Runtime/UMG/Private/WidgetLayoutLibrary.cpp
+++ b/Engine/Source/Runtime/UMG/Private/WidgetLayoutLibrary.cpp
@@ -79,7 +79,7 @@ bool UWidgetLayoutLibrary::ProjectWorldLocationToWidgetPositionWithDistance(APla
 	return false;
 }
 
-float UWidgetLayoutLibrary::GetViewportScale(UObject* WorldContextObject)
+float UWidgetLayoutLibrary::GetViewportScale(const UObject* WorldContextObject)
 {
 	static TFrameValue<float> ViewportScaleCache;
 
@@ -104,7 +104,7 @@ float UWidgetLayoutLibrary::GetViewportScale(UObject* WorldContextObject)
 	return ViewportScaleCache.GetValue();
 }
 
-float UWidgetLayoutLibrary::GetViewportScale(UGameViewportClient* ViewportClient)
+float UWidgetLayoutLibrary::GetViewportScale(const UGameViewportClient* ViewportClient)
 {
 	FVector2D ViewportSize;
 	ViewportClient->GetViewportSize(ViewportSize);
diff --git a/Engine/Source/Runtime/UMG/Public/Blueprint/GameViewportSubsystem.h b/Engine/Source/Runtime/UMG/Public/Blueprint/GameViewportSubsystem.h
new file mode 100644
--- /dev/null
+++ b/Engine/Source/Runtime/UMG/Public/Blueprint/GameViewportSubsystem.h
@@ -0,0 +1,117 @@
+// Copyright Epic Games, Inc. All Rights Reserved.
+
+#pragma once
+
+#include "Subsystems/GameInstanceSubsystem.h"
+#include "Layout/Margin.h"
+#include "Widgets/Layout/Anchors.h"
+#include "Widgets/Layout/SConstraintCanvas.h"
+#include "GameViewportSubsystem.generated.h"
+
+class ULevel;
+class ULocalPlayer;
+class UWidget;
+class UWorld;
+class SConstraintCanvas;
+
+/**
+ * The default value fills the entire screen / player region.
+ */
+USTRUCT(BlueprintType)
+struct UMG_API FGameViewportWidgetSlot
+{
+	GENERATED_BODY()
+
+	UPROPERTY(EditAnywhere, BlueprintReadWrite, Category = "User Interface")
+	FAnchors Anchors = FAnchors(0.f, 0.f, 1.f, 1.f);
+
+	UPROPERTY(EditAnywhere, BlueprintReadWrite, Category = "User Interface")
+	FMargin Offsets;
+
+	UPROPERTY(EditAnywhere, BlueprintReadWrite, Category = "User Interface")
+	FVector2D Alignment;
+
+	/** The higher the number, the more on top this widget will be. */
+	UPROPERTY(EditAnywhere, BlueprintReadWrite, Category = "User Interface")
+	int32 ZOrder = 0;
+};
+
+
+/**
+ * 
+ */
+UCLASS()
+class UMG_API UGameViewportSubsystem : public UGameInstanceSubsystem
+{
+	GENERATED_BODY()
+
+public:
+	//~ Begin Subsystem
+	virtual void Initialize(FSubsystemCollectionBase& Collection) override;
+	virtual void Deinitialize() override;
+	//~ End Subsystem
+
+	static UGameViewportSubsystem* Get(UWorld* World);
+	
+public:
+	/* @return true if the widget was added to the viewport using AddWidget or AddWidgetForPlayer. */
+	UFUNCTION(BlueprintPure, BlueprintCosmetic, Category = "User Interface")
+	bool IsWidgetAdded(const UWidget* Widget) const;
+
+	/** Adds it to the game's viewport. */
+	UFUNCTION(BlueprintCallable, BlueprintCosmetic, Category = "User Interface")
+	void AddWidget(UWidget* Widget, FGameViewportWidgetSlot Slot);
+
+	/**
+	 * Adds the widget to the game's viewport in the section dedicated to the player. This is valuable in a split screen
+	 * game where you need to only show a widget over a player's portion of the viewport.
+	 */
+	UFUNCTION(BlueprintCallable, BlueprintCosmetic, Category = "User Interface")
+	void AddWidgetForPlayer(UWidget* Widget, ULocalPlayer* Player, FGameViewportWidgetSlot Slot);
+
+	/** Removes the widget from the viewport. */
+	UFUNCTION(BlueprintCallable, BlueprintCosmetic, Category = "User Interface")
+	void RemoveWidget(UWidget* Widget);
+
+	/* The slot info from previously added widget or info that is store for later. */
+	UFUNCTION(BlueprintCallable, BlueprintCosmetic, Category = "User Interface")
+	FGameViewportWidgetSlot GetWidgetSlot(const UWidget* Widget) const;
+
+	/* Update the slot info of a previously added widget or Store the slot info for later use. */
+	UFUNCTION(BlueprintCallable, BlueprintCosmetic, Category = "User Interface")
+	void SetWidgetSlot(UWidget* Widget, FGameViewportWidgetSlot Slot);
+
+public:
+	/**
+	 * Helper function to set the position in the viewport for the Slot.
+	 * @param Position The 2D position to set the widget to in the viewport.
+	 * @param bRemoveDPIScale If you've already calculated inverse DPI, set this to false.
+	 * Otherwise inverse DPI is applied to the position so that when the location is scaled
+	 * by DPI, it ends up in the expected position.
+	 */
+	UFUNCTION(BlueprintCallable, Category = "GameViewportWidgetSlot")
+	static FGameViewportWidgetSlot SetWidgetSlotPosition(FGameViewportWidgetSlot Slot, const UWidget* Widget, FVector2D Position, bool bRemoveDPIScale);
+
+	/**
+	 * Helper function to set the desired size in the viewport for the Slot.
+	 */
+	UFUNCTION(BlueprintCallable, Category = "GameViewportWidgetSlot")
+	static FGameViewportWidgetSlot SetWidgetSlotDesiredSize(FGameViewportWidgetSlot Slot, FVector2D Size);
+
+private:
+	void AddToScreen(UWidget* Widget, ULocalPlayer* Player, FGameViewportWidgetSlot& Slot);
+	void OnLevelRemovedFromWorld(ULevel* InLevel, UWorld* InWorld);
+	FMargin GetFullScreenOffsetForWidget(UWidget* Widget) const;
+	TPair<FMargin, bool> GetOffsetAttribute(UWidget* Widget, const FGameViewportWidgetSlot& Slot) const;
+
+private:
+	struct FSlotInfo
+	{
+		FGameViewportWidgetSlot Slot;
+		TWeakObjectPtr<ULocalPlayer> LocalPlayer;
+		TWeakPtr<SConstraintCanvas> FullScreenWidget;
+		SConstraintCanvas::FSlot* FullScreenWidgetSlot = nullptr;
+	};
+	using FViewportWidget = TMap<TWeakObjectPtr<UWidget>, FSlotInfo>;
+	FViewportWidget ViewportWidgets;
+};
diff --git a/Engine/Source/Runtime/UMG/Public/Blueprint/UserWidget.h b/Engine/Source/Runtime/UMG/Public/Blueprint/UserWidget.h
--- a/Engine/Source/Runtime/UMG/Public/Blueprint/UserWidget.h
+++ b/Engine/Source/Runtime/UMG/Public/Blueprint/UserWidget.h
@@ -192,8 +192,6 @@ enum class EDesignPreviewSizeMode : uint8
 
 #endif
 
-//TODO UMG If you want to host a widget that's full screen there may need to be a SWindow equivalent that you spawn it into.
-
 DECLARE_DYNAMIC_MULTICAST_DELEGATE(FOnConstructEvent);
 
 DECLARE_DYNAMIC_DELEGATE( FOnInputAction );
@@ -277,15 +275,10 @@ public:
 	/**
 	 * Removes the widget from the viewport.
 	 */
+	UE_DEPRECATED(5.1, "RemoveFromViewport is deprecated. Use RemoveFromParent instead.")
 	UFUNCTION(BlueprintCallable, BlueprintCosmetic, Category="User Interface|Viewport", meta=( DeprecatedFunction, DeprecationMessage="Use RemoveFromParent instead" ))
 	void RemoveFromViewport();
 
-	/**
-	 * Removes the widget from its parent widget.  If this widget was added to the player's screen or the viewport
-	 * it will also be removed from those containers.
-	 */
-	virtual void RemoveFromParent() override;
-
 	/**
 	 * Sets the widgets position in the viewport.
 	 * @param Position The 2D position to set the widget to in the viewport.
@@ -317,16 +310,13 @@ public:
 	FVector2D GetAlignmentInViewport() const;
 
 	/*  */
+	UE_DEPRECATED(5.1, "GetIsVisible is deprecated. Please use IsInViewport instead.")
 	UFUNCTION(BlueprintPure, BlueprintCosmetic, Category="Appearance", meta=( DeprecatedFunction, DeprecationMessage="Use IsInViewport instead" ))
 	bool GetIsVisible() const;
 
 	/** Sets the visibility of the widget. */
 	virtual void SetVisibility(ESlateVisibility InVisibility) override;
 
-	/* @return true if the widget was added to the viewport using AddToViewport. */
-	UFUNCTION(BlueprintPure, BlueprintCosmetic, Category="Appearance")
-	bool IsInViewport() const;
-
 	/** Sets the player context associated with this UI. */
 	void SetPlayerContext(const FLocalPlayerContext& InPlayerContext);
 
@@ -1281,24 +1271,12 @@ private:
 	uint8 bStoppingAllAnimations : 1;
 
 protected:
-
-	/** Adds the widget to the screen, either to the viewport or to the player's screen depending on if the LocalPlayer is null. */
-	virtual void AddToScreen(ULocalPlayer* LocalPlayer, int32 ZOrder);
-
-	/**
-	 * Called when a top level widget is in the viewport and the world is potentially coming to and end. When this occurs, 
-	 * it's not save to keep widgets on the screen.  We automatically remove them when this happens and mark them for pending kill.
-	 */
-	virtual void OnLevelRemovedFromWorld(ULevel* InLevel, UWorld* InWorld);
-
 	virtual TSharedRef<SWidget> RebuildWidget() override;
 	virtual void OnWidgetRebuilt() override;
 
+	UE_DEPRECATED(5.1, "GetFullScreenOffset is deprecated. Use the GameViewportSubsystem.")
 	FMargin GetFullScreenOffset() const;
 
-	//native SObjectWidget methods (see the corresponding BlueprintImplementableEvent declarations above for more info on each)
-	friend class SObjectWidget;
-
 	virtual void NativeOnInitialized();
 	virtual void NativePreConstruct();
 	virtual void NativeConstruct();
@@ -1445,14 +1423,6 @@ protected:
 private:
 	static void OnLatentActionsChanged(UObject* ObjectWhichChanged, ELatentActionChangeType ChangeType);
 
-	void InvalidateFullScreenWidget(EInvalidateWidgetReason InvalidateReason);
-
-	FAnchors ViewportAnchors;
-	FMargin ViewportOffsets;
-	FVector2D ViewportAlignment;
-
-	TWeakPtr<SWidget> FullScreenWidget;
-
 	/** The player context that is associated with this UI.  Think of this as the owner of the UI. */
 	FLocalPlayerContext PlayerContext;
 
diff --git a/Engine/Source/Runtime/UMG/Public/Blueprint/WidgetLayoutLibrary.h b/Engine/Source/Runtime/UMG/Public/Blueprint/WidgetLayoutLibrary.h
--- a/Engine/Source/Runtime/UMG/Public/Blueprint/WidgetLayoutLibrary.h
+++ b/Engine/Source/Runtime/UMG/Public/Blueprint/WidgetLayoutLibrary.h
@@ -9,6 +9,7 @@
 #include "WidgetLayoutLibrary.generated.h"
 
 class APlayerController;
+class UGameViewportClient;
 class UCanvasPanelSlot;
 class UGridSlot;
 class UHorizontalBoxSlot;
@@ -54,12 +55,12 @@ public:
 	 * Gets the current DPI Scale being applied to the viewport and all the Widgets.
 	 */
 	UFUNCTION(BlueprintPure, BlueprintCosmetic, Category="Viewport", meta=( WorldContext="WorldContextObject" ))
-	static float GetViewportScale(UObject* WorldContextObject);
+	static float GetViewportScale(const UObject* WorldContextObject);
 
 	/**
 	 * Gets the current DPI Scale being applied to the viewport and all the Widgets.
 	 */
-	static float GetViewportScale(class UGameViewportClient* ViewportClient);
+	static float GetViewportScale(const UGameViewportClient* ViewportClient);
 
 	/**
 	 * Gets the size of the game viewport.
diff --git a/Engine/Source/Runtime/UMG/Public/Components/Widget.h b/Engine/Source/Runtime/UMG/Public/Components/Widget.h
--- a/Engine/Source/Runtime/UMG/Public/Components/Widget.h
+++ b/Engine/Source/Runtime/UMG/Public/Components/Widget.h
@@ -29,6 +29,7 @@
 
 class ULocalPlayer;
 class SObjectWidget;
+class UGameViewportSubsystem;
 class UPanelSlot;
 class UPropertyBinding;
 class UUserWidget;
@@ -212,6 +213,8 @@ class UMG_API UWidget : public UVisual, public INotifyFieldValueChanged
 {
 	GENERATED_UCLASS_BODY()
 
+	friend UGameViewportSubsystem;
+
 public:
 	UE_FIELD_NOTIFICATION_DECLARE_CLASS_DESCRIPTOR_BASE_BEGIN(UMG_API)
 		UE_FIELD_NOTIFICATION_DECLARE_FIELD(ToolTipText)
@@ -378,6 +381,9 @@ protected:
 	UPROPERTY(EditAnywhere, BlueprintReadOnly, Category="Performance")
 	uint8 bIsVolatile:1;
 
+	/** Cached value that indicate if the widget was added to the GameViewportSubsystem. */
+	uint8 bIsManagedByGameViewportSubsystem:1;
+
 public:
 #if WITH_EDITORONLY_DATA
 	/** Stores the design time flag setting if the widget is hidden inside the designer */
@@ -514,6 +520,10 @@ public:
 	UFUNCTION(BlueprintCallable, Category="Widget")
 	virtual void SetIsEnabled(bool bInIsEnabled);
 
+	/* @return true if the widget was added to the viewport using AddToViewport or AddToPlayerScreen. */
+	UFUNCTION(BlueprintPure, BlueprintCosmetic, Category = "Appearance")
+	bool IsInViewport() const;
+
 	/** @return the tooltip text for the widget. */
 	FText GetToolTipText() const;
 
@@ -948,11 +958,12 @@ public:
 
 	static TSubclassOf<UPropertyBinding> FindBinderClassForDestination(FProperty* Property);
 
-	// Begin UObject
+	//~ Begin UObject
 	virtual UWorld* GetWorld() const override;
+	virtual void BeginDestroy() override;
 	virtual void FinishDestroy() override;
 	virtual void PreSave(FObjectPreSaveContext ObjectSaveContext) override;
-	// End UObject
+	//~ End UObject
 
 	FORCEINLINE bool CanSafelyRouteEvent()
 	{
