UMG: Fix bad GameViewportSubsystem variable initialization
#rnx

diff --git a/Engine/Source/Runtime/UMG/Public/Blueprint/GameViewportSubsystem.h b/Engine/Source/Runtime/UMG/Public/Blueprint/GameViewportSubsystem.h
--- a/Engine/Source/Runtime/UMG/Public/Blueprint/GameViewportSubsystem.h
+++ b/Engine/Source/Runtime/UMG/Public/Blueprint/GameViewportSubsystem.h
@@ -29,7 +29,7 @@ struct UMG_API FGameViewportWidgetSlot
 	FMargin Offsets;
 
 	UPROPERTY(EditAnywhere, BlueprintReadWrite, Category = "User Interface")
-	FVector2D Alignment;
+	FVector2D Alignment = FVector2D(0.f, 0.f);
 
 	/** The higher the number, the more on top this widget will be. */
 	UPROPERTY(EditAnywhere, BlueprintReadWrite, Category = "User Interface")
