Temprarily disabled ensure when stopping widget animations

diff --git a/Engine/Source/Runtime/UMG/Private/UserWidget.cpp b/Engine/Source/Runtime/UMG/Private/UserWidget.cpp
--- a/Engine/Source/Runtime/UMG/Private/UserWidget.cpp
+++ b/Engine/Source/Runtime/UMG/Private/UserWidget.cpp
@@ -700,8 +700,8 @@ void UUserWidget::OnAnimationFinishedPlaying(UUMGSequencePlayer& Player)
 
 	if ( Player.GetPlaybackStatus() == EMovieScenePlayerStatus::Stopped )
 	{
-		ensureAlways(!StoppedSequencePlayers.Contains(&Player));
-		StoppedSequencePlayers.Add(&Player);
+		//ensureAlways(!StoppedSequencePlayers.Contains(&Player));
+		StoppedSequencePlayers.AddUnique(&Player);
 
 		if (AnimationTickManager)
 		{
