Update the WidgetInteractionComponent to utilize the new FInputDeviceId constructor instead of the old one, which was creating an invalid input device ID of 8 for virtual slate users.

#rnx

diff --git a/Engine/Source/Runtime/UMG/Private/Components/WidgetInteractionComponent.cpp b/Engine/Source/Runtime/UMG/Private/Components/WidgetInteractionComponent.cpp
--- a/Engine/Source/Runtime/UMG/Private/Components/WidgetInteractionComponent.cpp
+++ b/Engine/Source/Runtime/UMG/Private/Components/WidgetInteractionComponent.cpp
@@ -9,6 +9,7 @@
 #include "Framework/Application/SlateUser.h"
 #include "Kismet/GameplayStatics.h"
 #include "DrawDebugHelpers.h"
+#include "GenericPlatform/GenericPlatformInputDeviceMapper.h"
 
 #include "Components/WidgetComponent.h"
 
@@ -445,21 +446,39 @@ void UWidgetInteractionComponent::PressPointerKey(FKey Key)
 	FWidgetPath WidgetPathUnderFinger = LastWidgetPath.ToWidgetPath();
 
 	ensure(PointerIndex >= 0);
+
 	FPointerEvent PointerEvent;
+
+	// Find the primary input device for this Slate User
+	FInputDeviceId InputDeviceId = INPUTDEVICEID_NONE;
+	if (TSharedPtr<FSlateUser> SlateUser = FSlateApplication::Get().GetUser(VirtualUser->GetUserIndex()))
+	{
+		FPlatformUserId PlatUser = SlateUser->GetPlatformUserId();
+		InputDeviceId = IPlatformInputDeviceMapper::Get().GetPrimaryInputDeviceForUser(PlatUser);
+	}
+
+	// Just in case there was no input device assigned to this virtual user, get the default platform
+	// input device
+	if (!InputDeviceId.IsValid())
+	{
+		InputDeviceId = IPlatformInputDeviceMapper::Get().GetDefaultInputDevice();
+	}
+	
 	if (Key.IsTouch())
 	{
 		PointerEvent = FPointerEvent(
-			VirtualUser->GetUserIndex(),
+			InputDeviceId,
 			(uint32)PointerIndex,
 			LocalHitLocation,
 			LastLocalHitLocation,
 			1.0f,
 			false);
+		
 	}
 	else
 	{
 		PointerEvent = FPointerEvent(
-			VirtualUser->GetUserIndex(),
+			InputDeviceId,
 			(uint32)PointerIndex,
 			LocalHitLocation,
 			LastLocalHitLocation,
diff --git a/Engine/Source/Runtime/UMG/UMG.Build.cs b/Engine/Source/Runtime/UMG/UMG.Build.cs
--- a/Engine/Source/Runtime/UMG/UMG.Build.cs
+++ b/Engine/Source/Runtime/UMG/UMG.Build.cs
@@ -15,7 +15,8 @@ public class UMG : ModuleRules
 				"Slate",
 				"SlateCore",
 				"RenderCore",
-				"RHI"
+				"RHI",
+				"ApplicationCore"
 			}
 		);
 
