UMG:
- Change UWidget::GetRenderTransform() to return a const ref instead of a copy

diff --git a/Engine/Source/Runtime/UMG/Private/Components/Widget.cpp b/Engine/Source/Runtime/UMG/Private/Components/Widget.cpp
--- a/Engine/Source/Runtime/UMG/Private/Components/Widget.cpp
+++ b/Engine/Source/Runtime/UMG/Private/Components/Widget.cpp
@@ -207,7 +207,7 @@ PRAGMA_ENABLE_DEPRECATION_WARNINGS
 }
 
 PRAGMA_DISABLE_DEPRECATION_WARNINGS
-FWidgetTransform UWidget::GetRenderTransform() const
+const FWidgetTransform& UWidget::GetRenderTransform() const
 {
 	return RenderTransform;
 }
diff --git a/Engine/Source/Runtime/UMG/Public/Components/Widget.h b/Engine/Source/Runtime/UMG/Public/Components/Widget.h
--- a/Engine/Source/Runtime/UMG/Public/Components/Widget.h
+++ b/Engine/Source/Runtime/UMG/Public/Components/Widget.h
@@ -473,7 +473,7 @@ public:
 
 public:
 	/** */
-	FWidgetTransform GetRenderTransform() const;
+	const FWidgetTransform& GetRenderTransform() const;
 
 	/** */
 	UFUNCTION(BlueprintCallable, Category="Widget|Transform")
