Fixed warning on  'unbalanced' deprecation macro
=> https://horde.devtools.epicgames.com/job/633461967a624c2b738896ed?step=d583

diff --git a/Engine/Source/Runtime/UMG/Private/Components/MultiLineEditableTextBox.cpp b/Engine/Source/Runtime/UMG/Private/Components/MultiLineEditableTextBox.cpp
--- a/Engine/Source/Runtime/UMG/Private/Components/MultiLineEditableTextBox.cpp
+++ b/Engine/Source/Runtime/UMG/Private/Components/MultiLineEditableTextBox.cpp
@@ -77,10 +77,8 @@ void UMultiLineEditableTextBox::PostLoad()
 
 	if (GetLinkerCustomVersion(FUE5ReleaseStreamObjectVersion::GUID) < FUE5ReleaseStreamObjectVersion::RemoveDuplicatedStyleInfo)
 	{
-	PRAGMA_DISABLE_DEPRECATION_WARNINGS
 		TextStyle_DEPRECATED.SetFont(WidgetStyle.Font_DEPRECATED);
 		WidgetStyle.SetTextStyle(TextStyle_DEPRECATED);
-	PRAGMA_ENABLE_DEPRECATION_WARNINGS
 	}
 }
 #endif
