Reverted GridPanel warning about row and column span values that are going past the end of the slots array. This was never really a big issue and impacted how we were creating uniform grid panels.

#rnx

diff --git a/Engine/Source/Runtime/UMG/Private/Components/GridPanel.cpp b/Engine/Source/Runtime/UMG/Private/Components/GridPanel.cpp
--- a/Engine/Source/Runtime/UMG/Private/Components/GridPanel.cpp
+++ b/Engine/Source/Runtime/UMG/Private/Components/GridPanel.cpp
@@ -135,47 +135,6 @@ const FText UGridPanel::GetPaletteCategory()
 	return LOCTEXT("Panel", "Panel");
 }
 
-void UGridPanel::ValidateCompiledDefaults(IWidgetCompilerLog& CompileLog) const
-{
-	int32 NumRows = 0;
-	int32 NumColumns = 0;
-
-	// find the maximum row & column
-	for (int32 Idx = 0; Idx < Slots.Num(); ++Idx)
-	{
-		if (const UGridSlot* ChildSlot = Cast<UGridSlot>(Slots[Idx].Get()))
-		{
-			NumRows = FMath::Max(NumRows, ChildSlot->GetRow());
-			NumColumns = FMath::Max(NumColumns, ChildSlot->GetColumn());
-		}
-	}
-
-	// Row span of 1 for the last row is valid
-	NumRows += 1;
-	NumColumns += 1;
-
-	// check if any row span goes past the end of the grid and warn
-	for (int32 Idx = 0; Idx < Slots.Num(); ++Idx)
-	{
-		if (const UGridSlot* ChildSlot = Cast<UGridSlot>(Slots[Idx].Get()))
-		{
-			if (ChildSlot->GetRow() + ChildSlot->GetRowSpan() > NumRows)
-			{
-				FText InfoText = FText::FormatOrdered(LOCTEXT("RowSpanPastEnd", "Slot at row {0}, column {1} has a row span value of {2}, which goes past the end of the grid. This behaviour will be deprecated in the future. Slots should not use row span to stretch themselves past the end of the grid."), ChildSlot->GetRow(), ChildSlot->GetColumn(), ChildSlot->GetRowSpan());
-				CompileLog.Warning(InfoText);
-			}
-
-			if (ChildSlot->GetColumn() + ChildSlot->GetColumnSpan() > NumColumns)
-			{
-				FText InfoText = FText::FormatOrdered(LOCTEXT("ColumnSpanPastEnd", "Slot at row {0}, column {1} has a column span value of {2}, which goes past the end of the grid. This behaviour will be deprecated in the future. Slots should not use column span to stretch themselves past the end of the grid."), ChildSlot->GetRow(), ChildSlot->GetColumn(), ChildSlot->GetColumnSpan());
-				CompileLog.Warning(InfoText);
-			}
-		}
-	}
-
-	UPanelWidget::ValidateCompiledDefaults(CompileLog);
-}
-
 #endif
 
 /////////////////////////////////////////////////////
diff --git a/Engine/Source/Runtime/UMG/Public/Components/GridPanel.h b/Engine/Source/Runtime/UMG/Public/Components/GridPanel.h
--- a/Engine/Source/Runtime/UMG/Public/Components/GridPanel.h
+++ b/Engine/Source/Runtime/UMG/Public/Components/GridPanel.h
@@ -54,7 +54,6 @@ public:
 
 #if WITH_EDITOR
 	virtual const FText GetPaletteCategory() override;
-	virtual void ValidateCompiledDefaults(IWidgetCompilerLog& CompileLog) const override;
 #endif
 
 protected:
