[UMG] Remove all widgets from the GameViewport when the subsystem is removed.

diff --git a/Engine/Source/Runtime/UMG/Private/Blueprint/GameViewportSubsystem.cpp b/Engine/Source/Runtime/UMG/Private/Blueprint/GameViewportSubsystem.cpp
--- a/Engine/Source/Runtime/UMG/Private/Blueprint/GameViewportSubsystem.cpp
+++ b/Engine/Source/Runtime/UMG/Private/Blueprint/GameViewportSubsystem.cpp
@@ -44,6 +44,16 @@ void UGameViewportSubsystem::Initialize(FSubsystemCollectionBase& Collection)
 void UGameViewportSubsystem::Deinitialize()
 {
 	FWorldDelegates::LevelRemovedFromWorld.RemoveAll(this);
+
+	for (auto& Itt : ViewportWidgets)
+	{
+		if (UWidget* Widget = Itt.Key.Get())
+		{
+			RemoveWidgetInternal(Widget, Itt.Value);
+		}
+	}
+	ViewportWidgets.Empty();
+
 	Super::Deinitialize();
 }
 
@@ -168,22 +178,27 @@ void UGameViewportSubsystem::RemoveWidget(UWidget* Widget)
 	{
 		FSlotInfo SlotInfo;
 		ViewportWidgets.RemoveAndCopyValue(Widget, SlotInfo);
-		Widget->bIsManagedByGameViewportSubsystem = false;
-		if (TSharedPtr<SWidget> WidgetHost = SlotInfo.FullScreenWidget.Pin())
+		RemoveWidgetInternal(Widget, SlotInfo);
+	}
+}
+
+void UGameViewportSubsystem::RemoveWidgetInternal(UWidget* Widget, FSlotInfo& SlotInfo)
+{
+	Widget->bIsManagedByGameViewportSubsystem = false;
+	if (TSharedPtr<SWidget> WidgetHost = SlotInfo.FullScreenWidget.Pin())
+	{
+		// If this is a game world remove the widget from the current world's viewport.
+		UWorld* World = GetGameInstance()->GetWorld();
+		if (World && World->IsGameWorld())
 		{
-			// If this is a game world remove the widget from the current world's viewport.
-			UWorld* World = GetGameInstance()->GetWorld();
-			if (World && World->IsGameWorld())
+			if (UGameViewportClient* ViewportClient = World->GetGameViewport())
 			{
-				if (UGameViewportClient* ViewportClient = World->GetGameViewport())
-				{
-					TSharedRef<SWidget> WidgetHostRef = WidgetHost.ToSharedRef();
-					ViewportClient->RemoveViewportWidgetContent(WidgetHostRef);
+				TSharedRef<SWidget> WidgetHostRef = WidgetHost.ToSharedRef();
+				ViewportClient->RemoveViewportWidgetContent(WidgetHostRef);
 
-					// We may no longer have access to our owning player if the player controller was destroyed
-					// Passing nullptr to RemoveViewportWidgetForPlayer will search all player layers for this widget
-					ViewportClient->RemoveViewportWidgetForPlayer(SlotInfo.LocalPlayer.Get(), WidgetHostRef);
-				}
+				// We may no longer have access to our owning player if the player controller was destroyed
+				// Passing nullptr to RemoveViewportWidgetForPlayer will search all player layers for this widget
+				ViewportClient->RemoveViewportWidgetForPlayer(SlotInfo.LocalPlayer.Get(), WidgetHostRef);
 			}
 		}
 	}
@@ -257,7 +272,7 @@ void UGameViewportSubsystem::OnLevelRemovedFromWorld(ULevel* InLevel, UWorld* In
 	// dangerous actor references that won't carry over into the next world.
 	if (InLevel == nullptr)
 	{
-		for (FViewportWidget::TIterator Itt = ViewportWidgets.CreateIterator(); Itt; ++Itt)
+		for (FViewportWidgetList::TIterator Itt = ViewportWidgets.CreateIterator(); Itt; ++Itt)
 		{
 			if (UWidget* Widget = Itt.Key().Get())
 			{
diff --git a/Engine/Source/Runtime/UMG/Public/Blueprint/GameViewportSubsystem.h b/Engine/Source/Runtime/UMG/Public/Blueprint/GameViewportSubsystem.h
--- a/Engine/Source/Runtime/UMG/Public/Blueprint/GameViewportSubsystem.h
+++ b/Engine/Source/Runtime/UMG/Public/Blueprint/GameViewportSubsystem.h
@@ -98,12 +98,6 @@ public:
 	UFUNCTION(BlueprintCallable, Category = "GameViewportWidgetSlot")
 	static FGameViewportWidgetSlot SetWidgetSlotDesiredSize(FGameViewportWidgetSlot Slot, FVector2D Size);
 
-private:
-	void AddToScreen(UWidget* Widget, ULocalPlayer* Player, FGameViewportWidgetSlot& Slot);
-	void OnLevelRemovedFromWorld(ULevel* InLevel, UWorld* InWorld);
-	FMargin GetFullScreenOffsetForWidget(UWidget* Widget) const;
-	TPair<FMargin, bool> GetOffsetAttribute(UWidget* Widget, const FGameViewportWidgetSlot& Slot) const;
-
 private:
 	struct FSlotInfo
 	{
@@ -112,6 +106,15 @@ private:
 		TWeakPtr<SConstraintCanvas> FullScreenWidget;
 		SConstraintCanvas::FSlot* FullScreenWidgetSlot = nullptr;
 	};
-	using FViewportWidget = TMap<TWeakObjectPtr<UWidget>, FSlotInfo>;
-	FViewportWidget ViewportWidgets;
+
+private:
+	void AddToScreen(UWidget* Widget, ULocalPlayer* Player, FGameViewportWidgetSlot& Slot);
+	void RemoveWidgetInternal(UWidget* Widget, FSlotInfo& SlotInfo);
+	void OnLevelRemovedFromWorld(ULevel* InLevel, UWorld* InWorld);
+	FMargin GetFullScreenOffsetForWidget(UWidget* Widget) const;
+	TPair<FMargin, bool> GetOffsetAttribute(UWidget* Widget, const FGameViewportWidgetSlot& Slot) const;
+
+private:
+	using FViewportWidgetList = TMap<TWeakObjectPtr<UWidget>, FSlotInfo>;
+	FViewportWidgetList ViewportWidgets;
 };
