UMG: Ensure that animation players in the process of being stopped cannot be re-used

diff --git a/Engine/Source/Runtime/UMG/Private/UserWidget.cpp b/Engine/Source/Runtime/UMG/Private/UserWidget.cpp
--- a/Engine/Source/Runtime/UMG/Private/UserWidget.cpp
+++ b/Engine/Source/Runtime/UMG/Private/UserWidget.cpp
@@ -423,7 +423,7 @@ UUMGSequencePlayer* UUserWidget::GetOrAddSequencePlayer(UWidgetAnimation* InAnim
 		{
 			// We need to make sure we haven't stopped the animation, otherwise it'll get canceled on the next frame.
 			if (Player->GetAnimation() == InAnimation
-			 && !StoppedSequencePlayers.Contains(Player))
+			 && !StoppedSequencePlayers.Contains(Player) && !Player->IsStopping())
 			{
 				FoundPlayer = Player;
 				break;
diff --git a/Engine/Source/Runtime/UMG/Public/Animation/UMGSequencePlayer.h b/Engine/Source/Runtime/UMG/Public/Animation/UMGSequencePlayer.h
--- a/Engine/Source/Runtime/UMG/Public/Animation/UMGSequencePlayer.h
+++ b/Engine/Source/Runtime/UMG/Public/Animation/UMGSequencePlayer.h
@@ -65,6 +65,9 @@ public:
 	/** Gets the current time position in the player (in seconds). */
 	bool IsPlayingForward() const { return bIsPlayingForward; }
 
+	/** Check whether this player is currently being stopped */
+	bool IsStopping() const { return bIsStopping; }
+
 	/** IMovieScenePlayer interface */
 	virtual FMovieSceneRootEvaluationTemplateInstance& GetEvaluationTemplate() override { return RootTemplateInstance; }
 	virtual UMovieSceneEntitySystemLinker* ConstructEntitySystemLinker() override;
