Fix deprecation issue when running editor utility widget, by moving the deprecation code from the post load to the serialize method.
During the loading process, when the object is compiled, it is duplicated before the call to post load, so the deprecation code was not called before duplication, and as the deprecated properties of the object were not duplicated => the old font was lost.

diff --git a/Engine/Source/Runtime/UMG/Private/Components/EditableTextBox.cpp b/Engine/Source/Runtime/UMG/Private/Components/EditableTextBox.cpp
--- a/Engine/Source/Runtime/UMG/Private/Components/EditableTextBox.cpp
+++ b/Engine/Source/Runtime/UMG/Private/Components/EditableTextBox.cpp
@@ -25,7 +25,9 @@ static FEditableTextBoxStyle* EditorEditableTextBoxStyle = nullptr;
 UEditableTextBox::UEditableTextBox(const FObjectInitializer& ObjectInitializer)
 	: Super(ObjectInitializer)
 {
-
+#if WITH_EDITOR
+	bIsFontDeprecationDone = false;
+#endif
 	PRAGMA_DISABLE_DEPRECATION_WARNINGS
 	IsReadOnly = false;
 	IsPassword = false;
@@ -74,27 +76,23 @@ UEditableTextBox::UEditableTextBox(const FObjectInitializer& ObjectInitializer)
 #endif
 }
 
-#if WITH_EDITORONLY_DATA
-void UEditableTextBox::PostLoad()
+void UEditableTextBox::Serialize(FArchive& Ar)
 {
-	Super::PostLoad();
+	Ar.UsingCustomVersion(FUE5ReleaseStreamObjectVersion::GUID);
+
+	Super::Serialize(Ar);
 
+#if WITH_EDITOR
 PRAGMA_DISABLE_DEPRECATION_WARNINGS
-	if (GetLinkerCustomVersion(FUE5ReleaseStreamObjectVersion::GUID) < FUE5ReleaseStreamObjectVersion::RemoveDuplicatedStyleInfo)
+	if (Ar.IsLoading() && !bIsFontDeprecationDone && GetLinkerCustomVersion(FUE5ReleaseStreamObjectVersion::GUID) < FUE5ReleaseStreamObjectVersion::RemoveDuplicatedStyleInfo)
 	{
 		FTextBlockStyle& TextStyle = WidgetStyle.TextStyle;
 		TextStyle.SetFont(WidgetStyle.Font_DEPRECATED);
+		bIsFontDeprecationDone = true;
 	}
 PRAGMA_ENABLE_DEPRECATION_WARNINGS
-}
-
-void UEditableTextBox::Serialize(FArchive& Ar)
-{
-	Ar.UsingCustomVersion(FUE5ReleaseStreamObjectVersion::GUID);
-
-	Super::Serialize(Ar);
-}
 #endif
+}
 
 void UEditableTextBox::ReleaseSlateResources(bool bReleaseChildren)
 {
diff --git a/Engine/Source/Runtime/UMG/Private/Components/MultiLineEditableTextBox.cpp b/Engine/Source/Runtime/UMG/Private/Components/MultiLineEditableTextBox.cpp
--- a/Engine/Source/Runtime/UMG/Private/Components/MultiLineEditableTextBox.cpp
+++ b/Engine/Source/Runtime/UMG/Private/Components/MultiLineEditableTextBox.cpp
@@ -62,6 +62,8 @@ UMultiLineEditableTextBox::UMultiLineEditableTextBox(const FObjectInitializer& O
 		// The CDO isn't an editor widget and thus won't use the editor style, call post edit change to mark difference from CDO
 		PostEditChange();
 	}
+
+	bIsFontDeprecationDone = false;
 #endif // WITH_EDITOR
 
 	bIsReadOnly = false;
@@ -70,24 +72,20 @@ UMultiLineEditableTextBox::UMultiLineEditableTextBox(const FObjectInitializer& O
 	AutoWrapText = true;
 }
 
-#if WITH_EDITOR
-void UMultiLineEditableTextBox::PostLoad()
+void UMultiLineEditableTextBox::Serialize(FArchive& Ar)
 {
-	Super::PostLoad();
+	Ar.UsingCustomVersion(FUE5ReleaseStreamObjectVersion::GUID);
 
-	if (GetLinkerCustomVersion(FUE5ReleaseStreamObjectVersion::GUID) < FUE5ReleaseStreamObjectVersion::RemoveDuplicatedStyleInfo)
+	Super::Serialize(Ar);
+
+#if WITH_EDITOR
+	if (Ar.IsLoading() && !bIsFontDeprecationDone && GetLinkerCustomVersion(FUE5ReleaseStreamObjectVersion::GUID) < FUE5ReleaseStreamObjectVersion::RemoveDuplicatedStyleInfo)
 	{
 		TextStyle_DEPRECATED.SetFont(WidgetStyle.Font_DEPRECATED);
 		WidgetStyle.SetTextStyle(TextStyle_DEPRECATED);
+		bIsFontDeprecationDone = true;
 	}
-}
 #endif
-
-void UMultiLineEditableTextBox::Serialize(FArchive& Ar)
-{
-	Super::Serialize(Ar);
-
-	Ar.UsingCustomVersion(FUE5ReleaseStreamObjectVersion::GUID);
 }
 
 void UMultiLineEditableTextBox::ReleaseSlateResources(bool bReleaseChildren)
diff --git a/Engine/Source/Runtime/UMG/Public/Components/EditableTextBox.h b/Engine/Source/Runtime/UMG/Public/Components/EditableTextBox.h
--- a/Engine/Source/Runtime/UMG/Public/Components/EditableTextBox.h
+++ b/Engine/Source/Runtime/UMG/Public/Components/EditableTextBox.h
@@ -263,11 +263,7 @@ public:
 #if WITH_EDITOR
 	virtual const FText GetPaletteCategory() override;
 #endif
-
-#if WITH_EDITORONLY_DATA
 	virtual void Serialize(FArchive& Ar) override;
-	virtual void PostLoad() override;
-#endif
 
 protected:
 	//~ Begin UWidget Interface
@@ -286,4 +282,10 @@ protected:
 
 	PROPERTY_BINDING_IMPLEMENTATION(FText, Text);
 	PROPERTY_BINDING_IMPLEMENTATION(FText, HintText);
+
+private:
+#if WITH_EDITORONLY_DATA
+	UPROPERTY()
+	bool bIsFontDeprecationDone;
+#endif
 };
diff --git a/Engine/Source/Runtime/UMG/Public/Components/MultiLineEditableTextBox.h b/Engine/Source/Runtime/UMG/Public/Components/MultiLineEditableTextBox.h
--- a/Engine/Source/Runtime/UMG/Public/Components/MultiLineEditableTextBox.h
+++ b/Engine/Source/Runtime/UMG/Public/Components/MultiLineEditableTextBox.h
@@ -146,9 +146,7 @@ public:
 
 #if WITH_EDITOR
 	virtual const FText GetPaletteCategory() override;
-	virtual void PostLoad() override;
 #endif
-
 	virtual void Serialize(FArchive& Ar) override;
 
 protected:
@@ -163,4 +161,10 @@ protected:
 	TSharedPtr<SMultiLineEditableTextBox> MyEditableTextBlock;
 
 	PROPERTY_BINDING_IMPLEMENTATION(FText, HintText);
+
+private:
+#if WITH_EDITORONLY_DATA
+	UPROPERTY()
+	bool bIsFontDeprecationDone;
+#endif
 };
