Added support to replace rich text block decorators at runtime

diff --git a/Engine/Source/Runtime/Slate/Private/Widgets/Text/SRichTextBlock.cpp b/Engine/Source/Runtime/Slate/Private/Widgets/Text/SRichTextBlock.cpp
--- a/Engine/Source/Runtime/Slate/Private/Widgets/Text/SRichTextBlock.cpp
+++ b/Engine/Source/Runtime/Slate/Private/Widgets/Text/SRichTextBlock.cpp
@@ -193,6 +193,15 @@ void SRichTextBlock::SetDecoratorStyleSet(const ISlateStyle* NewDecoratorStyleSe
 	}
 }
 
+void SRichTextBlock::SetDecorators(TArrayView<TSharedRef<ITextDecorator>> InDecorators)
+{
+	if (Marshaller.IsValid())
+	{
+		Marshaller->SetDecorators(InDecorators);
+		Refresh();
+	}
+}
+
 void SRichTextBlock::SetOverflowPolicy(TOptional<ETextOverflowPolicy> InOverflowPolicy)
 {
 	TextLayoutCache->SetTextOverflowPolicy(InOverflowPolicy);
diff --git a/Engine/Source/Runtime/Slate/Public/Framework/Text/RichTextLayoutMarshaller.h b/Engine/Source/Runtime/Slate/Public/Framework/Text/RichTextLayoutMarshaller.h
--- a/Engine/Source/Runtime/Slate/Public/Framework/Text/RichTextLayoutMarshaller.h
+++ b/Engine/Source/Runtime/Slate/Public/Framework/Text/RichTextLayoutMarshaller.h
@@ -38,6 +38,14 @@ public:
 	/** Set the Multiplier applied on the font size of the text*/
 	void SetFontSizeMultiplier(const float NewFontSizeMultiplier);
 
+	/**
+	 * Replace all decorators on this marshaller
+	 */
+	inline void SetDecorators(TArrayView<TSharedRef<ITextDecorator>> InDecorators)
+	{
+		Decorators = InDecorators;
+	}
+
 	/**
 	 * Append an inline decorator to this marshaller
 	 */
diff --git a/Engine/Source/Runtime/Slate/Public/Widgets/Text/SRichTextBlock.h b/Engine/Source/Runtime/Slate/Public/Widgets/Text/SRichTextBlock.h
--- a/Engine/Source/Runtime/Slate/Public/Widgets/Text/SRichTextBlock.h
+++ b/Engine/Source/Runtime/Slate/Public/Widgets/Text/SRichTextBlock.h
@@ -222,6 +222,9 @@ public:
 	/**  */
 	void SetDecoratorStyleSet(const ISlateStyle* NewDecoratorStyleSet);
 
+	/** Replaces the decorators for this text block */
+	void SetDecorators(TArrayView<TSharedRef<ITextDecorator>> InDecorators);
+
 	/** Sets the overflow policy for this text block */
 	void SetOverflowPolicy(TOptional<ETextOverflowPolicy> InOverflowPolicy);
 
diff --git a/Engine/Source/Runtime/UMG/Private/Components/RichTextBlock.cpp b/Engine/Source/Runtime/UMG/Private/Components/RichTextBlock.cpp
--- a/Engine/Source/Runtime/UMG/Private/Components/RichTextBlock.cpp
+++ b/Engine/Source/Runtime/UMG/Private/Components/RichTextBlock.cpp
@@ -323,6 +323,23 @@ UMaterialInstanceDynamic* URichTextBlock::GetDefaultDynamicMaterial()
 	return nullptr;
 }
 
+void URichTextBlock::SetDecorators(const TArray<TSubclassOf<URichTextBlockDecorator>>& InDecoratorClasses)
+{
+	DecoratorClasses = InDecoratorClasses;
+
+	StyleInstance.Reset();
+	UpdateStyleData();
+
+	if (MyRichTextBlock.IsValid())
+	{
+		TArray<TSharedRef<ITextDecorator>> CreatedDecorators;
+		CreateDecorators(CreatedDecorators);
+
+		MyRichTextBlock->SetDecoratorStyleSet(StyleInstance.Get());
+		MyRichTextBlock->SetDecorators(CreatedDecorators);
+	}
+}
+
 void URichTextBlock::SetDefaultColorAndOpacity(FSlateColor InColorAndOpacity)
 {
 	BeginDefaultStyleOverride();
diff --git a/Engine/Source/Runtime/UMG/Public/Components/RichTextBlock.h b/Engine/Source/Runtime/UMG/Public/Components/RichTextBlock.h
--- a/Engine/Source/Runtime/UMG/Public/Components/RichTextBlock.h
+++ b/Engine/Source/Runtime/UMG/Public/Components/RichTextBlock.h
@@ -126,6 +126,12 @@ public:
 	UFUNCTION(BlueprintCallable, Category = "Appearance")
 	UMaterialInstanceDynamic* GetDefaultDynamicMaterial();
 
+	/**
+	 * Replaces the existing decorators with the list provided
+	 */
+	UFUNCTION(BlueprintCallable, Category = "Appearance")
+	void SetDecorators(const TArray<TSubclassOf<URichTextBlockDecorator>>& InDecoratorClasses);
+
 public:
 	URichTextBlock(const FObjectInitializer& ObjectInitializer);
 
