Unify the 8 separate variations on WriteRawToTexture_RenderThread in to a single function on FTexture2DDynamicResource

diff --git a/Engine/Source/Runtime/Engine/Private/Texture2DDynamic.cpp b/Engine/Source/Runtime/Engine/Private/Texture2DDynamic.cpp
--- a/Engine/Source/Runtime/Engine/Private/Texture2DDynamic.cpp
+++ b/Engine/Source/Runtime/Engine/Private/Texture2DDynamic.cpp
@@ -90,6 +90,36 @@ FTexture2DRHIRef FTexture2DDynamicResource::GetTexture2DRHI()
 	return Texture2DRHI;
 }
 
+#if !UE_SERVER
+void FTexture2DDynamicResource::WriteRawToTexture_RenderThread(TArrayView64<const uint8> RawData)
+{
+	check(IsInRenderingThread());
+
+	const int32 Width = Texture2DRHI->GetSizeX();
+	const int32 Height = Texture2DRHI->GetSizeY();
+
+	uint32 DestStride = 0;
+	uint8* DestData = reinterpret_cast<uint8*>(RHILockTexture2D(Texture2DRHI, 0, RLM_WriteOnly, DestStride, false, false));
+
+	for (int32 y = 0; y < Height; y++)
+	{
+		uint8* DestPtr = &DestData[((int64)Height - 1 - y) * DestStride];
+
+		const FColor* SrcPtr = &((FColor*)(RawData.GetData()))[((int64)Height - 1 - y) * Width];
+		for (int32 x = 0; x < Width; x++)
+		{
+			*DestPtr++ = SrcPtr->B;
+			*DestPtr++ = SrcPtr->G;
+			*DestPtr++ = SrcPtr->R;
+			*DestPtr++ = SrcPtr->A;
+			SrcPtr++;
+		}
+	}
+
+	RHIUnlockTexture2D(Texture2DRHI, 0, false, false);
+}
+
+#endif
 
 /*-----------------------------------------------------------------------------
 	UTexture2DDynamic
diff --git a/Engine/Source/Runtime/Engine/Public/TextureResource.h b/Engine/Source/Runtime/Engine/Public/TextureResource.h
--- a/Engine/Source/Runtime/Engine/Public/TextureResource.h
+++ b/Engine/Source/Runtime/Engine/Public/TextureResource.h
@@ -236,6 +236,10 @@ public:
 	/** Returns the Texture2DRHI, which can be used for locking/unlocking the mips. */
 	ENGINE_API FTexture2DRHIRef GetTexture2DRHI();
 
+#if !UE_SERVER
+	ENGINE_API void WriteRawToTexture_RenderThread(TArrayView64<const uint8> RawData);
+#endif
+
 private:
 	/** The owner of this resource. */
 	class UTexture2DDynamic* Owner;
diff --git a/Engine/Source/Runtime/UMG/Private/Blueprint/AsyncTaskDownloadImage.cpp b/Engine/Source/Runtime/UMG/Private/Blueprint/AsyncTaskDownloadImage.cpp
--- a/Engine/Source/Runtime/UMG/Private/Blueprint/AsyncTaskDownloadImage.cpp
+++ b/Engine/Source/Runtime/UMG/Private/Blueprint/AsyncTaskDownloadImage.cpp
@@ -16,45 +16,6 @@
 // UAsyncTaskDownloadImage
 //----------------------------------------------------------------------//
 
-#if !UE_SERVER
-
-static void WriteRawToTexture_RenderThread(FTexture2DDynamicResource* TextureResource, TArray64<uint8>* RawData, bool bUseSRGB = true)
-{
-	check(IsInRenderingThread());
-
-	if (TextureResource)
-	{
-		FRHITexture2D* TextureRHI = TextureResource->GetTexture2DRHI();
-
-		int32 Width = TextureRHI->GetSizeX();
-		int32 Height = TextureRHI->GetSizeY();
-
-		uint32 DestStride = 0;
-		uint8* DestData = reinterpret_cast<uint8*>(RHILockTexture2D(TextureRHI, 0, RLM_WriteOnly, DestStride, false, false));
-
-		for (int32 y = 0; y < Height; y++)
-		{
-			uint8* DestPtr = &DestData[((int64)Height - 1 - y) * DestStride];
-
-			const FColor* SrcPtr = &((FColor*)(RawData->GetData()))[((int64)Height - 1 - y) * Width];
-			for (int32 x = 0; x < Width; x++)
-			{
-				*DestPtr++ = SrcPtr->B;
-				*DestPtr++ = SrcPtr->G;
-				*DestPtr++ = SrcPtr->R;
-				*DestPtr++ = SrcPtr->A;
-				SrcPtr++;
-			}
-		}
-
-		RHIUnlockTexture2D(TextureRHI, 0, false, false);
-	}
-
-	delete RawData;
-}
-
-#endif
-
 UAsyncTaskDownloadImage::UAsyncTaskDownloadImage(const FObjectInitializer& ObjectInitializer)
 	: Super(ObjectInitializer)
 {
@@ -108,9 +69,9 @@ void UAsyncTaskDownloadImage::HandleImageRequest(FHttpRequestPtr HttpRequest, FH
 		{
 			if ( ImageWrapper.IsValid() && ImageWrapper->SetCompressed(HttpResponse->GetContent().GetData(), HttpResponse->GetContentLength()) )
 			{
-				TArray64<uint8>* RawData = new TArray64<uint8>();
+				TArray64<uint8> RawData;
 				const ERGBFormat InFormat = ERGBFormat::BGRA;
-				if ( ImageWrapper->GetRaw(InFormat, 8, *RawData) )
+				if ( ImageWrapper->GetRaw(InFormat, 8, RawData) )
 				{
 					if ( UTexture2DDynamic* Texture = UTexture2DDynamic::Create(ImageWrapper->GetWidth(), ImageWrapper->GetHeight()) )
 					{
@@ -121,15 +82,11 @@ void UAsyncTaskDownloadImage::HandleImageRequest(FHttpRequestPtr HttpRequest, FH
 						if (TextureResource)
 						{
 							ENQUEUE_RENDER_COMMAND(FWriteRawDataToTexture)(
-								[TextureResource, RawData](FRHICommandListImmediate& RHICmdList)
+								[TextureResource, RawData = MoveTemp(RawData)](FRHICommandListImmediate& RHICmdList)
 								{
-									WriteRawToTexture_RenderThread(TextureResource, RawData);
+									TextureResource->WriteRawToTexture_RenderThread(RawData);
 								});
 						}
-						else
-						{
-							delete RawData;
-						}
 						OnSuccess.Broadcast(Texture);						
 						return;
 					}
