Ensure we have a default font set for Editable text box, like it was done for editable text recently.

Tests: drag n drop a new editable text box: now the font family is set, and we can play with the text size, as well as position of the editable text box, without having it disappear (because we fallback on the composite font that can disappear too!).

diff --git a/Engine/Source/Runtime/UMG/Private/Components/EditableTextBox.cpp b/Engine/Source/Runtime/UMG/Private/Components/EditableTextBox.cpp
--- a/Engine/Source/Runtime/UMG/Private/Components/EditableTextBox.cpp
+++ b/Engine/Source/Runtime/UMG/Private/Components/EditableTextBox.cpp
@@ -51,6 +51,11 @@ UEditableTextBox::UEditableTextBox(const FObjectInitializer& ObjectInitializer)
 	}
 
 	WidgetStyle = *DefaultEditableTextBoxStyle;
+	if (!IsRunningDedicatedServer())
+	{
+		static ConstructorHelpers::FObjectFinder<UFont> defaultFontObj(*UWidget::GetDefaultFontName());
+		WidgetStyle.SetFont(FSlateFontInfo(defaultFontObj.Object, 24, FName("Regular")));
+	}
 
 #if WITH_EDITOR 
 	if (EditorEditableTextBoxStyle == nullptr)
