Reverting a change where a method that was renamed for clarity, but that changes the public API.

diff --git a/Engine/Source/Runtime/Slate/Private/Widgets/Input/SEditableText.cpp b/Engine/Source/Runtime/Slate/Private/Widgets/Input/SEditableText.cpp
--- a/Engine/Source/Runtime/Slate/Private/Widgets/Input/SEditableText.cpp
+++ b/Engine/Source/Runtime/Slate/Private/Widgets/Input/SEditableText.cpp
@@ -104,7 +104,7 @@ FSlateFontInfo SEditableText::GetFont() const
 	return Font.Get();
 }
 
-void SEditableText::SetEditableTextStyle( const FEditableTextStyle& InNewTextStyle )
+void SEditableText::SetTextStyle( const FEditableTextStyle& InNewTextStyle )
 {
 	Font = InNewTextStyle.Font;
 	ColorAndOpacity = InNewTextStyle.ColorAndOpacity;
diff --git a/Engine/Source/Runtime/Slate/Public/Widgets/Input/SEditableText.h b/Engine/Source/Runtime/Slate/Public/Widgets/Input/SEditableText.h
--- a/Engine/Source/Runtime/Slate/Public/Widgets/Input/SEditableText.h
+++ b/Engine/Source/Runtime/Slate/Public/Widgets/Input/SEditableText.h
@@ -246,7 +246,7 @@ public:
 	 *
 	 * @param  NewTextStyle	The new text style to use
 	 */
-	void SetEditableTextStyle( const FEditableTextStyle& InNewTextStyle );
+	void SetTextStyle( const FEditableTextStyle& InNewTextStyle );
 
 	/** @See TextStyle */
 	/**
diff --git a/Engine/Source/Runtime/UMG/Private/Components/EditableText.cpp b/Engine/Source/Runtime/UMG/Private/Components/EditableText.cpp
--- a/Engine/Source/Runtime/UMG/Private/Components/EditableText.cpp
+++ b/Engine/Source/Runtime/UMG/Private/Components/EditableText.cpp
@@ -123,7 +123,7 @@ void UEditableText::SynchronizeProperties()
 	TAttribute<FText> HintTextBinding = PROPERTY_BINDING(FText, HintText);
 
 	MyEditableText->SetText(TextBinding);
-	MyEditableText->SetEditableTextStyle(WidgetStyle);
+	MyEditableText->SetTextStyle(WidgetStyle);
 	MyEditableText->SetHintText(HintTextBinding);
 	MyEditableText->SetIsReadOnly(IsReadOnly);
 	MyEditableText->SetIsPassword(IsPassword);
@@ -348,7 +348,7 @@ void UEditableText::SetFont(FSlateFontInfo InFontInfo)
 
 	if (MyEditableText.IsValid())
 	{
-		MyEditableText->SetEditableTextStyle(WidgetStyle);
+		MyEditableText->SetTextStyle(WidgetStyle);
 	}
 }
 
@@ -358,7 +358,7 @@ void UEditableText::SetFontMaterial(UMaterialInterface* InMaterial)
 
 	if (MyEditableText.IsValid())
 	{
-		MyEditableText->SetEditableTextStyle(WidgetStyle);
+		MyEditableText->SetTextStyle(WidgetStyle);
 	}
 }
 
@@ -368,7 +368,7 @@ void UEditableText::SetFontOutlineMaterial(UMaterialInterface* InMaterial)
 
 	if (MyEditableText.IsValid())
 	{
-		MyEditableText->SetEditableTextStyle(WidgetStyle);
+		MyEditableText->SetTextStyle(WidgetStyle);
 	}
 }
 
