Pause widget component ticking when in disabled tick mode and widget is not visible.

diff --git a/Engine/Source/Runtime/UMG/Private/Components/WidgetComponent.cpp b/Engine/Source/Runtime/UMG/Private/Components/WidgetComponent.cpp
--- a/Engine/Source/Runtime/UMG/Private/Components/WidgetComponent.cpp
+++ b/Engine/Source/Runtime/UMG/Private/Components/WidgetComponent.cpp
@@ -741,9 +741,12 @@ void UWidgetComponent::OnHiddenInGameChanged()
 	Super::OnHiddenInGameChanged();
 	
 	// If the component is changed from hidden to shown in game, we must start a tick to render the widget Component
-	if (bHiddenInGame == false && TickMode == ETickMode::Automatic)
+	if (bHiddenInGame == false)
 	{
-		SetComponentTickEnabled(true);
+		if (ShouldReenableComponentTickWhenWidgetBecomesVisible())
+		{
+			SetComponentTickEnabled(true);
+		}
 	}
 }
 
@@ -941,13 +944,17 @@ EVisibility UWidgetComponent::ConvertWindowVisibilityToVisibility(EWindowVisibil
 
 void UWidgetComponent::OnWidgetVisibilityChanged(ESlateVisibility InVisibility)
 {
-	ensure(TickMode == ETickMode::Automatic);
+	ensure(TickMode != ETickMode::Enabled);
 	ensure(Widget);
 	ensure(bOnWidgetVisibilityChangedRegistered);
 
 	if (InVisibility != ESlateVisibility::Collapsed && InVisibility != ESlateVisibility::Hidden)
 	{
-		SetComponentTickEnabled(true);
+		if (ShouldReenableComponentTickWhenWidgetBecomesVisible())
+		{
+			SetComponentTickEnabled(true);
+		}
+
 		if (bOnWidgetVisibilityChangedRegistered)
 		{
 			Widget->OnNativeVisibilityChanged.RemoveAll(this);
@@ -958,9 +965,7 @@ void UWidgetComponent::OnWidgetVisibilityChanged(ESlateVisibility InVisibility)
 
 void UWidgetComponent::SetWindowVisibility(EWindowVisibility InVisibility)
 {
-	ensure(TickMode == ETickMode::Automatic);
 	ensure(Widget);
-	ensure(bOnWidgetVisibilityChangedRegistered);
 
 	WindowVisibility = InVisibility;
  	if (SlateWindow.IsValid())
@@ -970,7 +975,11 @@ void UWidgetComponent::SetWindowVisibility(EWindowVisibility InVisibility)
 
 	if (IsWidgetVisible())
 	{
-		SetComponentTickEnabled(TickMode != ETickMode::Disabled);
+		if (ShouldReenableComponentTickWhenWidgetBecomesVisible())
+		{
+			SetComponentTickEnabled(true);
+		}
+
 		if (bOnWidgetVisibilityChangedRegistered)
 		{
 			if (Widget)
@@ -990,8 +999,8 @@ void UWidgetComponent::SetTickMode(ETickMode InTickMode)
 
 bool UWidgetComponent::IsWidgetVisible() const
 {
-	//  If we are in World Space, if the SlateWindow is not visible the Widget is not visible.
-	if (Space == EWidgetSpace::World && (!SlateWindow.IsValid() || !SlateWindow->GetVisibility().IsVisible()))
+	//  If we are in World Space, if the component or the SlateWindow is not visible the Widget is not visible.
+	if (Space == EWidgetSpace::World && (!IsVisible() || !SlateWindow.IsValid() || !SlateWindow->GetVisibility().IsVisible()))
 	{
 		return false;
 	}	
@@ -1177,15 +1186,8 @@ void UWidgetComponent::TickComponent(float DeltaTime, enum ELevelTick TickType,
 			return;	
 		}
 
-		// Tick Mode is Disabled, we stop here and Disable the Component Tick
-		if (TickMode == ETickMode::Disabled && !bRedrawRequested)
-		{
-			SetComponentTickEnabled(false);
-			return;
-		}
-
-		// We have a Widget, it's invisible and we are in automatic TickMode, we disable ticking and register a callback to know if visibility changes.
-		if (Widget && TickMode == ETickMode::Automatic && !IsWidgetVisible())
+		// We have a Widget, it's invisible and we are in automatic or disabled TickMode, we disable ticking and register a callback to know if visibility changes.
+		if (Widget && TickMode != ETickMode::Enabled && !IsWidgetVisible())
 		{
 			SetComponentTickEnabled(false);
 			if (!bOnWidgetVisibilityChangedRegistered)
@@ -1195,6 +1197,13 @@ void UWidgetComponent::TickComponent(float DeltaTime, enum ELevelTick TickType,
 			}
 		}
 
+		// Tick Mode is Disabled, we stop here and Disable the Component Tick
+		if (TickMode == ETickMode::Disabled && !bRedrawRequested)
+		{
+			SetComponentTickEnabled(false);
+			return;
+		}
+
 	    if ( Space == EWidgetSpace::World)
 	    {
 			if ( ShouldDrawWidget() )
@@ -1217,6 +1226,11 @@ void UWidgetComponent::TickComponent(float DeltaTime, enum ELevelTick TickType,
 #endif // !UE_SERVER
 }
 
+bool UWidgetComponent::ShouldReenableComponentTickWhenWidgetBecomesVisible() const
+{
+	return (TickMode != ETickMode::Disabled) || bRedrawRequested;
+}
+
 void UWidgetComponent::UpdateWidgetOnScreen()
 {
 	if ((GetUserWidgetObject() || GetSlateWidget().IsValid()) && Space == EWidgetSpace::Screen)
diff --git a/Engine/Source/Runtime/UMG/Public/Components/WidgetComponent.h b/Engine/Source/Runtime/UMG/Public/Components/WidgetComponent.h
--- a/Engine/Source/Runtime/UMG/Public/Components/WidgetComponent.h
+++ b/Engine/Source/Runtime/UMG/Public/Components/WidgetComponent.h
@@ -630,6 +630,7 @@ protected:
 	class FWidgetRenderer* WidgetRenderer;
 
 private: 
+	bool ShouldReenableComponentTickWhenWidgetBecomesVisible() const;
 
 	/** The User Widget object displayed and managed by this component */
 	UPROPERTY(Transient, DuplicateTransient)
