Correct deprecation versions
#rnx

diff --git a/Engine/Plugins/Enterprise/DatasmithImporter/Source/DatasmithTranslator/Public/DatasmithTranslator.h b/Engine/Plugins/Enterprise/DatasmithImporter/Source/DatasmithTranslator/Public/DatasmithTranslator.h
--- a/Engine/Plugins/Enterprise/DatasmithImporter/Source/DatasmithTranslator/Public/DatasmithTranslator.h
+++ b/Engine/Plugins/Enterprise/DatasmithImporter/Source/DatasmithTranslator/Public/DatasmithTranslator.h
@@ -122,7 +122,7 @@ public:
 	 * @param OptionClasses list of classes that will be displayed to the user
 	 */
 	virtual void GetSceneImportOptions(TArray<TObjectPtr<UDatasmithOptionsBase>>& Options) {}
-	UE_DEPRECATED(5.2, "Deprecated, please use same method using array of TObjectPtr instead")
+	UE_DEPRECATED(5.1, "Deprecated, please use same method using array of TObjectPtr instead")
 	virtual void GetSceneImportOptions(TArray<TStrongObjectPtr<UDatasmithOptionsBase>>& Options) {}
 
 	/**
@@ -131,7 +131,7 @@ public:
 	 * @param Options Actual values for the displayed options.
 	 */
 	virtual void SetSceneImportOptions(const TArray<TObjectPtr<UDatasmithOptionsBase>>& Options) {}
-	UE_DEPRECATED(5.2, "Deprecated, please use same method using array of TObjectPtr instead")
+	UE_DEPRECATED(5.1, "Deprecated, please use same method using array of TObjectPtr instead")
 	virtual void SetSceneImportOptions(TArray<TStrongObjectPtr<UDatasmithOptionsBase>>& Options) {}
 
 private:
diff --git a/Engine/Plugins/NetcodeUnitTest/NetcodeUnitTest/Source/NetcodeUnitTest/Public/NUTUtilDebug.h b/Engine/Plugins/NetcodeUnitTest/NetcodeUnitTest/Source/NetcodeUnitTest/Public/NUTUtilDebug.h
--- a/Engine/Plugins/NetcodeUnitTest/NetcodeUnitTest/Source/NetcodeUnitTest/Public/NUTUtilDebug.h
+++ b/Engine/Plugins/NetcodeUnitTest/NetcodeUnitTest/Source/NetcodeUnitTest/Public/NUTUtilDebug.h
@@ -700,7 +700,7 @@ ENUM_CLASS_FLAGS(ELogTraceFlags);
  *
  * Most easily used through the LogTrace console command, as documented in UnitTestManager.cpp
  */
-class UE_DEPRECATED(5.2, "FLogStackTraceManager is deprecated, use FLogTraceManager/GLogTrace now, instead.") FLogStackTraceManager
+class UE_DEPRECATED(5.1, "FLogStackTraceManager is deprecated, use FLogTraceManager/GLogTrace now, instead.") FLogStackTraceManager
 	: public FOutputDevice
 {
 public:
@@ -730,7 +730,7 @@ public:
 	 * @param bPartial	Whether or not to do partial matches for this line (case insensitive, and matches substrings)
 	 * @param bDump		Whether or not to dump traces as they are encountered
 	 */
-	UE_DEPRECATED(5.2, "Use AddLogTrace with ELogTraceFlags now, instead.")
+	UE_DEPRECATED(5.1, "Use AddLogTrace with ELogTraceFlags now, instead.")
 	void AddLogTrace(FString LogLine, bool bPartial=true, bool bDump=true)
 	{
 		ELogTraceFlags TraceFlags = (bPartial ? ELogTraceFlags::Partial : ELogTraceFlags::Full);
diff --git a/Engine/Source/Developer/AssetTools/Public/IAssetTypeActions.h b/Engine/Source/Developer/AssetTools/Public/IAssetTypeActions.h
--- a/Engine/Source/Developer/AssetTools/Public/IAssetTypeActions.h
+++ b/Engine/Source/Developer/AssetTools/Public/IAssetTypeActions.h
@@ -74,16 +74,16 @@ public:
 	virtual FColor GetTypeColor() const = 0;
 	
 	/** Returns true if this class can supply actions for InObjects. */
-//   	UE_DEPRECATED(5.2, "HasActions is no longer used all IAssetTypeActions are assumed to need to have them and GetActions is just called instead.  GetActions is always called unless DEPRECATE_ASSET_TYPE_ACTIONS_CALLING_GETACTIONS is true.")
+//   	UE_DEPRECATED(5.1, "HasActions is no longer used all IAssetTypeActions are assumed to need to have them and GetActions is just called instead.  GetActions is always called unless DEPRECATE_ASSET_TYPE_ACTIONS_CALLING_GETACTIONS is true.")
    	virtual bool HasActions( const TArray<UObject*>& InObjects ) const { return true; }
 
 #if DEPRECATE_ASSET_TYPE_ACTIONS_NEEDING_LOADED_OBJECTS
 	/** Generates a menubuilder for the specified objects. */
-	UE_DEPRECATED(5.2, "Look at 'HOW TO ADD ASSET EXTENSION MENUS' in IAssetTypeActions.h")
+	UE_DEPRECATED(5.1, "Look at 'HOW TO ADD ASSET EXTENSION MENUS' in IAssetTypeActions.h")
 	virtual void GetActions( const TArray<UObject*>& InObjects, class FMenuBuilder& MenuBuilder ) final { }
 
 	/** Generates a menu section for the specified objects. */
-	UE_DEPRECATED(5.2, "Look at 'HOW TO ADD ASSET EXTENSION MENUS' in IAssetTypeActions.h")
+	UE_DEPRECATED(5.1, "Look at 'HOW TO ADD ASSET EXTENSION MENUS' in IAssetTypeActions.h")
 	virtual void GetActions(const TArray<UObject*>& InObjects, struct FToolMenuSection& Section) final { }
 #else
 	/** Generates a menubuilder for the specified objects. */
diff --git a/Engine/Source/Editor/ContentBrowser/Public/ContentBrowserMenuContexts.h b/Engine/Source/Editor/ContentBrowser/Public/ContentBrowserMenuContexts.h
--- a/Engine/Source/Editor/ContentBrowser/Public/ContentBrowserMenuContexts.h
+++ b/Engine/Source/Editor/ContentBrowser/Public/ContentBrowserMenuContexts.h
@@ -36,7 +36,7 @@ public:
 
 	TWeakPtr<IAssetTypeActions> CommonAssetTypeActions;
 
-	UE_DEPRECATED(5.2, "Use SelectedAssets now, this field will not contain any objects.")
+	UE_DEPRECATED(5.1, "Use SelectedAssets now, this field will not contain any objects.")
 	TArray<TWeakObjectPtr<UObject>> SelectedObjects;
 
 	UPROPERTY()
diff --git a/Engine/Source/Runtime/ApplicationCore/Public/GenericPlatform/IForceFeedbackSystem.h b/Engine/Source/Runtime/ApplicationCore/Public/GenericPlatform/IForceFeedbackSystem.h
--- a/Engine/Source/Runtime/ApplicationCore/Public/GenericPlatform/IForceFeedbackSystem.h
+++ b/Engine/Source/Runtime/ApplicationCore/Public/GenericPlatform/IForceFeedbackSystem.h
@@ -21,7 +21,7 @@
  *
  * Note: This class is deprecated and will be removed in favor of IInputInterface
  */
-UE_DEPRECATED(5.2, "IForceFeedbackSystem has been deprecated, use IInputInterface instead.")
+UE_DEPRECATED(5.1, "IForceFeedbackSystem has been deprecated, use IInputInterface instead.")
 class IForceFeedbackSystem
 	: public IInputInterface
 {
diff --git a/Engine/Source/Runtime/Core/Public/GenericPlatform/GenericPlatformChunkInstall.h b/Engine/Source/Runtime/Core/Public/GenericPlatform/GenericPlatformChunkInstall.h
--- a/Engine/Source/Runtime/Core/Public/GenericPlatform/GenericPlatformChunkInstall.h
+++ b/Engine/Source/Runtime/Core/Public/GenericPlatform/GenericPlatformChunkInstall.h
@@ -211,16 +211,16 @@ public:
 	virtual void RemoveChunkInstallDelegate( FDelegateHandle Delegate ) = 0;
 
 
-	UE_DEPRECATED(5.2, "Call SupportsNamedChunkInstall instead")
+	UE_DEPRECATED(5.1, "Call SupportsNamedChunkInstall instead")
 	virtual bool SupportsIntelligentInstall() = 0;
 
-	UE_DEPRECATED(5.2, "Call IsNamedChunkInProgress instead")
+	UE_DEPRECATED(5.1, "Call IsNamedChunkInProgress instead")
 	virtual bool IsChunkInstallationPending(const TArray<FCustomChunk>& ChunkTagsID) = 0;
 
-	UE_DEPRECATED(5.2, "Call InstallNamedChunks instead")
+	UE_DEPRECATED(5.1, "Call InstallNamedChunks instead")
 	virtual bool InstallChunks(const TArray<FCustomChunk>& ChunkTagsID) = 0;
 
-	UE_DEPRECATED(5.2, "Call UninstallNamedChunks instead")
+	UE_DEPRECATED(5.1, "Call UninstallNamedChunks instead")
 	virtual bool UninstallChunks(const TArray<FCustomChunk>& ChunkTagsID) = 0;
 
 	UE_DEPRECATED(5.2, "Call GetNamedChunksByType instead")
diff --git a/Engine/Source/Runtime/Core/Public/GenericPlatform/GenericPlatformMisc.h b/Engine/Source/Runtime/Core/Public/GenericPlatform/GenericPlatformMisc.h
--- a/Engine/Source/Runtime/Core/Public/GenericPlatform/GenericPlatformMisc.h
+++ b/Engine/Source/Runtime/Core/Public/GenericPlatform/GenericPlatformMisc.h
@@ -1703,16 +1703,16 @@ public:
 
 	static bool RequestDeviceCheckToken(TFunction<void(const TArray<uint8>&)> QuerySucceededFunc, TFunction<void(const FString&, const FString&)> QueryFailedFunc);
 
-	UE_DEPRECATED(5.2, "Use named chunks instead")
+	UE_DEPRECATED(5.1, "Use named chunks instead")
 	static TArray<FCustomChunk> GetOnDemandChunksForPakchunkIndices(const TArray<int32>& PakchunkIndices);
 
-	UE_DEPRECATED(5.2, "Use IPlatformChunkInstall::GetNamedChunksByType instead")
+	UE_DEPRECATED(5.1, "Use IPlatformChunkInstall::GetNamedChunksByType instead")
 	static TArray<FCustomChunk> GetAllOnDemandChunks();
 
-	UE_DEPRECATED(5.2, "Use IPlatformChunkInstall::GetNamedChunksByType instead")
+	UE_DEPRECATED(5.1, "Use IPlatformChunkInstall::GetNamedChunksByType instead")
 	static TArray<FCustomChunk> GetAllLanguageChunks();
 
-	UE_DEPRECATED(5.2, "Use IPlatformChunkInstall::GetNamedChunksByType instead")
+	UE_DEPRECATED(5.1, "Use IPlatformChunkInstall::GetNamedChunksByType instead")
 	static TArray<FCustomChunk> GetCustomChunksByType(ECustomChunkType DesiredChunkType);
 
 	/**
diff --git a/Engine/Source/Runtime/Core/Public/Stats/Stats2.h b/Engine/Source/Runtime/Core/Public/Stats/Stats2.h
--- a/Engine/Source/Runtime/Core/Public/Stats/Stats2.h
+++ b/Engine/Source/Runtime/Core/Public/Stats/Stats2.h
@@ -2436,7 +2436,7 @@ DECLARE_FLOAT_COUNTER_STAT_EXTERN(TEXT("Seconds Per Cycle"),STAT_SecondsPerCycle
 #if STATS || ENABLE_STATNAMEDEVENTS
 namespace Stats
 {
-	UE_DEPRECATED(5.2, "No replacement. Engine usage was redundant")
+	UE_DEPRECATED(5.1, "No replacement. Engine usage was redundant")
 	FORCEINLINE bool IsThreadCollectingData()
 	{
 #if STATS && CPUPROFILERTRACE_ENABLED
diff --git a/Engine/Source/Runtime/Core/Public/Tasks/TaskPrivate.h b/Engine/Source/Runtime/Core/Public/Tasks/TaskPrivate.h
--- a/Engine/Source/Runtime/Core/Public/Tasks/TaskPrivate.h
+++ b/Engine/Source/Runtime/Core/Public/Tasks/TaskPrivate.h
@@ -88,7 +88,7 @@ namespace UE::Tasks
 		CORE_API FTaskBase* ExchangeCurrentTask(FTaskBase* Task);
 
 		// Returns true if called from inside a task that is being retracted
-		UE_DEPRECATED(5.2, "You should not use this function as it exists only to patch another system and can be removed any time.")
+		UE_DEPRECATED(5.1, "You should not use this function as it exists only to patch another system and can be removed any time.")
 		CORE_API bool IsThreadRetractingTask();
 
 		// An abstract base class for task implementation. 
diff --git a/Engine/Source/Runtime/CoreUObject/Public/UObject/Class.h b/Engine/Source/Runtime/CoreUObject/Public/UObject/Class.h
--- a/Engine/Source/Runtime/CoreUObject/Public/UObject/Class.h
+++ b/Engine/Source/Runtime/CoreUObject/Public/UObject/Class.h
@@ -2112,7 +2112,7 @@ public:
 	/** searches the list of all enum value names for the specified name
 	 * @return the value the specified name represents if found, otherwise INDEX_NONE
 	 */
-	UE_DEPRECATED(5.2, "LookupEnumName that takes only enum name is deprecated. Please use the version of this function that also takes enum package name.")
+	UE_DEPRECATED(5.1, "LookupEnumName that takes only enum name is deprecated. Please use the version of this function that also takes enum package name.")
 	static int64 LookupEnumName(FName TestName, UEnum** FoundEnum = nullptr)
 	{
 		return LookupEnumName(FName(), TestName, EFindFirstObjectOptions::None, FoundEnum);
@@ -2131,7 +2131,7 @@ public:
 	/** searches the list of all enum value names for the specified name
 	 * @return the value the specified name represents if found, otherwise INDEX_NONE
 	 */
-	UE_DEPRECATED(5.2, "LookupEnumNameSlow that takes only enum name is deprecated. Please use the version of this function that also takes enum package name.")
+	UE_DEPRECATED(5.1, "LookupEnumNameSlow that takes only enum name is deprecated. Please use the version of this function that also takes enum package name.")
 	static int64 LookupEnumNameSlow(const TCHAR* InTestShortName, UEnum** FoundEnum = nullptr)
 	{
 		return LookupEnumNameSlow(FName(), InTestShortName, EFindFirstObjectOptions::None, FoundEnum);
diff --git a/Engine/Source/Runtime/CoreUObject/Public/UObject/Stack.h b/Engine/Source/Runtime/CoreUObject/Public/UObject/Stack.h
--- a/Engine/Source/Runtime/CoreUObject/Public/UObject/Stack.h
+++ b/Engine/Source/Runtime/CoreUObject/Public/UObject/Stack.h
@@ -262,7 +262,7 @@ public:
 	/** 
 	 * This will return a string of the form "ScopeName.FunctionName" associated with this stack frame:
 	 */
-	UE_DEPRECATED(5.2, "Please use GetStackDescription(FStringBuilderBase&).")
+	UE_DEPRECATED(5.1, "Please use GetStackDescription(FStringBuilderBase&).")
 	COREUOBJECT_API FString GetStackDescription() const;
 
 	/**
diff --git a/Engine/Source/Runtime/Datasmith/DatasmithCore/Public/DatasmithDefinitions.h b/Engine/Source/Runtime/Datasmith/DatasmithCore/Public/DatasmithDefinitions.h
--- a/Engine/Source/Runtime/Datasmith/DatasmithCore/Public/DatasmithDefinitions.h
+++ b/Engine/Source/Runtime/Datasmith/DatasmithCore/Public/DatasmithDefinitions.h
@@ -28,7 +28,7 @@ enum class EDatasmithElementType : uint64
 	Camera                         = 1ull << 10,
 	Shader                         = 1ull << 11,
 	BaseMaterial                   = 1ull << 12,
-	MasterMaterial UE_DEPRECATED(5.2, "EDatasmithElementType::MasterMaterial will not be supported in 5.2. Please use EDatasmithElementType::MaterialInstance instead.") = 1ull << 13,
+	MasterMaterial UE_DEPRECATED(5.1, "EDatasmithElementType::MasterMaterial will not be supported in 5.2. Please use EDatasmithElementType::MaterialInstance instead.") = 1ull << 13,
 	MaterialInstance               = 1ull << 13,
 	KeyValueProperty               = 1ull << 14,
 	Texture                        = 1ull << 15,
diff --git a/Engine/Source/Runtime/Datasmith/DatasmithCore/Public/DatasmithSceneFactory.h b/Engine/Source/Runtime/Datasmith/DatasmithCore/Public/DatasmithSceneFactory.h
--- a/Engine/Source/Runtime/Datasmith/DatasmithCore/Public/DatasmithSceneFactory.h
+++ b/Engine/Source/Runtime/Datasmith/DatasmithCore/Public/DatasmithSceneFactory.h
@@ -57,7 +57,7 @@ public:
 	static TSharedRef< IDatasmithMaterialElement > CreateMaterial( const TCHAR* InName );
 
 PRAGMA_DISABLE_DEPRECATION_WARNINGS
-	UE_DEPRECATED(5.2, "FDatasmithSceneFactory::CreateMasterMaterial will not be supported in 5.2. Please use FDatasmithSceneFactory::CreateMaterialInstance instead.")
+	UE_DEPRECATED(5.1, "FDatasmithSceneFactory::CreateMasterMaterial will not be supported in 5.2. Please use FDatasmithSceneFactory::CreateMaterialInstance instead.")
 	static TSharedRef< IDatasmithMasterMaterialElement > CreateMasterMaterial(const TCHAR* InName)
 	{
 		return CreateMaterialInstance(InName);
diff --git a/Engine/Source/Runtime/Datasmith/DatasmithCore/Public/IDatasmithSceneElements.h b/Engine/Source/Runtime/Datasmith/DatasmithCore/Public/IDatasmithSceneElements.h
--- a/Engine/Source/Runtime/Datasmith/DatasmithCore/Public/IDatasmithSceneElements.h
+++ b/Engine/Source/Runtime/Datasmith/DatasmithCore/Public/IDatasmithSceneElements.h
@@ -1788,4 +1788,4 @@ public:
 };
 
 
-using IDatasmithMasterMaterialElement UE_DEPRECATED(5.2, "IDatasmithMasterMaterialElement will not be supported in 5.2. Please use IDatasmithMaterialInstanceElement instead.") = IDatasmithMaterialInstanceElement;
+using IDatasmithMasterMaterialElement UE_DEPRECATED(5.1, "IDatasmithMasterMaterialElement will not be supported in 5.2. Please use IDatasmithMaterialInstanceElement instead.") = IDatasmithMaterialInstanceElement;
diff --git a/Engine/Source/Runtime/Engine/Classes/Engine/DemoNetDriver.h b/Engine/Source/Runtime/Engine/Classes/Engine/DemoNetDriver.h
--- a/Engine/Source/Runtime/Engine/Classes/Engine/DemoNetDriver.h
+++ b/Engine/Source/Runtime/Engine/Classes/Engine/DemoNetDriver.h
@@ -34,7 +34,7 @@ DECLARE_DELEGATE_OneParam(FOnGotoTimeDelegate, const bool /* bWasSuccessful */);
 DECLARE_MULTICAST_DELEGATE_OneParam(FOnDemoStartedDelegate, UDemoNetDriver* /* DemoNetDriver */);
 
 PRAGMA_DISABLE_DEPRECATION_WARNINGS
-UE_DEPRECATED(5.2, "No longer used")
+UE_DEPRECATED(5.1, "No longer used")
 DECLARE_MULTICAST_DELEGATE_TwoParams(FOnDemoFailedToStartDelegate, UDemoNetDriver* /* DemoNetDriver */, EDemoPlayFailure::Type /* FailureType*/);
 PRAGMA_ENABLE_DEPRECATION_WARNINGS
 
@@ -701,7 +701,7 @@ protected:
 	bool DemoReplicateActor(AActor* Actor, UNetConnection* Connection, bool bMustReplicate);
 
 	PRAGMA_DISABLE_DEPRECATION_WARNINGS
-	UE_DEPRECATED(5.2, "Please use NotifyDemoPlaybackError instead")
+	UE_DEPRECATED(5.1, "Please use NotifyDemoPlaybackError instead")
 	void NotifyDemoPlaybackFailure(EDemoPlayFailure::Type FailureType);
 	PRAGMA_ENABLE_DEPRECATION_WARNINGS
 	
diff --git a/Engine/Source/Runtime/Engine/Classes/Engine/EngineBaseTypes.h b/Engine/Source/Runtime/Engine/Classes/Engine/EngineBaseTypes.h
--- a/Engine/Source/Runtime/Engine/Classes/Engine/EngineBaseTypes.h
+++ b/Engine/Source/Runtime/Engine/Classes/Engine/EngineBaseTypes.h
@@ -629,7 +629,7 @@ enum ETravelType
 UENUM(BlueprintType)
 namespace EDemoPlayFailure
 {
-	enum UE_DEPRECATED(5.2, "No longer used in favor of EReplayResult") Type
+	enum UE_DEPRECATED(5.1, "No longer used in favor of EReplayResult") Type
 	{
 		/** A Generic failure. */
 		Generic,
@@ -655,7 +655,7 @@ namespace EDemoPlayFailure
 namespace EDemoPlayFailure
 {
 	PRAGMA_DISABLE_DEPRECATION_WARNINGS
-	UE_DEPRECATED(5.2, "EDemoPlayFailure is now deprecated")
+	UE_DEPRECATED(5.1, "EDemoPlayFailure is now deprecated")
 	inline const TCHAR* ToString(EDemoPlayFailure::Type FailureType)
 	{
 		switch (FailureType)
diff --git a/Engine/Source/Runtime/Engine/Classes/Engine/GameInstance.h b/Engine/Source/Runtime/Engine/Classes/Engine/GameInstance.h
--- a/Engine/Source/Runtime/Engine/Classes/Engine/GameInstance.h
+++ b/Engine/Source/Runtime/Engine/Classes/Engine/GameInstance.h
@@ -431,7 +431,7 @@ public:
 
 	PRAGMA_DISABLE_DEPRECATION_WARNINGS
 	/** Called when demo playback fails for any reason */
-	UE_DEPRECATED(5.2, "Now takes a EReplayResult instead.")
+	UE_DEPRECATED(5.1, "Now takes a EReplayResult instead.")
 	virtual void HandleDemoPlaybackFailure(EDemoPlayFailure::Type FailureType, const FString& ErrorString = TEXT("")) { }
 	PRAGMA_ENABLE_DEPRECATION_WARNINGS
 
diff --git a/Engine/Source/Runtime/Engine/Public/DDSLoader.h b/Engine/Source/Runtime/Engine/Public/DDSLoader.h
--- a/Engine/Source/Runtime/Engine/Public/DDSLoader.h
+++ b/Engine/Source/Runtime/Engine/Public/DDSLoader.h
@@ -95,7 +95,7 @@ struct FDDS10FileHeader
 
 // @todo Oodle : deprecate old DDSLoader after it is removed from InterchangeDDSTranslator 
 // 
-//class UE_DEPRECATED(5.2, "Use DDSFile instead") FDDSLoadHelper;
+//class UE_DEPRECATED(5.1, "Use DDSFile instead") FDDSLoadHelper;
 class FDDSLoadHelper
 {
 public:
diff --git a/Engine/Source/Runtime/Engine/Public/Net/RPCDoSDetection.h b/Engine/Source/Runtime/Engine/Public/Net/RPCDoSDetection.h
--- a/Engine/Source/Runtime/Engine/Public/Net/RPCDoSDetection.h
+++ b/Engine/Source/Runtime/Engine/Public/Net/RPCDoSDetection.h
@@ -471,7 +471,7 @@ public:
 	 * @param Update	Whether or not we are escalating or de-escalating the severity state
 	 * @param Reason	The reason for the escalation change
 	 */
-	UE_DEPRECATED(5.2, "UpdateSeverity will be made private soon.")
+	UE_DEPRECATED(5.1, "UpdateSeverity will be made private soon.")
 	void UpdateSeverity(ERPCDoSSeverityUpdate Update, ERPCDoSEscalateReason Reason)
 	{
 		UpdateSeverity_Private(Update, Reason);
diff --git a/Engine/Source/Runtime/Engine/Public/Net/UnrealNetwork.h b/Engine/Source/Runtime/Engine/Public/Net/UnrealNetwork.h
--- a/Engine/Source/Runtime/Engine/Public/Net/UnrealNetwork.h
+++ b/Engine/Source/Runtime/Engine/Public/Net/UnrealNetwork.h
@@ -64,7 +64,7 @@ DECLARE_DELEGATE_OneParam(FGetOverridableVersionDataForDemoHeaderWriteDelegate,
 DECLARE_MULTICAST_DELEGATE_OneParam(FOnReplayStartedDelegate, UWorld* /*World*/);
 
 PRAGMA_DISABLE_DEPRECATION_WARNINGS
-UE_DEPRECATED(5.2, "Deprecated in favor of FOnReplayPlaybackFailureDelegate")
+UE_DEPRECATED(5.1, "Deprecated in favor of FOnReplayPlaybackFailureDelegate")
 DECLARE_MULTICAST_DELEGATE_TwoParams(FOnReplayStartFailureDelegate, UWorld* /*World*/, EDemoPlayFailure::Type /*Error*/);
 PRAGMA_ENABLE_DEPRECATION_WARNINGS
 
@@ -100,7 +100,7 @@ struct ENGINE_API FNetworkReplayDelegates
 
 	PRAGMA_DISABLE_DEPRECATION_WARNINGS
 	/** Public delegate to be notified when a replay failed to start */
-	UE_DEPRECATED(5.2, "Deprecated in favor of OnReplayPlaybackFailure")
+	UE_DEPRECATED(5.1, "Deprecated in favor of OnReplayPlaybackFailure")
 	static FOnReplayStartFailureDelegate OnReplayStartFailure;
 	PRAGMA_ENABLE_DEPRECATION_WARNINGS
 
diff --git a/Engine/Source/Runtime/NetworkReplayStreaming/HttpNetworkReplayStreaming/Public/HttpNetworkReplayStreaming.h b/Engine/Source/Runtime/NetworkReplayStreaming/HttpNetworkReplayStreaming/Public/HttpNetworkReplayStreaming.h
--- a/Engine/Source/Runtime/NetworkReplayStreaming/HttpNetworkReplayStreaming/Public/HttpNetworkReplayStreaming.h
+++ b/Engine/Source/Runtime/NetworkReplayStreaming/HttpNetworkReplayStreaming/Public/HttpNetworkReplayStreaming.h
@@ -294,7 +294,7 @@ public:
 	void RefreshViewer( const bool bFinal );
 	void ConditionallyRefreshViewer();
 	
-	UE_DEPRECATED(5.2, "No longer used")
+	UE_DEPRECATED(5.1, "No longer used")
 	void SetLastError(const ENetworkReplayError::Type InLastError) { SetLastError(EHttpReplayResult::Unknown); }
 	void SetLastError(FHttpReplayResult&& Result);
 
diff --git a/Engine/Source/Runtime/NetworkReplayStreaming/LocalFileNetworkReplayStreaming/Public/LocalFileNetworkReplayStreaming.h b/Engine/Source/Runtime/NetworkReplayStreaming/LocalFileNetworkReplayStreaming/Public/LocalFileNetworkReplayStreaming.h
--- a/Engine/Source/Runtime/NetworkReplayStreaming/LocalFileNetworkReplayStreaming/Public/LocalFileNetworkReplayStreaming.h
+++ b/Engine/Source/Runtime/NetworkReplayStreaming/LocalFileNetworkReplayStreaming/Public/LocalFileNetworkReplayStreaming.h
@@ -408,7 +408,7 @@ public:
 	FLocalFileReplayInfo ReplayInfo;
 	TArray<uint8> DataBuffer;
 
-	UE_DEPRECATED(5.2, "No longer used")
+	UE_DEPRECATED(5.1, "No longer used")
 	bool bAsyncError = false;
 
 	ELocalFileReplayResult AsyncError = ELocalFileReplayResult::Success;
@@ -610,7 +610,7 @@ protected:
 	void CancelStreamingRequests();
 
 	PRAGMA_DISABLE_DEPRECATION_WARNINGS
-	UE_DEPRECATED(5.2, "No longer used")
+	UE_DEPRECATED(5.1, "No longer used")
 	void SetLastError(const ENetworkReplayError::Type InLastError) 
 	{ 
 		SetLastError(ELocalFileReplayResult::Unknown); 
diff --git a/Engine/Source/Runtime/NetworkReplayStreaming/NetworkReplayStreaming/Public/NetworkReplayStreaming.h b/Engine/Source/Runtime/NetworkReplayStreaming/NetworkReplayStreaming/Public/NetworkReplayStreaming.h
--- a/Engine/Source/Runtime/NetworkReplayStreaming/NetworkReplayStreaming/Public/NetworkReplayStreaming.h
+++ b/Engine/Source/Runtime/NetworkReplayStreaming/NetworkReplayStreaming/Public/NetworkReplayStreaming.h
@@ -101,7 +101,7 @@ struct FNetworkReplayStreamInfo
 
 namespace ENetworkReplayError
 {
-	enum UE_DEPRECATED(5.2, "No longer used") Type
+	enum UE_DEPRECATED(5.1, "No longer used") Type
 	{
 		/** There are currently no issues */
 		None,
@@ -624,7 +624,7 @@ public:
 
 	PRAGMA_DISABLE_DEPRECATION_WARNINGS
 	/** Returns the last error that occurred while streaming replays */
-	UE_DEPRECATED(5.2, "Deprecated in favor of HasError and HandleLastError")
+	UE_DEPRECATED(5.1, "Deprecated in favor of HasError and HandleLastError")
 	virtual ENetworkReplayError::Type GetLastError() const;
 	PRAGMA_ENABLE_DEPRECATION_WARNINGS
 
diff --git a/Engine/Source/Runtime/SlateCore/Public/Styling/SlateTypes.h b/Engine/Source/Runtime/SlateCore/Public/Styling/SlateTypes.h
--- a/Engine/Source/Runtime/SlateCore/Public/Styling/SlateTypes.h
+++ b/Engine/Source/Runtime/SlateCore/Public/Styling/SlateTypes.h
@@ -973,7 +973,7 @@ struct SLATECORE_API FEditableTextBoxStyle : public FSlateWidgetStyle
 
 #if WITH_EDITORONLY_DATA
 	/** Font family and size to be used when displaying this text. */
-	UE_DEPRECATED(5.2, "Font has been deprecated as it was duplicated information already available elsewhere. Please use TextStyle.Font instead.")
+	UE_DEPRECATED(5.1, "Font has been deprecated as it was duplicated information already available elsewhere. Please use TextStyle.Font instead.")
 	UPROPERTY()
 	FSlateFontInfo Font_DEPRECATED;
 #endif
diff --git a/Engine/Source/Runtime/UMG/Public/Components/MultiLineEditableTextBox.h b/Engine/Source/Runtime/UMG/Public/Components/MultiLineEditableTextBox.h
--- a/Engine/Source/Runtime/UMG/Public/Components/MultiLineEditableTextBox.h
+++ b/Engine/Source/Runtime/UMG/Public/Components/MultiLineEditableTextBox.h
@@ -50,7 +50,7 @@ public:
 
 #if WITH_EDITORONLY_DATA
 	/** The text style */
-	UE_DEPRECATED(5.2, "TextStyle has been deprecated as it was mainly duplicated information already available inside WidgetStyle. Please use the WidgetStyle.TextStyle instead.")
+	UE_DEPRECATED(5.1, "TextStyle has been deprecated as it was mainly duplicated information already available inside WidgetStyle. Please use the WidgetStyle.TextStyle instead.")
 	UPROPERTY(meta=(DisplayName="Text Style"))
 	FTextBlockStyle TextStyle_DEPRECATED;
 #endif
diff --git a/Engine/Source/Runtime/UMG/Public/Components/Throbber.h b/Engine/Source/Runtime/UMG/Public/Components/Throbber.h
--- a/Engine/Source/Runtime/UMG/Public/Components/Throbber.h
+++ b/Engine/Source/Runtime/UMG/Public/Components/Throbber.h
@@ -23,27 +23,27 @@ class UMG_API UThrobber : public UWidget
 public:
 	
 	/** How many pieces there are */
-	UE_DEPRECATED(5.2, "Direct access to NumberOfPieces is deprecated. Please use the getter or setter.")
+	UE_DEPRECATED(5.1, "Direct access to NumberOfPieces is deprecated. Please use the getter or setter.")
 	UPROPERTY(EditAnywhere, BlueprintReadWrite, Getter, Setter, BlueprintSetter="SetNumberOfPieces", Category=Appearance, meta=( ClampMin="1", ClampMax="25", UIMin="1", UIMax="25" ))
 	int32 NumberOfPieces;
 
 	/** Should the pieces animate horizontally? */
-	UE_DEPRECATED(5.2, "Direct access to bAnimateHorizontally is deprecated. Please use the getter or setter.")
+	UE_DEPRECATED(5.1, "Direct access to bAnimateHorizontally is deprecated. Please use the getter or setter.")
 	UPROPERTY(EditAnywhere, BlueprintReadWrite, Getter="IsAnimateHorizontally", Setter="SetAnimateHorizontally", BlueprintSetter = "SetAnimateHorizontally", Category = Appearance)
 	bool bAnimateHorizontally;
 
 	/** Should the pieces animate vertically? */
-	UE_DEPRECATED(5.2, "Direct access to bAnimateVertically is deprecated. Please use the getter or setter.")
+	UE_DEPRECATED(5.1, "Direct access to bAnimateVertically is deprecated. Please use the getter or setter.")
 	UPROPERTY(EditAnywhere, BlueprintReadWrite, Getter="IsAnimateVertically", Setter="SetAnimateVertically", BlueprintSetter="SetAnimateVertically", Category = Appearance)
 	bool bAnimateVertically;
 
 	/** Should the pieces animate their opacity? */
-	UE_DEPRECATED(5.2, "Direct access to bAnimateOpacity is deprecated. Please use the getter or setter.")
+	UE_DEPRECATED(5.1, "Direct access to bAnimateOpacity is deprecated. Please use the getter or setter.")
 	UPROPERTY(EditAnywhere, BlueprintReadWrite, Getter="IsAnimateOpacity", Setter="SetAnimateOpacity", BlueprintSetter="SetAnimateOpacity", Category = Appearance)
 	bool bAnimateOpacity;
 
 	/** The animated pieces. */
-	UE_DEPRECATED(5.2, "Direct access to Image is deprecated. Please use the getter or setter.")
+	UE_DEPRECATED(5.1, "Direct access to Image is deprecated. Please use the getter or setter.")
 	UPROPERTY(EditAnywhere, BlueprintReadWrite, Getter, Setter, Category=Appearance)
 	FSlateBrush Image;
 
