UMG: In GameViewportSubsystem, make a copy of the intenal array before removing all the widgets. Removing a widget trigger a callback and user may remove them self invalidating the array.

diff --git a/Engine/Source/Runtime/UMG/Private/Blueprint/GameViewportSubsystem.cpp b/Engine/Source/Runtime/UMG/Private/Blueprint/GameViewportSubsystem.cpp
--- a/Engine/Source/Runtime/UMG/Private/Blueprint/GameViewportSubsystem.cpp
+++ b/Engine/Source/Runtime/UMG/Private/Blueprint/GameViewportSubsystem.cpp
@@ -8,6 +8,7 @@
 #include "Engine/GameViewportClient.h"
 #include "Engine/LocalPlayer.h"
 #include "Engine/World.h"
+#include "Templates/Tuple.h"
 
 #include UE_INLINE_GENERATED_CPP_BY_NAME(GameViewportSubsystem)
 
@@ -45,14 +46,29 @@ void UGameViewportSubsystem::Deinitialize()
 {
 	FWorldDelegates::LevelRemovedFromWorld.RemoveAll(this);
 
-	for (auto& Itt : ViewportWidgets)
+	if (ViewportWidgets.Num() > 0)
 	{
-		if (UWidget* Widget = Itt.Key.Get())
+		using FElement = TTuple<UWidget*, TWeakPtr<SConstraintCanvas>, TWeakObjectPtr<ULocalPlayer>>;
+		FMemMark Mark(FMemStack::Get());
+		TArray<FElement, TMemStackAllocator<>> TmpViewportWidgets;
+		TmpViewportWidgets.Reserve(ViewportWidgets.Num());
+		for (auto& Itt : ViewportWidgets)
 		{
-			RemoveWidgetInternal(Widget, Itt.Value);
+			if (UWidget* Widget = Itt.Key.Get())
+			{
+				TmpViewportWidgets.Emplace(Widget, Itt.Value.FullScreenWidget, Itt.Value.LocalPlayer);
+			}
+		}
+		ViewportWidgets.Empty();
+
+		for (FElement& Element : TmpViewportWidgets)
+		{
+			if (Element.Get<0>()->bIsManagedByGameViewportSubsystem)
+			{
+				RemoveWidgetInternal(Element.Get<0>(), Element.Get<1>(), Element.Get<2>());
+			}
 		}
 	}
-	ViewportWidgets.Empty();
 
 	Super::Deinitialize();
 }
@@ -178,14 +194,14 @@ void UGameViewportSubsystem::RemoveWidget(UWidget* Widget)
 	{
 		FSlotInfo SlotInfo;
 		ViewportWidgets.RemoveAndCopyValue(Widget, SlotInfo);
-		RemoveWidgetInternal(Widget, SlotInfo);
+		RemoveWidgetInternal(Widget, SlotInfo.FullScreenWidget, SlotInfo.LocalPlayer);
 	}
 }
 
-void UGameViewportSubsystem::RemoveWidgetInternal(UWidget* Widget, FSlotInfo& SlotInfo)
+void UGameViewportSubsystem::RemoveWidgetInternal(UWidget* Widget, const TWeakPtr<SConstraintCanvas>& FullScreenWidget, const TWeakObjectPtr<ULocalPlayer>& LocalPlayer)
 {
 	Widget->bIsManagedByGameViewportSubsystem = false;
-	if (TSharedPtr<SWidget> WidgetHost = SlotInfo.FullScreenWidget.Pin())
+	if (TSharedPtr<SWidget> WidgetHost = FullScreenWidget.Pin())
 	{
 		// If this is a game world remove the widget from the current world's viewport.
 		UWorld* World = GetGameInstance()->GetWorld();
@@ -198,7 +214,7 @@ void UGameViewportSubsystem::RemoveWidgetInternal(UWidget* Widget, FSlotInfo& Sl
 
 				// We may no longer have access to our owning player if the player controller was destroyed
 				// Passing nullptr to RemoveViewportWidgetForPlayer will search all player layers for this widget
-				ViewportClient->RemoveViewportWidgetForPlayer(SlotInfo.LocalPlayer.Get(), WidgetHostRef);
+				ViewportClient->RemoveViewportWidgetForPlayer(LocalPlayer.Get(), WidgetHostRef);
 			}
 		}
 	}
diff --git a/Engine/Source/Runtime/UMG/Public/Blueprint/GameViewportSubsystem.h b/Engine/Source/Runtime/UMG/Public/Blueprint/GameViewportSubsystem.h
--- a/Engine/Source/Runtime/UMG/Public/Blueprint/GameViewportSubsystem.h
+++ b/Engine/Source/Runtime/UMG/Public/Blueprint/GameViewportSubsystem.h
@@ -109,7 +109,7 @@ private:
 
 private:
 	void AddToScreen(UWidget* Widget, ULocalPlayer* Player, FGameViewportWidgetSlot& Slot);
-	void RemoveWidgetInternal(UWidget* Widget, FSlotInfo& SlotInfo);
+	void RemoveWidgetInternal(UWidget* Widget, const TWeakPtr<SConstraintCanvas>& FullScreenWidget, const TWeakObjectPtr<ULocalPlayer>& LocalPlayer);
 	void OnLevelRemovedFromWorld(ULevel* InLevel, UWorld* InWorld);
 	FMargin GetFullScreenOffsetForWidget(UWidget* Widget) const;
 	TPair<FMargin, bool> GetOffsetAttribute(UWidget* Widget, const FGameViewportWidgetSlot& Slot) const;
