Fix display of (all flavor of) editable texts when no font family is configured.

While setting a default font in all the editable texts was forcing a font to be configured, it was also clearing the composite font that was previously set in the default ctor, and if the font was overwritten by a null font, we were no longer able to fallback on the composite font. Now we ensure to save and restore the composite font after getting the default font, before applying it to the style of the editable texts.

diff --git a/Engine/Source/Runtime/UMG/Private/Components/EditableText.cpp b/Engine/Source/Runtime/UMG/Private/Components/EditableText.cpp
--- a/Engine/Source/Runtime/UMG/Private/Components/EditableText.cpp
+++ b/Engine/Source/Runtime/UMG/Private/Components/EditableText.cpp
@@ -36,8 +36,12 @@ UEditableText::UEditableText(const FObjectInitializer& ObjectInitializer)
 	WidgetStyle = *DefaultEditableTextStyle;
 	if (!IsRunningDedicatedServer())
 	{
-		static ConstructorHelpers::FObjectFinder<UFont> defaultFontObj(*UWidget::GetDefaultFontName());
-		WidgetStyle.SetFont(FSlateFontInfo(defaultFontObj.Object, 24, FName("Regular")));
+		static ConstructorHelpers::FObjectFinder<UFont> DefaultFontObj(*UWidget::GetDefaultFontName());
+		FSlateFontInfo Font(DefaultFontObj.Object, 24, FName("Regular"));
+		//The FSlateFontInfo just created doesn't contain a composite font (while the default from the WidgetStyle does),
+		//so in the case the Font object is replaced by a null one, we have to keep the composite one as a fallback.
+		Font.CompositeFont = WidgetStyle.Font.CompositeFont;
+		WidgetStyle.SetFont(Font);
 	}
 
 
diff --git a/Engine/Source/Runtime/UMG/Private/Components/EditableTextBox.cpp b/Engine/Source/Runtime/UMG/Private/Components/EditableTextBox.cpp
--- a/Engine/Source/Runtime/UMG/Private/Components/EditableTextBox.cpp
+++ b/Engine/Source/Runtime/UMG/Private/Components/EditableTextBox.cpp
@@ -53,8 +53,13 @@ UEditableTextBox::UEditableTextBox(const FObjectInitializer& ObjectInitializer)
 	WidgetStyle = *DefaultEditableTextBoxStyle;
 	if (!IsRunningDedicatedServer())
 	{
-		static ConstructorHelpers::FObjectFinder<UFont> defaultFontObj(*UWidget::GetDefaultFontName());
-		WidgetStyle.SetFont(FSlateFontInfo(defaultFontObj.Object, 24, FName("Regular")));
+		static ConstructorHelpers::FObjectFinder<UFont> DefaultFontObj(*UWidget::GetDefaultFontName());
+		FSlateFontInfo Font(DefaultFontObj.Object, 24, FName("Regular"));
+		//The FSlateFontInfo just created doesn't contain a composite font (while the default from the WidgetStyle does),
+		//so in the case the Font object is replaced by a null one, we have to keep the composite one as a fallback.
+		Font.CompositeFont = WidgetStyle.TextStyle.Font.CompositeFont;
+
+		WidgetStyle.SetFont(Font);
 	}
 
 #if WITH_EDITOR 
diff --git a/Engine/Source/Runtime/UMG/Private/Components/MultiLineEditableText.cpp b/Engine/Source/Runtime/UMG/Private/Components/MultiLineEditableText.cpp
--- a/Engine/Source/Runtime/UMG/Private/Components/MultiLineEditableText.cpp
+++ b/Engine/Source/Runtime/UMG/Private/Components/MultiLineEditableText.cpp
@@ -65,8 +65,12 @@ UMultiLineEditableText::UMultiLineEditableText(const FObjectInitializer& ObjectI
 
 	if (!IsRunningDedicatedServer())
 	{
-		static ConstructorHelpers::FObjectFinder<UFont> RobotoFontObj(*UWidget::GetDefaultFontName());
-		WidgetStyle.SetFont(FSlateFontInfo(RobotoFontObj.Object, 12, FName("Bold")));
+		static ConstructorHelpers::FObjectFinder<UFont> DefaultFontObj(*UWidget::GetDefaultFontName());
+		FSlateFontInfo Font(DefaultFontObj.Object, 12, FName("Bold"));
+		//The FSlateFontInfo just created doesn't contain a composite font (while the default from the WidgetStyle does),
+		//so in the case the Font object is replaced by a null one, we have to keep the composite one as a fallback.
+		Font.CompositeFont = WidgetStyle.Font.CompositeFont;
+		WidgetStyle.SetFont(Font);
 	}
 }
 
diff --git a/Engine/Source/Runtime/UMG/Private/Components/MultiLineEditableTextBox.cpp b/Engine/Source/Runtime/UMG/Private/Components/MultiLineEditableTextBox.cpp
--- a/Engine/Source/Runtime/UMG/Private/Components/MultiLineEditableTextBox.cpp
+++ b/Engine/Source/Runtime/UMG/Private/Components/MultiLineEditableTextBox.cpp
@@ -41,8 +41,11 @@ UMultiLineEditableTextBox::UMultiLineEditableTextBox(const FObjectInitializer& O
 
 	if (!IsRunningDedicatedServer())
 	{
-		static ConstructorHelpers::FObjectFinder<UFont> defaultFontObj(*UWidget::GetDefaultFontName());
-		WidgetStyle.SetFont(FSlateFontInfo(defaultFontObj.Object, 24, FName("Regular")));
+		static ConstructorHelpers::FObjectFinder<UFont> DefaultFontObj(*UWidget::GetDefaultFontName());
+		FSlateFontInfo Font(DefaultFontObj.Object, 24, FName("Regular"));
+		//The FSlateFontInfo just created doesn't contain a composite font (while the default from the WidgetStyle does),
+		//so in the case the Font object is replaced by a null one, we have to keep the composite one as a fallback.
+		Font.CompositeFont = WidgetStyle.TextStyle.Font.CompositeFont;
 	}
 
 #if WITH_EDITOR 
