Add accessors for InputKeySelector and Invalidation box

diff --git a/Engine/Source/Runtime/UMG/Private/Components/InputKeySelector.cpp b/Engine/Source/Runtime/UMG/Private/Components/InputKeySelector.cpp
--- a/Engine/Source/Runtime/UMG/Private/Components/InputKeySelector.cpp
+++ b/Engine/Source/Runtime/UMG/Private/Components/InputKeySelector.cpp
@@ -39,10 +39,10 @@ UInputKeySelector::UInputKeySelector( const FObjectInitializer& ObjectInitialize
 		// Unlink UMG default colors.
 		DefaultInputKeySelectorTextStyle->UnlinkColors();
 	}
-
+PRAGMA_DISABLE_DEPRECATION_WARNINGS
 	WidgetStyle = *DefaultInputKeySelectorButtonStyle;
 	TextStyle = *DefaultInputKeySelectorTextStyle;
-
+PRAGMA_ENABLE_DEPRECATION_WARNINGS
 #if WITH_EDITOR 
 	if (EditorInputKeySelectorButtonStyle == nullptr)
 	{
@@ -62,26 +62,29 @@ UInputKeySelector::UInputKeySelector( const FObjectInitializer& ObjectInitialize
 
 	if (IsEditorWidget())
 	{
+PRAGMA_DISABLE_DEPRECATION_WARNINGS
 		WidgetStyle = *EditorInputKeySelectorButtonStyle;
 		TextStyle = *EditorInputKeySelectorTextStyle;
-
+PRAGMA_ENABLE_DEPRECATION_WARNINGS
 		// The CDO isn't an editor widget and thus won't use the editor style, call post edit change to mark difference from CDO
 		PostEditChange();
 	}
 #endif // WITH_EDITOR
-
+PRAGMA_DISABLE_DEPRECATION_WARNINGS
 	KeySelectionText = NSLOCTEXT("InputKeySelector", "DefaultKeySelectionText", "...");
 	NoKeySpecifiedText = NSLOCTEXT("InputKeySelector", "DefaultEmptyText", "Empty");
 	SelectedKey = FInputChord(EKeys::Invalid);
 	bAllowModifierKeys = true;
 	bAllowGamepadKeys = false;
-
+PRAGMA_ENABLE_DEPRECATION_WARNINGS
 	EscapeKeys.AddUnique(EKeys::Gamepad_Special_Right); // In most (if not all) cases this is going to be the menu button
 
 	if (!IsRunningDedicatedServer())
 	{
 		static ConstructorHelpers::FObjectFinder<UFont> RobotoFontObj(*UWidget::GetDefaultFontName());
+PRAGMA_DISABLE_DEPRECATION_WARNINGS
 		TextStyle.Font = FSlateFontInfo(RobotoFontObj.Object, 24, FName("Bold"));
+PRAGMA_ENABLE_DEPRECATION_WARNINGS
 	}
 }
 
@@ -92,13 +95,23 @@ void UInputKeySelector::Serialize(FArchive& Ar)
 	Ar.UsingCustomVersion(FFrameworkObjectVersion::GUID);
 }
 
+PRAGMA_DISABLE_DEPRECATION_WARNINGS
 void UInputKeySelector::SetSelectedKey( const FInputChord& InSelectedKey )
 {
-	if ( MyInputKeySelector.IsValid() )
+	if (SelectedKey != InSelectedKey)
 	{
-		MyInputKeySelector->SetSelectedKey( InSelectedKey );
+		BroadcastFieldValueChanged(FFieldNotificationClassDescriptor::SelectedKey);
+		if (MyInputKeySelector.IsValid())
+		{
+			MyInputKeySelector->SetSelectedKey(InSelectedKey);
+		}
+		SelectedKey = InSelectedKey;
 	}
-	SelectedKey = InSelectedKey;
+}
+
+FInputChord UInputKeySelector::GetSelectedKey() const
+{
+	return SelectedKey;
 }
 
 void UInputKeySelector::SetKeySelectionText( FText InKeySelectionText )
@@ -110,6 +123,16 @@ void UInputKeySelector::SetKeySelectionText( FText InKeySelectionText )
 	KeySelectionText = MoveTemp(InKeySelectionText);
 }
 
+const FText& UInputKeySelector::GetNoKeySpecifiedText() const
+{
+	return NoKeySpecifiedText;
+}
+
+const FText& UInputKeySelector::GetKeySelectionText() const
+{
+	return KeySelectionText;
+}
+
 void UInputKeySelector::SetNoKeySpecifiedText(FText InNoKeySpecifiedText)
 {
 	if (MyInputKeySelector.IsValid())
@@ -128,6 +151,11 @@ void UInputKeySelector::SetAllowModifierKeys( const bool bInAllowModifierKeys )
 	bAllowModifierKeys = bInAllowModifierKeys;
 }
 
+bool UInputKeySelector::AllowModifierKeys() const
+{
+	return bAllowModifierKeys;
+}
+
 void UInputKeySelector::SetAllowGamepadKeys(const bool bInAllowGamepadKeys)
 {
 	if (MyInputKeySelector.IsValid())
@@ -137,6 +165,12 @@ void UInputKeySelector::SetAllowGamepadKeys(const bool bInAllowGamepadKeys)
 	bAllowGamepadKeys = bInAllowGamepadKeys;
 }
 
+bool UInputKeySelector::AllowGamepadKeys() const
+{
+	return bAllowGamepadKeys;
+}
+PRAGMA_ENABLE_DEPRECATION_WARNINGS
+
 bool UInputKeySelector::GetIsSelectingKey() const
 {
 	return MyInputKeySelector.IsValid() ? MyInputKeySelector->GetIsSelectingKey() : false;
@@ -144,11 +178,36 @@ bool UInputKeySelector::GetIsSelectingKey() const
 
 void UInputKeySelector::SetButtonStyle( const FButtonStyle* InButtonStyle )
 {
-	if ( MyInputKeySelector.IsValid() )
+	SetButtonStyle(*InButtonStyle);
+}
+
+PRAGMA_DISABLE_DEPRECATION_WARNINGS
+void UInputKeySelector::SetButtonStyle(const FButtonStyle& InButtonStyle)
+{
+	if (MyInputKeySelector.IsValid())
+	{
+		MyInputKeySelector->SetButtonStyle(&InButtonStyle);
+	}
+	WidgetStyle = InButtonStyle;
+}
+
+const FButtonStyle& UInputKeySelector::GetButtonStyle() const
+{
+	return WidgetStyle;
+}
+
+void UInputKeySelector::SetTextStyle(const FTextBlockStyle& InTextStyle)
+{
+	if (MyInputKeySelector.IsValid())
 	{
-		MyInputKeySelector->SetButtonStyle(InButtonStyle);
+		MyInputKeySelector->SetTextStyle(&InTextStyle);
 	}
-	WidgetStyle = *InButtonStyle;
+	TextStyle = InTextStyle;
+}
+
+const FTextBlockStyle& UInputKeySelector::GetTextStyle() const
+{
+	return TextStyle;
 }
 
 void UInputKeySelector::SetEscapeKeys(const TArray<FKey>& InKeys)
@@ -159,6 +218,22 @@ void UInputKeySelector::SetEscapeKeys(const TArray<FKey>& InKeys)
 	}
 	EscapeKeys = InKeys;
 }
+
+void UInputKeySelector::SetMargin(const FMargin& InMargin)
+{
+	if (MyInputKeySelector.IsValid())
+	{
+		MyInputKeySelector->SetMargin(InMargin);
+	}
+	Margin = InMargin;
+}
+
+const FMargin& UInputKeySelector::GetMargin() const
+{
+	return Margin;
+}
+PRAGMA_ENABLE_DEPRECATION_WARNINGS
+
 #if WITH_EDITOR
 const FText UInputKeySelector::GetPaletteCategory()
 {
@@ -166,6 +241,7 @@ const FText UInputKeySelector::GetPaletteCategory()
 }
 #endif
 
+PRAGMA_DISABLE_DEPRECATION_WARNINGS
 void UInputKeySelector::SynchronizeProperties()
 {
 	Super::SynchronizeProperties();
@@ -208,8 +284,10 @@ TSharedRef<SWidget> UInputKeySelector::RebuildWidget()
 void UInputKeySelector::HandleKeySelected(const FInputChord& InSelectedKey)
 {
 	SelectedKey = InSelectedKey;
+	BroadcastFieldValueChanged(FFieldNotificationClassDescriptor::SelectedKey);
 	OnKeySelected.Broadcast(SelectedKey);
 }
+PRAGMA_ENABLE_DEPRECATION_WARNINGS
 
 void UInputKeySelector::HandleIsSelectingKeyChanged()
 {
diff --git a/Engine/Source/Runtime/UMG/Private/Components/InvalidationBox.cpp b/Engine/Source/Runtime/UMG/Private/Components/InvalidationBox.cpp
--- a/Engine/Source/Runtime/UMG/Private/Components/InvalidationBox.cpp
+++ b/Engine/Source/Runtime/UMG/Private/Components/InvalidationBox.cpp
@@ -15,7 +15,9 @@
 UInvalidationBox::UInvalidationBox(const FObjectInitializer& ObjectInitializer)
 	: Super(ObjectInitializer)
 {
+PRAGMA_DISABLE_DEPRECATION_WARNINGS
 	bCanCache = true;
+PRAGMA_ENABLE_DEPRECATION_WARNINGS
 	SetVisibilityInternal(ESlateVisibility::SelfHitTestInvisible);
 }
 
@@ -34,9 +36,9 @@ TSharedRef<SWidget> UInvalidationBox::RebuildWidget()
 		.DebugName(GetPathName())
 #endif
 		;
-
+PRAGMA_DISABLE_DEPRECATION_WARNINGS
 	MyInvalidationPanel->SetCanCache(IsDesignTime() ? false : bCanCache);
-
+PRAGMA_ENABLE_DEPRECATION_WARNINGS
 	if ( GetChildrenCount() > 0 )
 	{
 		MyInvalidationPanel->SetContent(GetContentSlot()->Content ? GetContentSlot()->Content->TakeWidget() : SNullWidget::NullWidget);
@@ -67,6 +69,7 @@ void UInvalidationBox::InvalidateCache()
 {
 }
 
+PRAGMA_DISABLE_DEPRECATION_WARNINGS
 bool UInvalidationBox::GetCanCache() const
 {
 	if ( MyInvalidationPanel.IsValid() )
@@ -85,6 +88,7 @@ void UInvalidationBox::SetCanCache(bool CanCache)
 		return MyInvalidationPanel->SetCanCache(bCanCache);
 	}
 }
+PRAGMA_ENABLE_DEPRECATION_WARNINGS
 
 #if WITH_EDITOR
 
diff --git a/Engine/Source/Runtime/UMG/Public/Components/InputKeySelector.h b/Engine/Source/Runtime/UMG/Public/Components/InputKeySelector.h
--- a/Engine/Source/Runtime/UMG/Public/Components/InputKeySelector.h
+++ b/Engine/Source/Runtime/UMG/Public/Components/InputKeySelector.h
@@ -28,36 +28,44 @@ public:
 
 public:
 	/** The button style used at runtime */
-	UPROPERTY(EditAnywhere, BlueprintReadWrite, Category = Appearance, meta = (DisplayName = "Style"))
+	UE_DEPRECATED(5.2, "Direct access to WidgetStyle is now deprecated. Use the setter or getter.")
+	UPROPERTY(EditAnywhere, BlueprintReadWrite, Getter = "GetButtonStyle", Setter = "SetButtonStyle", Category = Appearance, meta = (DisplayName = "Style"))
 	FButtonStyle WidgetStyle;
 
 	/** The button style used at runtime */
-	UPROPERTY(EditAnywhere, BlueprintReadWrite, Category = Appearance, meta = (DisplayName = "Text Style"))
+	UE_DEPRECATED(5.2, "Direct access to TextStyle is now deprecated. Use the setter or getter.")
+	UPROPERTY(EditAnywhere, BlueprintReadWrite, Getter, Setter, Category = Appearance, meta = (DisplayName = "Text Style"))
 	FTextBlockStyle TextStyle;
 
 	/** The currently selected key chord. */
-	UPROPERTY( BlueprintReadOnly, Category = "Key Selection" )
+	UE_DEPRECATED(5.2, "Direct access to SelectedKey is now deprecated. Use the setter or getter.")
+	UPROPERTY( BlueprintReadWrite, Getter, Setter, BlueprintSetter = "SetSelectedKey", FieldNotify, Category = "Key Selection")
 	FInputChord SelectedKey;
 
+	UE_DEPRECATED(5.2, "Direct access to Margin is now deprecated. Use the setter or getter.")
 	/** The amount of blank space around the text used to display the currently selected key. */
-	UPROPERTY( EditAnywhere, BlueprintReadOnly, Category = Appearance )
+	UPROPERTY( EditAnywhere, BlueprintReadWrite, Getter, Setter, Category = Appearance )
 	FMargin Margin;
 
 	/** Sets the text which is displayed while selecting keys. */
-	UPROPERTY( EditAnywhere, BlueprintReadOnly, Category = Appearance )
+	UE_DEPRECATED(5.2, "Direct access to KeySelectionText is now deprecated. Use the setter or getter.")
+	UPROPERTY( EditAnywhere, BlueprintReadWrite, Getter, Setter, BlueprintSetter = "SetKeySelectionText", Category = Appearance )
 	FText KeySelectionText;
 
 	/** Sets the text to display when no key text is available or not selecting a key. */
-	UPROPERTY(EditAnywhere, BlueprintReadOnly, Category = Appearance)
+	UE_DEPRECATED(5.2, "Direct access to NoKeySpecifiedText is now deprecated. Use the setter or getter.")
+	UPROPERTY( EditAnywhere, BlueprintReadWrite, Getter, Setter, BlueprintSetter = "SetNoKeySpecifiedText", Category = Appearance )
 	FText NoKeySpecifiedText;
 
 	/** When true modifier keys such as control and alt are allowed in the */
 	/** input chord representing the selected key, if false modifier keys are ignored. */
-	UPROPERTY( EditAnywhere, BlueprintReadOnly, Category = "Key Selection" )
+	UE_DEPRECATED(5.2, "Direct access to bAllowModifierKeys is now deprecated. Use the setter or getter.")
+	UPROPERTY( EditAnywhere, BlueprintReadWrite, Getter = "AllowModifierKeys", Setter = "SetAllowModifierKeys", BlueprintSetter = "SetAllowModifierKeys", Category = "Key Selection")
 	bool bAllowModifierKeys;
 
 	/** When true gamepad keys are allowed in the input chord representing the selected key, otherwise they are ignored. */
-	UPROPERTY(EditAnywhere, BlueprintReadOnly, Category = "Key Selection")
+	UE_DEPRECATED(5.2, "Direct access to bAllowGamepadKeys is now deprecated. Use the setter or getter.")
+	UPROPERTY( EditAnywhere, BlueprintReadWrite, Getter = "AllowGamepadKeys", Setter = "SetAllowGamepadKeys", BlueprintSetter = "SetAllowGamepadKeys", Category = "Key Selection")
 	bool bAllowGamepadKeys;
 
 	/** When true gamepad keys are allowed in the input chord representing the selected key, otherwise they are ignored. */
@@ -77,22 +85,37 @@ public:
 	UFUNCTION( BlueprintCallable, Category = "Widget" )
 	void SetSelectedKey(const FInputChord& InSelectedKey);
 
+	/** Returns the currently selected key. */
+	FInputChord GetSelectedKey() const;
+
 	/** Sets the text which is displayed while selecting keys. */
 	UFUNCTION( BlueprintCallable, Category = "Widget" )
 	void SetKeySelectionText( FText InKeySelectionText );
 
+	/** Returns the text which is displayed while selecting keys. */
+	const FText& GetKeySelectionText() const;
+
 	/** Sets the text to display when no key text is available or not selecting a key. */
 	UFUNCTION(BlueprintCallable, Category = "Widget")
 	void SetNoKeySpecifiedText(FText InNoKeySpecifiedText);
 
+	/** Returns the text to display when no key text is available or not selecting a key. */
+	const FText& GetNoKeySpecifiedText() const;
+
 	/** Sets whether or not modifier keys are allowed in the selected key. */
 	UFUNCTION(BlueprintCallable, Category = "Widget" )
 	void SetAllowModifierKeys(bool bInAllowModifierKeys );
 
+	/** Returns true if modifier keys keys are allowed, otherwise returns false. */
+	bool AllowModifierKeys() const;
+
 	/** Sets whether or not gamepad keys are allowed in the selected key. */
 	UFUNCTION(BlueprintCallable, Category = "Widget")
 	void SetAllowGamepadKeys(bool bInAllowGamepadKeys);
 
+	/** Returns true if gamepad keys are allowed, otherwise returns false. */
+	bool AllowGamepadKeys() const;
+
 	/** Returns true if the widget is currently selecting a key, otherwise returns false. */
 	UFUNCTION( BlueprintCallable, Category = "Widget" )
 	bool GetIsSelectingKey() const;
@@ -102,8 +125,27 @@ public:
 	void SetTextBlockVisibility(const ESlateVisibility InVisibility);
 
 	/** Sets the style of the button used to start key selection mode. */
+	UE_DEPRECATED(5.2, "SetButtonStyle with pointer type parameter is deprecated. Please pass a reference instead.")
 	void SetButtonStyle(const FButtonStyle* ButtonStyle);
 
+	/** Sets the style of the button used to start key selection mode. */
+	void SetButtonStyle(const FButtonStyle& ButtonStyle);
+
+	/** Returns the style of the button used to start key selection mode. */
+	const FButtonStyle& GetButtonStyle() const;
+
+	/** Sets the style of the text used inside the button. */
+	void SetTextStyle(const FTextBlockStyle& InTextStyle);
+
+	/** Returns the style of the text used inside the button. */
+	const FTextBlockStyle& GetTextStyle() const;
+
+	/** Sets the amount of blank space around the text used to display the currently selected key. */
+	void SetMargin(const FMargin& InMargin);
+
+	/** Returns the amount of blank space around the text used to display the currently selected key. */
+	const FMargin& GetMargin() const;
+
 	/** Sets escape keys. */
 	UFUNCTION(BlueprintCallable, Category = "Widget")
 	void SetEscapeKeys(const TArray<FKey>& InKeys);
diff --git a/Engine/Source/Runtime/UMG/Public/Components/InvalidationBox.h b/Engine/Source/Runtime/UMG/Public/Components/InvalidationBox.h
--- a/Engine/Source/Runtime/UMG/Public/Components/InvalidationBox.h
+++ b/Engine/Source/Runtime/UMG/Public/Components/InvalidationBox.h
@@ -65,17 +65,10 @@ protected:
 	 * Should the invalidation panel cache the widgets?  Making this false makes it so the invalidation
 	 * panel stops acting like an invalidation panel, just becomes a simple container widget.
 	 */
-	UPROPERTY(EditAnywhere, Category="Caching")
+	UE_DEPRECATED(5.2, "Direct access to bCanCache is deprecated. Please use the getter or setter.")
+	UPROPERTY(EditAnywhere, BlueprintReadWrite, Getter = "GetCanCache", Setter = "SetCanCache", BlueprintGetter = "GetCanCache", BlueprintSetter = "SetCanCache", Category = "Caching")
 	bool bCanCache;
 
-	/**
-	 * Caches the locations for child draw elements relative to the invalidation box,
-	 * this adds extra overhead to drawing them every frame.  However, in cases where
-	 * the position of the invalidation boxes changes every frame this can be a big savings.
-	 */
-	UPROPERTY()
-	bool CacheRelativeTransforms_DEPRECATED;
-
 protected:
 	TSharedPtr<class SInvalidationPanel> MyInvalidationPanel;
 };
