Add accessor for SafeZone, SafeZoneSlot and WindowTitleBarArea

diff --git a/Engine/Source/Runtime/UMG/Private/Components/SafeZone.cpp b/Engine/Source/Runtime/UMG/Private/Components/SafeZone.cpp
--- a/Engine/Source/Runtime/UMG/Private/Components/SafeZone.cpp
+++ b/Engine/Source/Runtime/UMG/Private/Components/SafeZone.cpp
@@ -51,7 +51,10 @@ void USafeZone::OnSlotAdded( UPanelSlot* InSlot )
 {
 	Super::OnSlotAdded( InSlot );
 
-	UpdateWidgetProperties();
+	if (MySafeZone.IsValid())
+	{
+		CastChecked<USafeZoneSlot>(InSlot)->BuildSlot(MySafeZone.ToSharedRef());
+	}
 }
 
 void USafeZone::OnSlotRemoved( UPanelSlot* InSlot )
@@ -69,6 +72,7 @@ UClass* USafeZone::GetSlotClass() const
 	return USafeZoneSlot::StaticClass();
 }
 
+PRAGMA_DISABLE_DEPRECATION_WARNINGS
 void USafeZone::UpdateWidgetProperties()
 {
 	if ( MySafeZone.IsValid() && GetChildrenCount() > 0 )
@@ -98,15 +102,8 @@ void USafeZone::SetSidesToPad(bool InPadLeft, bool InPadRight, bool InPadTop, bo
 }
 
 TSharedRef<SWidget> USafeZone::RebuildWidget()
-{
-	USafeZoneSlot* SafeSlot = Slots.Num() > 0 ? Cast< USafeZoneSlot >( Slots[ 0 ] ) : nullptr;
-	
+{	
 	MySafeZone = SNew( SSafeZone )
-		.IsTitleSafe( SafeSlot ? SafeSlot->bIsTitleSafe : false )
-		.SafeAreaScale( SafeSlot ? SafeSlot->SafeAreaScale : FMargin(1,1,1,1))
-		.HAlign( SafeSlot ? SafeSlot->HAlign.GetValue() : HAlign_Fill )
-		.VAlign( SafeSlot ? SafeSlot->VAlign.GetValue() : VAlign_Fill )
-		.Padding( SafeSlot ? SafeSlot->Padding : FMargin() )
 		.PadLeft( PadLeft )
 		.PadRight( PadRight )
 		.PadTop( PadTop )
@@ -119,8 +116,17 @@ TSharedRef<SWidget> USafeZone::RebuildWidget()
 			GetChildAt( 0 ) ? GetChildAt( 0 )->TakeWidget() : SNullWidget::NullWidget
 		];
 
+	if (GetChildrenCount() > 0)
+	{
+		if (USafeZoneSlot* SafeSlot = Cast<USafeZoneSlot>(GetContentSlot()))
+		{
+			SafeSlot->BuildSlot(MySafeZone.ToSharedRef());
+		}
+	}
+
 	return MySafeZone.ToSharedRef();
 }
+PRAGMA_ENABLE_DEPRECATION_WARNINGS
 
 void USafeZone::ReleaseSlateResources(bool bReleaseChildren)
 {
@@ -129,5 +135,59 @@ void USafeZone::ReleaseSlateResources(bool bReleaseChildren)
 	MySafeZone.Reset();
 }
 
+PRAGMA_DISABLE_DEPRECATION_WARNINGS
+void USafeZone::SetPadLeft(bool InPadLeft)
+{
+	if (PadLeft != InPadLeft)
+	{
+		SetSidesToPad(InPadLeft, PadRight, PadTop, PadBottom);
+	}
+}
+
+bool USafeZone::GetPadLeft() const
+{
+	return PadLeft;
+}
+
+void USafeZone::SetPadRight(bool InPadRight)
+{
+	if (PadRight != InPadRight)
+	{
+		SetSidesToPad(PadLeft, InPadRight, PadTop, PadBottom);
+	}
+}
+
+bool USafeZone::GetPadRight() const
+{
+	return PadRight;
+}
+
+void USafeZone::SetPadTop(bool InPadTop)
+{
+	if (PadTop != InPadTop)
+	{
+		SetSidesToPad(PadLeft, PadRight, InPadTop, PadBottom);
+	}
+}
+
+bool USafeZone::GetPadTop() const
+{
+	return PadTop;
+}
+
+void USafeZone::SetPadBottom(bool InPadBottom)
+{
+	if (PadBottom != InPadBottom)
+	{
+		SetSidesToPad(PadLeft, PadRight, PadTop, InPadBottom);
+	}
+}
+
+bool USafeZone::GetPadBottom() const
+{
+	return PadBottom;
+}
+PRAGMA_ENABLE_DEPRECATION_WARNINGS
+
 #undef LOCTEXT_NAMESPACE
 
diff --git a/Engine/Source/Runtime/UMG/Private/Components/SafeZoneSlot.cpp b/Engine/Source/Runtime/UMG/Private/Components/SafeZoneSlot.cpp
--- a/Engine/Source/Runtime/UMG/Private/Components/SafeZoneSlot.cpp
+++ b/Engine/Source/Runtime/UMG/Private/Components/SafeZoneSlot.cpp
@@ -1,16 +1,41 @@
 // Copyright Epic Games, Inc. All Rights Reserved.
 
 #include "Components/SafeZoneSlot.h"
+#include "Widgets/Layout/SSafeZone.h"
 
 #include "Components/SafeZone.h"
 
 #include UE_INLINE_GENERATED_CPP_BY_NAME(SafeZoneSlot)
 
+PRAGMA_DISABLE_DEPRECATION_WARNINGS
 USafeZoneSlot::USafeZoneSlot()
 {
 	bIsTitleSafe = true;
 	SafeAreaScale = FMargin(1, 1, 1, 1);
 }
+PRAGMA_ENABLE_DEPRECATION_WARNINGS
+
+void USafeZoneSlot::ReleaseSlateResources(bool bReleaseChildren)
+{
+	Super::ReleaseSlateResources(bReleaseChildren);
+
+	SafeZone.Reset();
+}
+
+void USafeZoneSlot::BuildSlot(TSharedRef<SSafeZone> InSafeZone)
+{
+	SafeZone = InSafeZone;
+
+	PRAGMA_DISABLE_DEPRECATION_WARNINGS
+	InSafeZone->SetTitleSafe(bIsTitleSafe);
+	InSafeZone->SetSafeAreaScale(SafeAreaScale);
+	InSafeZone->SetHAlign(HAlign.GetValue());
+	InSafeZone->SetVAlign(VAlign.GetValue());
+	InSafeZone->SetPadding(Padding);
+	PRAGMA_ENABLE_DEPRECATION_WARNINGS
+
+	InSafeZone->SetContent(Content ? Content->TakeWidget() : SNullWidget::NullWidget);
+}
 
 void USafeZoneSlot::SynchronizeProperties()
 {
@@ -22,3 +47,74 @@ void USafeZoneSlot::SynchronizeProperties()
 	}
 }
 
+PRAGMA_DISABLE_DEPRECATION_WARNINGS
+void USafeZoneSlot::SetIsTitleSafe(bool InIsTitleSafe)
+{
+	bIsTitleSafe = InIsTitleSafe;
+	if (SafeZone.IsValid())
+	{
+		SafeZone.Pin()->SetTitleSafe(InIsTitleSafe);
+	}
+}
+
+bool USafeZoneSlot::IsTitleSafe() const
+{
+	return bIsTitleSafe;
+}
+
+void USafeZoneSlot::SetSafeAreaScale(const FMargin& InSafeAreaScale)
+{
+	SafeAreaScale = InSafeAreaScale;
+	if (SafeZone.IsValid())
+	{
+		SafeZone.Pin()->SetSafeAreaScale(SafeAreaScale);
+	}
+}
+
+const FMargin& USafeZoneSlot::GetSafeAreaScale() const
+{
+	return SafeAreaScale;
+}
+
+void USafeZoneSlot::SetHorizontalAlignment(EHorizontalAlignment InHAlign)
+{
+	HAlign = InHAlign;
+	if (SafeZone.IsValid())
+	{
+		SafeZone.Pin()->SetHAlign(HAlign);
+	}
+}
+
+const EHorizontalAlignment USafeZoneSlot::GetHorizontalAlignment() const
+{
+	return HAlign.GetValue();
+}
+
+void USafeZoneSlot::SetVerticalAlignment(EVerticalAlignment InVAlign)
+{
+	VAlign = InVAlign;
+	if (SafeZone.IsValid())
+	{
+		SafeZone.Pin()->SetVAlign(VAlign);
+	}
+}
+
+const EVerticalAlignment USafeZoneSlot::GetVerticalAlignment() const
+{
+	return VAlign.GetValue();
+}
+
+void USafeZoneSlot::SetPadding(const FMargin& InPadding)
+{
+	Padding = InPadding;
+	if (SafeZone.IsValid())
+	{
+		SafeZone.Pin()->SetPadding(InPadding);
+	}
+}
+
+const FMargin& USafeZoneSlot::GetPadding() const
+{
+	return Padding;
+}
+PRAGMA_ENABLE_DEPRECATION_WARNINGS
diff --git a/Engine/Source/Runtime/UMG/Private/Components/WindowTitleBarArea.cpp b/Engine/Source/Runtime/UMG/Private/Components/WindowTitleBarArea.cpp
--- a/Engine/Source/Runtime/UMG/Private/Components/WindowTitleBarArea.cpp
+++ b/Engine/Source/Runtime/UMG/Private/Components/WindowTitleBarArea.cpp
@@ -22,8 +22,9 @@ UWindowTitleBarArea::UWindowTitleBarArea(const FObjectInitializer& ObjectInitial
 {
 	bIsVariable = false;
 	SetVisibilityInternal(ESlateVisibility::Visible);
-
+PRAGMA_DISABLE_DEPRECATION_WARNINGS
 	bDoubleClickTogglesFullscreen = false;
+PRAGMA_ENABLE_DEPRECATION_WARNINGS
 }
 
 void UWindowTitleBarArea::ReleaseSlateResources(bool bReleaseChildren)
@@ -42,6 +43,7 @@ void UWindowTitleBarArea::ReleaseSlateResources(bool bReleaseChildren)
 TSharedRef<SWidget> UWindowTitleBarArea::RebuildWidget()
 {
 	MyWindowTitleBarArea = SNew(SWindowTitleBarArea);
+PRAGMA_DISABLE_DEPRECATION_WARNINGS
 	if (bDoubleClickTogglesFullscreen)
 	{
 		WindowActionNotificationHandle = FSlateApplication::Get().RegisterOnWindowActionNotification(BIND_UOBJECT_DELEGATE(FOnWindowAction, HandleWindowAction));
@@ -51,7 +53,7 @@ TSharedRef<SWidget> UWindowTitleBarArea::RebuildWidget()
 		FSlateApplication::Get().UnregisterOnWindowActionNotification(WindowActionNotificationHandle);
 		WindowActionNotificationHandle.Reset();
 	}
-
+PRAGMA_ENABLE_DEPRECATION_WARNINGS
 	MyWindowTitleBarArea->SetRequestToggleFullscreenCallback(BIND_UOBJECT_DELEGATE(FSimpleDelegate, RequestToggleFullscreen));
 
 	if (GetChildrenCount() > 0)
@@ -64,7 +66,9 @@ TSharedRef<SWidget> UWindowTitleBarArea::RebuildWidget()
 		MyWindowTitleBarArea->SetGameWindow(GEngine->GameViewport->GetWindow());
 	}
 
+PRAGMA_DISABLE_DEPRECATION_WARNINGS
 	MyWindowTitleBarArea->SetWindowButtonsVisibility(bWindowButtonsEnabled);
+PRAGMA_ENABLE_DEPRECATION_WARNINGS
 
 	return MyWindowTitleBarArea.ToSharedRef();
 }
@@ -117,6 +121,42 @@ void UWindowTitleBarArea::SetVerticalAlignment(EVerticalAlignment InVerticalAlig
 		MyWindowTitleBarArea->SetVAlign(InVerticalAlignment);
 	}
 }
+PRAGMA_DISABLE_DEPRECATION_WARNINGS
+void UWindowTitleBarArea::SetWindowButtonsEnabled(bool InWindowButtonsEnabled)
+{
+	bWindowButtonsEnabled = InWindowButtonsEnabled;
+	
+	if (MyWindowTitleBarArea.IsValid())
+	{
+		MyWindowTitleBarArea->SetWindowButtonsVisibility(bWindowButtonsEnabled);
+	}
+}
+
+bool UWindowTitleBarArea::IsWindowButtonsEnabled() const
+{
+	return bWindowButtonsEnabled;
+}
+
+void UWindowTitleBarArea::SetDoubleClickTogglesFullscreen(bool InDoubleClickTogglesFullscreen)
+{
+	bDoubleClickTogglesFullscreen = InDoubleClickTogglesFullscreen;
+
+	if (bDoubleClickTogglesFullscreen)
+	{
+		WindowActionNotificationHandle = FSlateApplication::Get().RegisterOnWindowActionNotification(BIND_UOBJECT_DELEGATE(FOnWindowAction, HandleWindowAction));
+	}
+	else if (WindowActionNotificationHandle.IsValid())
+	{
+		FSlateApplication::Get().UnregisterOnWindowActionNotification(WindowActionNotificationHandle);
+		WindowActionNotificationHandle.Reset();
+	}
+}
+
+bool UWindowTitleBarArea::IsDoubleClickTogglesFullscreen() const
+{
+	return bDoubleClickTogglesFullscreen;
+}
+PRAGMA_ENABLE_DEPRECATION_WARNINGS
 
 void UWindowTitleBarArea::PostLoad()
 {
diff --git a/Engine/Source/Runtime/UMG/Public/Components/SafeZone.h b/Engine/Source/Runtime/UMG/Public/Components/SafeZone.h
--- a/Engine/Source/Runtime/UMG/Public/Components/SafeZone.h
+++ b/Engine/Source/Runtime/UMG/Public/Components/SafeZone.h
@@ -50,19 +50,23 @@ public:
 public:
 
 	/** If this safe zone should pad for the left side of the screen's safe zone */
-	UPROPERTY(EditAnywhere, BlueprintReadOnly, Category = "SafeZone")
+	UE_DEPRECATED(5.2, "Direct access to PadLeft is deprecated. Please use the getter or setter.")
+	UPROPERTY(EditAnywhere, BlueprintReadWrite, Getter, Setter, Category = "SafeZone")
 	bool PadLeft;
 
+	UE_DEPRECATED(5.2, "Direct access to PadRight is deprecated. Please use the getter or setter.")
 	/** If this safe zone should pad for the right side of the screen's safe zone */
-	UPROPERTY(EditAnywhere, BlueprintReadOnly, Category = "SafeZone")
+	UPROPERTY(EditAnywhere, BlueprintReadWrite, Getter, Setter, Category = "SafeZone")
 	bool PadRight;
 
+	UE_DEPRECATED(5.2, "Direct access to PadTop is deprecated. Please use the getter or setter.")
 	/** If this safe zone should pad for the top side of the screen's safe zone */
-	UPROPERTY(EditAnywhere, BlueprintReadOnly, Category = "SafeZone")
+	UPROPERTY(EditAnywhere, BlueprintReadWrite, Getter, Setter, Category = "SafeZone")
 	bool PadTop;
 
+	UE_DEPRECATED(5.2, "Direct access to PadBottom is deprecated. Please use the getter or setter.")
 	/** If this safe zone should pad for the bottom side of the screen's safe zone */
-	UPROPERTY(EditAnywhere, BlueprintReadOnly, Category = "SafeZone")
+	UPROPERTY(EditAnywhere, BlueprintReadWrite, Getter, Setter, Category = "SafeZone")
 	bool PadBottom;
 
 protected:
@@ -72,6 +76,18 @@ protected:
 
 	virtual void ReleaseSlateResources(bool bReleaseChildren) override;
 
+	void SetPadLeft(bool InPadLeft);
+	bool GetPadLeft() const;
+
+	void SetPadRight(bool InPadRight);
+	bool GetPadRight() const;
+
+	void SetPadTop(bool InPadTop);
+	bool GetPadTop() const;
+
+	void SetPadBottom(bool InPadBottom);
+	bool GetPadBottom() const;
+
 	TSharedPtr< class SSafeZone > MySafeZone;
 
 #if WITH_EDITOR
diff --git a/Engine/Source/Runtime/UMG/Public/Components/SafeZoneSlot.h b/Engine/Source/Runtime/UMG/Public/Components/SafeZoneSlot.h
--- a/Engine/Source/Runtime/UMG/Public/Components/SafeZoneSlot.h
+++ b/Engine/Source/Runtime/UMG/Public/Components/SafeZoneSlot.h
@@ -9,28 +9,62 @@
 
 #include "SafeZoneSlot.generated.h"
 
+class SSafeZone;
+
 UCLASS()
 class UMG_API USafeZoneSlot : public UPanelSlot
 {
 	GENERATED_BODY()
 public:
 
-	UPROPERTY( EditAnywhere, BlueprintReadOnly, Category = SafeZone )
+	UE_DEPRECATED(5.2, "Direct access to bIsTitleSafe is deprecated. Please use the getter or setter.")
+	UPROPERTY( EditAnywhere, BlueprintReadWrite, Getter = "IsTitleSafe", Setter = "SetIsTitleSafe", Category = SafeZone)
 	bool bIsTitleSafe;
 
-	UPROPERTY( EditAnywhere, BlueprintReadOnly, Category = SafeZone )
+	UE_DEPRECATED(5.2, "Direct access to SafeAreaScale is deprecated. Please use the getter or setter.")
+	UPROPERTY( EditAnywhere, BlueprintReadWrite, Getter, Setter, Category = SafeZone )
 	FMargin SafeAreaScale;
 
-	UPROPERTY( EditAnywhere, BlueprintReadOnly, Category = SafeZone )
+	UE_DEPRECATED(5.2, "Direct access to HAlign is deprecated. Please use the getter or setter.")
+	UPROPERTY( EditAnywhere, BlueprintReadWrite, Getter = "GetHorizontalAlignment", Setter = "SetHorizontalAlignment", Category = SafeZone)
 	TEnumAsByte< EHorizontalAlignment > HAlign;
 
-	UPROPERTY( EditAnywhere, BlueprintReadOnly, Category = SafeZone )
+	UE_DEPRECATED(5.2, "Direct access to VAlign is deprecated. Please use the getter or setter.")
+	UPROPERTY( EditAnywhere, BlueprintReadWrite, Getter = "GetVerticalAlignment", Setter = "SetVerticalAlignment", Category = SafeZone)
 	TEnumAsByte< EVerticalAlignment > VAlign;
 
-	UPROPERTY( EditAnywhere, BlueprintReadOnly, Category = SafeZone )
+	UE_DEPRECATED(5.2, "Direct access to Padding is deprecated. Please use the getter or setter.")
+	UPROPERTY( EditAnywhere, BlueprintReadWrite, Getter, Setter, Category = SafeZone )
 	FMargin Padding;
 
 	USafeZoneSlot();
 
+	void SetIsTitleSafe(bool InIsTitleSafe);
+	bool IsTitleSafe() const;
+
+	void SetSafeAreaScale(const FMargin& InSafeAreaScale);
+	const FMargin& GetSafeAreaScale() const;
+
+	void SetHorizontalAlignment(EHorizontalAlignment InHAlign);
+	const EHorizontalAlignment GetHorizontalAlignment() const;
+
+	void SetVerticalAlignment(EVerticalAlignment InVAlign);
+	const EVerticalAlignment GetVerticalAlignment() const;
+
+	void SetPadding(const FMargin& InPadding);
+	const FMargin& GetPadding() const;
+
+	//~ UPanelSlot interface
 	virtual void SynchronizeProperties() override;
+	//~ End of UPanelSlot interface
+
+	/** Builds the underlying slot for the slate button. */
+	void BuildSlot(TSharedRef<SSafeZone> InSafeZone);
+
+	virtual void ReleaseSlateResources(bool bReleaseChildren) override;
+
+private:
+
+	/** A pointer to the button to allow us to adjust the alignment, padding...etc at runtime. */
+	TWeakPtr<SSafeZone> SafeZone;
 };
diff --git a/Engine/Source/Runtime/UMG/Public/Components/WindowTitleBarArea.h b/Engine/Source/Runtime/UMG/Public/Components/WindowTitleBarArea.h
--- a/Engine/Source/Runtime/UMG/Public/Components/WindowTitleBarArea.h
+++ b/Engine/Source/Runtime/UMG/Public/Components/WindowTitleBarArea.h
@@ -25,11 +25,13 @@ class UMG_API UWindowTitleBarArea : public UContentWidget
 public:
 
 	/** Should the title bar area diaplay window minimize/maximize/close buttons. */
-	UPROPERTY(EditAnywhere, BlueprintReadOnly, Category = "Interaction", meta = (DisplayName = "Window Buttons Enabled"))
+	UE_DEPRECATED(5.2, "Direct access to bWindowButtonsEnabled is deprecated. Please use the getter or setter.")
+	UPROPERTY(EditAnywhere, BlueprintReadWrite, Getter = "IsWindowButtonsEnabled", Setter = "SetWindowButtonsEnabled", Category = "Interaction", meta = (DisplayName = "Window Buttons Enabled"))
 	bool bWindowButtonsEnabled;
 
 	/** Should double clicking the title bar area toggle fullscreen instead of maximizing the window. */
-	UPROPERTY(EditAnywhere, BlueprintReadOnly, Category = "Interaction", meta=(DisplayName="Toggle Fullscreen"))
+	UE_DEPRECATED(5.2, "Direct access to bDoubleClickTogglesFullscreen is deprecated. Please use the getter or setter.")
+	UPROPERTY(EditAnywhere, BlueprintReadWrite, Getter = "IsDoubleClickTogglesFullscreen", Setter = "SetDoubleClickTogglesFullscreen", Category = "Interaction", meta = (DisplayName = "Toggle Fullscreen"))
 	bool bDoubleClickTogglesFullscreen;
 
 public:
@@ -43,6 +45,14 @@ public:
 	UFUNCTION(BlueprintCallable, Category = "Appearance")
 	void SetVerticalAlignment(EVerticalAlignment InVerticalAlignment);
 
+	void SetWindowButtonsEnabled(bool InWindowButtonsEnabled);
+
+	bool IsWindowButtonsEnabled() const;
+
+	void SetDoubleClickTogglesFullscreen(bool InDoubleClickTogglesFullscreen);
+
+	bool IsDoubleClickTogglesFullscreen() const;
+
 public:
 
 	//~ Begin UVisual Interface
