Make "SetFlowDirectionPreference" BPCallable on UWidget

diff --git a/Engine/Source/Runtime/UMG/Public/Components/Widget.h b/Engine/Source/Runtime/UMG/Public/Components/Widget.h
--- a/Engine/Source/Runtime/UMG/Public/Components/Widget.h
+++ b/Engine/Source/Runtime/UMG/Public/Components/Widget.h
@@ -301,7 +301,7 @@ public:
 
 	/** Allows you to set a new flow direction */
 	UE_DEPRECATED(5.1, "Direct access to FlowDirectionPreference is deprecated. Please use the getter or setter.")
-	UPROPERTY(EditAnywhere, Getter, Setter, Category="Localization")
+	UPROPERTY(EditAnywhere, Getter, Setter, BlueprintGetter="GetFlowDirectionPreference", BlueprintSetter="SetFlowDirectionPreference", Category="Localization")
 	EFlowDirectionPreference FlowDirectionPreference;
 
 	/**
@@ -506,10 +506,12 @@ public:
 	UFUNCTION(BlueprintCallable, Category="Widget|Transform")
 	void SetRenderTransformPivot(FVector2D Pivot);
 
-	/** */
+	/** Gets the flow direction preference of the widget */
+	UFUNCTION(BlueprintCallable, Category="Localization")
 	EFlowDirectionPreference GetFlowDirectionPreference() const;
 
-	/** */
+	/** Sets the flow direction preference of the widget */
+	UFUNCTION(BlueprintCallable, Category="Localization")
 	void SetFlowDirectionPreference(EFlowDirectionPreference FlowDirection);
 
 	/** Gets the current enabled status of the widget */
