Fix for SetPositionInViewport not properly setting the position

diff --git a/Engine/Source/Runtime/UMG/Private/Blueprint/GameViewportSubsystem.cpp b/Engine/Source/Runtime/UMG/Private/Blueprint/GameViewportSubsystem.cpp
--- a/Engine/Source/Runtime/UMG/Private/Blueprint/GameViewportSubsystem.cpp
+++ b/Engine/Source/Runtime/UMG/Private/Blueprint/GameViewportSubsystem.cpp
@@ -22,14 +22,7 @@ namespace UE::UMG::Private
 		// If the size is zero, and we're not stretched, then use the desired size.
 		FVector2D FinalSize = FVector2D(Slot.Offsets.Right, Slot.Offsets.Bottom);
 		bool bUseAutoSize = FinalSize.IsZero() && !Slot.Anchors.IsStretchedVertical() && !Slot.Anchors.IsStretchedHorizontal();
-		if (bUseAutoSize)
-		{
-			return TPair<FMargin, bool>(FinalSize, bUseAutoSize);
-		}
-		else
-		{
-			return TPair<FMargin, bool>(FMargin(Slot.Offsets.Left, Slot.Offsets.Top, FinalSize.X, FinalSize.Y), bUseAutoSize);
-		}
+		return TPair<FMargin, bool>(Slot.Offsets, bUseAutoSize);
 	}
 }
  
