Fix editable text not reflecting a change of style at runtime because no setter was called to apply the style to the internal struct.

test: created a user widget with all types of editable texts, then set the style of each one in via a blueprint graph. It wasn't working for the EditableText before the fix, and it now works for all the types.

diff --git a/Engine/Source/Runtime/UMG/Private/Components/EditableText.cpp b/Engine/Source/Runtime/UMG/Private/Components/EditableText.cpp
--- a/Engine/Source/Runtime/UMG/Private/Components/EditableText.cpp
+++ b/Engine/Source/Runtime/UMG/Private/Components/EditableText.cpp
@@ -277,6 +277,15 @@ bool UEditableText::GetSelectAllTextOnCommit() const
 	return SelectAllTextOnCommit;
 }
 
+void UEditableText::SetWidgetStyle(const FEditableTextStyle& InEditableTextStyle)
+{
+	WidgetStyle = InEditableTextStyle;
+	if (MyEditableText.IsValid())
+	{
+		MyEditableText->SetTextStyle(WidgetStyle);
+	}
+}
+
 bool UEditableText::GetIsReadOnly() const
 {
 	return IsReadOnly;
diff --git a/Engine/Source/Runtime/UMG/Public/Components/EditableText.h b/Engine/Source/Runtime/UMG/Public/Components/EditableText.h
--- a/Engine/Source/Runtime/UMG/Public/Components/EditableText.h
+++ b/Engine/Source/Runtime/UMG/Public/Components/EditableText.h
@@ -53,7 +53,7 @@ public:
 	FGetText HintTextDelegate;
 
 	/** The style */
-	UPROPERTY(EditAnywhere, BlueprintReadWrite, Category=Appearance, meta=(DisplayName="Style", ShowOnlyInnerProperties))
+	UPROPERTY(EditAnywhere, BlueprintReadWrite, Setter, Category=Appearance, meta=(DisplayName="Style", ShowOnlyInnerProperties))
 	FEditableTextStyle WidgetStyle;
 
 	/** Sets whether this text box can actually be modified interactively by the user */
@@ -204,6 +204,9 @@ public:
 	/** Whether to select all text when pressing enter to commit changes */
 	bool GetSelectAllTextOnCommit() const;
 
+	/** Set the style and ensure the visual will be updated accordingly */
+	void SetWidgetStyle(const FEditableTextStyle& InEditableTextStyle);
+
 	bool GetIsReadOnly() const;
 
 	UFUNCTION(BlueprintCallable, Category = "Widget", meta=(DisplayName="SetIsReadOnly (Editable Text)"))
