UEditableText and UEditableTextBox:
- Fixed issue where IsReadOnly and IsPassword parameters weren't being set when MyEditableText is created.

diff --git a/Engine/Source/Runtime/UMG/Private/Components/EditableText.cpp b/Engine/Source/Runtime/UMG/Private/Components/EditableText.cpp
--- a/Engine/Source/Runtime/UMG/Private/Components/EditableText.cpp
+++ b/Engine/Source/Runtime/UMG/Private/Components/EditableText.cpp
@@ -97,6 +97,8 @@ TSharedRef<SWidget> UEditableText::RebuildWidget()
 	PRAGMA_DISABLE_DEPRECATION_WARNINGS
 	MyEditableText = SNew(SEditableText)
 		.Style(&WidgetStyle)
+		.IsReadOnly(IsReadOnly)
+		.IsPassword(IsPassword)
 		.MinDesiredWidth(MinimumDesiredWidth)
 		.IsCaretMovedWhenGainFocus(IsCaretMovedWhenGainFocus)
 		.SelectAllTextWhenFocused(SelectAllTextWhenFocused)
diff --git a/Engine/Source/Runtime/UMG/Private/Components/EditableTextBox.cpp b/Engine/Source/Runtime/UMG/Private/Components/EditableTextBox.cpp
--- a/Engine/Source/Runtime/UMG/Private/Components/EditableTextBox.cpp
+++ b/Engine/Source/Runtime/UMG/Private/Components/EditableTextBox.cpp
@@ -116,6 +116,8 @@ TSharedRef<SWidget> UEditableTextBox::RebuildWidget()
 {
 	MyEditableTextBlock = SNew(SEditableTextBox)
 		.Style(&WidgetStyle)
+		.IsReadOnly(IsReadOnly)
+		.IsPassword(IsPassword)
 		.MinDesiredWidth(MinimumDesiredWidth)
 		.IsCaretMovedWhenGainFocus(IsCaretMovedWhenGainFocus)
 		.SelectAllTextWhenFocused(SelectAllTextWhenFocused)
