Hololens fixes

diff --git a/Engine/Source/Runtime/Core/Private/Misc/ConfigUtilities.cpp b/Engine/Source/Runtime/Core/Private/Misc/ConfigUtilities.cpp
--- a/Engine/Source/Runtime/Core/Private/Misc/ConfigUtilities.cpp
+++ b/Engine/Source/Runtime/Core/Private/Misc/ConfigUtilities.cpp
@@ -300,7 +300,7 @@ private:
 	};
 	friend uint32 GetTypeHash(const FConfigHistory& CH)
 	{
-		uint32 Hash = GetTypeHash(CH.Type);
+		uint32 Hash = ::GetTypeHash(CH.Type);
 		Hash = HashCombine(Hash, GetTypeHash(CH.FileName));
 		Hash = HashCombine(Hash, GetTypeHash(CH.SectionName));
 		Hash = HashCombine(Hash, GetTypeHash(CH.Key));
diff --git a/Engine/Source/Runtime/MovieScene/Private/EntitySystem/MovieSceneEntityLedger.cpp b/Engine/Source/Runtime/MovieScene/Private/EntitySystem/MovieSceneEntityLedger.cpp
--- a/Engine/Source/Runtime/MovieScene/Private/EntitySystem/MovieSceneEntityLedger.cpp
+++ b/Engine/Source/Runtime/MovieScene/Private/EntitySystem/MovieSceneEntityLedger.cpp
@@ -238,7 +238,7 @@ void FEntityLedger::TagGarbage(UMovieSceneEntitySystemLinker* Linker)
 
 	for (auto It = ImportedEntities.CreateIterator(); It; ++It)
 	{
-		if (It.Key().EntityOwner == nullptr)
+		if (!It.Key().EntityOwner.IsValid())
 		{
 			if (It.Value().EntityID)
 			{
diff --git a/Engine/Source/Runtime/MovieSceneTracks/Private/EntitySystem/Interrogation/MovieSceneInterrogationLinker.cpp b/Engine/Source/Runtime/MovieSceneTracks/Private/EntitySystem/Interrogation/MovieSceneInterrogationLinker.cpp
--- a/Engine/Source/Runtime/MovieSceneTracks/Private/EntitySystem/Interrogation/MovieSceneInterrogationLinker.cpp
+++ b/Engine/Source/Runtime/MovieSceneTracks/Private/EntitySystem/Interrogation/MovieSceneInterrogationLinker.cpp
@@ -273,7 +273,7 @@ void FInterrogationChannels::QueryWorldSpaceTransforms(UMovieSceneEntitySystemLi
 		int32 NumChildren = 0;
 
 		FInterrogationChannel Channel = FInterrogationChannel::FromIndex(ChannelBit.GetIndex());
-		if (ActiveChannelBits[Channel.AsIndex()] == false && SparseChannelInfo.FindObject(Channel) == nullptr)
+		if ((bool)ActiveChannelBits[Channel.AsIndex()] == false && SparseChannelInfo.FindObject(Channel) == nullptr)
 		{
 			continue;
 		}
@@ -500,7 +500,7 @@ void FInterrogationChannels::QueryLocalSpaceTransforms(UMovieSceneEntitySystemLi
 	for (TConstSetBitIterator<> ChannelBit(ChannelsToQuery); ChannelBit; ++ChannelBit)
 	{
 		FInterrogationChannel Channel = FInterrogationChannel::FromIndex(ChannelBit.GetIndex());
-		if (ActiveChannelBits[Channel.AsIndex()] == false && SparseChannelInfo.FindObject(Channel) == nullptr)
+		if ((bool)ActiveChannelBits[Channel.AsIndex()] == false && SparseChannelInfo.FindObject(Channel) == nullptr)
 		{
 			continue;
 		}
diff --git a/Engine/Source/Runtime/UMG/Private/FieldNotification/FieldMulticastDelegate.cpp b/Engine/Source/Runtime/UMG/Private/FieldNotification/FieldMulticastDelegate.cpp
--- a/Engine/Source/Runtime/UMG/Private/FieldNotification/FieldMulticastDelegate.cpp
+++ b/Engine/Source/Runtime/UMG/Private/FieldNotification/FieldMulticastDelegate.cpp
@@ -233,7 +233,7 @@ FFieldMulticastDelegate::FRemoveFromResult FFieldMulticastDelegate::RemoveFrom(c
 			bRemoved = true;
 			return !bFieldPresent;
 		}
-		if (Element.Key.Object == InObject && Element.Delegate.IsBound())
+		if (Element.Key.Object.Get() == InObject && Element.Delegate.IsBound())
 		{
 			bFieldPresent = true;
 			return !bRemoved;
@@ -280,7 +280,7 @@ FFieldMulticastDelegate::FRemoveFromResult FFieldMulticastDelegate::RemoveFrom(c
 				}
 			}
 		}
-		if (Element.Key.Object == InObject && Element.Delegate.IsBound())
+		if (Element.Key.Object.Get() == InObject && Element.Delegate.IsBound())
 		{
 			bFieldPresent = true;
 			return !bRemoved;
@@ -310,7 +310,7 @@ FFieldMulticastDelegate::FRemoveAllResult FFieldMulticastDelegate::RemoveAll(con
 		{
 			if (const IDelegateInstance* DelegateInstance = GetDelegateInstance(Element.Delegate))
 			{
-				if (Element.Key.Object == InObject)
+				if (Element.Key.Object.Get() == InObject)
 				{
 					if (DelegateInstance->HasSameObject(InUserObject))
 					{
@@ -336,7 +336,7 @@ FFieldMulticastDelegate::FRemoveAllResult FFieldMulticastDelegate::RemoveAll(con
 			{
 				Delegates.RemoveAt(Index);
 			}
-			else if (Element.Key.Object == InObject)
+			else if (Element.Key.Object.Get() == InObject)
 			{
 				if (DelegateInstance->HasSameObject(InUserObject))
 				{
@@ -372,7 +372,7 @@ FFieldMulticastDelegate::FRemoveAllResult FFieldMulticastDelegate::RemoveAll(con
 		{
 			if (const IDelegateInstance* DelegateInstance = GetDelegateInstance(Element.Delegate))
 			{
-				if (Element.Key.Id == InFieldId && Element.Key.Object == InObject)
+				if (Element.Key.Id == InFieldId && Element.Key.Object.Get() == InObject)
 				{
 					if (DelegateInstance->HasSameObject(InUserObject))
 					{
@@ -398,7 +398,7 @@ FFieldMulticastDelegate::FRemoveAllResult FFieldMulticastDelegate::RemoveAll(con
 			{
 				Delegates.RemoveAt(Index);
 			}
-			else if (Element.Key.Id == InFieldId && Element.Key.Object == InObject)
+			else if (Element.Key.Id == InFieldId && Element.Key.Object.Get() == InObject)
 			{
 				if (DelegateInstance->HasSameObject(InUserObject))
 				{
@@ -433,7 +433,7 @@ void FFieldMulticastDelegate::Broadcast(UObject* InObject, UE::FieldNotification
 			{
 				break;
 			}
-			if (Element.Key.Object == InObject)
+			if (Element.Key.Object.Get() == InObject)
 			{
 				Element.Delegate.ExecuteIfBound(InObject, InFieldId);
 			}
@@ -444,7 +444,7 @@ void FFieldMulticastDelegate::Broadcast(UObject* InObject, UE::FieldNotification
 	for (int32 Index = Delegates.Num() - 1; Index >= AddedEmplaceAt; --Index)
 	{
 		FInvocationElement& Element = Delegates[Index];
-		if (Element.Key.Id == InFieldId && Element.Key.Object == InObject)
+		if (Element.Key.Id == InFieldId && Element.Key.Object.Get() == InObject)
 		{
 			Element.Delegate.ExecuteIfBound(InObject, InFieldId);
 		}
