Sequencer: Guard against crashing when animating null materials

 - Previously we would assert and subsequently crash in these situations, now we handle this gracefully and log an error

diff --git a/Engine/Source/Runtime/MovieSceneTracks/Private/Systems/MovieSceneComponentMaterialSystem.cpp b/Engine/Source/Runtime/MovieSceneTracks/Private/Systems/MovieSceneComponentMaterialSystem.cpp
--- a/Engine/Source/Runtime/MovieSceneTracks/Private/Systems/MovieSceneComponentMaterialSystem.cpp
+++ b/Engine/Source/Runtime/MovieSceneTracks/Private/Systems/MovieSceneComponentMaterialSystem.cpp
@@ -30,6 +30,11 @@ FComponentMaterialAccessor::FComponentMaterialAccessor(UObject* InObject, int32
 	, MaterialIndex(InMaterialIndex)
 {}
 
+FComponentMaterialAccessor::operator bool() const
+{
+	return Object != nullptr;
+}
+
 FString FComponentMaterialAccessor::ToString() const
 {
 	return FString::Printf(TEXT("object %s (element %d)"), *Object->GetPathName(), MaterialIndex);
diff --git a/Engine/Source/Runtime/MovieSceneTracks/Public/Systems/MovieSceneComponentMaterialSystem.h b/Engine/Source/Runtime/MovieSceneTracks/Public/Systems/MovieSceneComponentMaterialSystem.h
--- a/Engine/Source/Runtime/MovieSceneTracks/Public/Systems/MovieSceneComponentMaterialSystem.h
+++ b/Engine/Source/Runtime/MovieSceneTracks/Public/Systems/MovieSceneComponentMaterialSystem.h
@@ -40,6 +40,8 @@ struct FComponentMaterialAccessor
 	FComponentMaterialAccessor(const FComponentMaterialKey& InKey);
 	FComponentMaterialAccessor(UObject* InObject, int32 InMaterialIndex);
 
+	explicit operator bool() const;
+
 	UMaterialInterface* GetMaterial() const;
 	void SetMaterial(UMaterialInterface* InMaterial) const;
 	UMaterialInstanceDynamic* CreateDynamicMaterial(UMaterialInterface* InMaterial);
diff --git a/Engine/Source/Runtime/MovieSceneTracks/Public/Systems/MovieSceneMaterialSystem.h b/Engine/Source/Runtime/MovieSceneTracks/Public/Systems/MovieSceneMaterialSystem.h
--- a/Engine/Source/Runtime/MovieSceneTracks/Public/Systems/MovieSceneMaterialSystem.h
+++ b/Engine/Source/Runtime/MovieSceneTracks/Public/Systems/MovieSceneMaterialSystem.h
@@ -151,6 +151,11 @@ struct TInitializeBoundMaterials
 	static bool InitializeBoundMaterial(typename TCallTraits<RequiredComponents>::ParamType... Inputs, UObject*& OutDynamicMaterial)
 	{
 		AccessorType Accessor(Inputs...);
+		if (!Accessor)
+		{
+			OutDynamicMaterial = nullptr;
+			return false;
+		}
 
 		UMaterialInterface* ExistingMaterial = Accessor.GetMaterial();
 
diff --git a/Engine/Source/Runtime/UMG/Private/Animation/MovieSceneWidgetMaterialSystem.cpp b/Engine/Source/Runtime/UMG/Private/Animation/MovieSceneWidgetMaterialSystem.cpp
--- a/Engine/Source/Runtime/UMG/Private/Animation/MovieSceneWidgetMaterialSystem.cpp
+++ b/Engine/Source/Runtime/UMG/Private/Animation/MovieSceneWidgetMaterialSystem.cpp
@@ -13,6 +13,7 @@
 
 #include "Materials/MaterialInstanceDynamic.h"
 
+#include "UMGPrivate.h"
 #include "Components/Widget.h"
 
 #include "String/Join.h"
@@ -32,7 +33,17 @@ FWidgetMaterialAccessor::FWidgetMaterialAccessor(UObject* InObject, FWidgetMater
 	, WidgetMaterialPath(MoveTemp(InWidgetMaterialPath))
 {
 	// Object must be a widget
-	check(!InObject || Widget);
+	if (InObject && !Widget)
+	{
+		UE_LOG(LogUMG, Warning, TEXT("Cannot animate widget material on object %s of type %s"),
+			*InObject->GetName(), *InObject->GetClass()->GetName()
+		);
+	}
+}
+
+FWidgetMaterialAccessor::operator bool() const
+{
+	return Widget != nullptr;
 }
 
 FString FWidgetMaterialAccessor::ToString() const
diff --git a/Engine/Source/Runtime/UMG/Public/Animation/MovieSceneWidgetMaterialSystem.h b/Engine/Source/Runtime/UMG/Public/Animation/MovieSceneWidgetMaterialSystem.h
--- a/Engine/Source/Runtime/UMG/Public/Animation/MovieSceneWidgetMaterialSystem.h
+++ b/Engine/Source/Runtime/UMG/Public/Animation/MovieSceneWidgetMaterialSystem.h
@@ -46,6 +46,8 @@ struct FWidgetMaterialAccessor
 	FWidgetMaterialAccessor(const FWidgetMaterialKey& InKey);
 	FWidgetMaterialAccessor(UObject* InObject, FWidgetMaterialPath InWidgetMaterialPath);
 
+	explicit operator bool() const;
+
 	UMaterialInterface* GetMaterial() const;
 	void SetMaterial(UMaterialInterface* InMaterial) const;
 	UMaterialInstanceDynamic* CreateDynamicMaterial(UMaterialInterface* InMaterial);
