[23.00]

UListView:
- Added BlueprintAssignable delegate BP_OnListViewScrolled.
- Override OnListViewScrolledInternal and broadcast BP_OnListViewScrolled.

diff --git a/Engine/Source/Runtime/UMG/Private/Components/ListView.cpp b/Engine/Source/Runtime/UMG/Private/Components/ListView.cpp
--- a/Engine/Source/Runtime/UMG/Private/Components/ListView.cpp
+++ b/Engine/Source/Runtime/UMG/Private/Components/ListView.cpp
@@ -349,24 +349,34 @@ UUserWidget& UListView::OnGenerateEntryWidgetInternal(UObject* Item, TSubclassOf
 
 void UListView::OnItemClickedInternal(UObject* ListItem)
 {
+	ITypedUMGListView::OnItemClickedInternal(ListItem);
 	BP_OnItemClicked.Broadcast(ListItem);
 }
 
 void UListView::OnItemDoubleClickedInternal(UObject* ListItem)
 {
+	ITypedUMGListView::OnItemDoubleClickedInternal(ListItem);
 	BP_OnItemDoubleClicked.Broadcast(ListItem);
 }
 
 void UListView::OnSelectionChangedInternal(UObject* FirstSelectedItem)
 {
+	ITypedUMGListView::OnSelectionChangedInternal(FirstSelectedItem);
 	BP_OnItemSelectionChanged.Broadcast(FirstSelectedItem, FirstSelectedItem != nullptr);
 }
 
 void UListView::OnItemScrolledIntoViewInternal(UObject* ListItem, UUserWidget& EntryWidget)
 {
+	ITypedUMGListView::OnItemScrolledIntoViewInternal(ListItem, EntryWidget);
 	BP_OnItemScrolledIntoView.Broadcast(ListItem, &EntryWidget);
 }
 
+void UListView::OnListViewScrolledInternal(float ItemOffset, float DistanceRemaining)
+{
+	ITypedUMGListView::OnListViewScrolledInternal(ItemOffset, DistanceRemaining);
+	BP_OnListViewScrolled.Broadcast(ItemOffset, DistanceRemaining);
+}
+
 /////////////////////////////////////////////////////
 
 #undef LOCTEXT_NAMESPACE
diff --git a/Engine/Source/Runtime/UMG/Public/Components/ListView.h b/Engine/Source/Runtime/UMG/Public/Components/ListView.h
--- a/Engine/Source/Runtime/UMG/Public/Components/ListView.h
+++ b/Engine/Source/Runtime/UMG/Public/Components/ListView.h
@@ -12,6 +12,7 @@ DECLARE_DYNAMIC_MULTICAST_DELEGATE_TwoParams(FOnListEntryInitializedDynamic, UOb
 DECLARE_DYNAMIC_MULTICAST_DELEGATE_TwoParams(FOnListItemSelectionChangedDynamic, UObject*, Item, bool, bIsSelected);
 DECLARE_DYNAMIC_MULTICAST_DELEGATE_TwoParams(FOnItemIsHoveredChangedDynamic, UObject*, Item, bool, bIsHovered);
 DECLARE_DYNAMIC_MULTICAST_DELEGATE_TwoParams(FOnListItemScrolledIntoViewDynamic, UObject*, Item, UUserWidget*, Widget);
+DECLARE_DYNAMIC_MULTICAST_DELEGATE_TwoParams(FOnListViewScrolledDynamic, float, ItemOffset, float, DistanceRemaining);
 
 /**
  * A virtualized list that allows up to thousands of items to be displayed.
@@ -145,6 +146,7 @@ protected:
 	virtual void OnItemDoubleClickedInternal(UObject* Item) override;
 	virtual void OnSelectionChangedInternal(UObject* FirstSelectedItem) override;
 	virtual void OnItemScrolledIntoViewInternal(UObject* Item, UUserWidget& EntryWidget) override;
+	virtual void OnListViewScrolledInternal(float ItemOffset, float DistanceRemaining) override;
 
 	void HandleOnEntryInitializedInternal(UObject* Item, const TSharedRef<ITableRow>& TableRow);
 
@@ -272,4 +274,7 @@ private:
 
 	UPROPERTY(BlueprintAssignable, Category = Events, meta = (DisplayName = "On Item Scrolled Into View"))
 	FOnListItemScrolledIntoViewDynamic BP_OnItemScrolledIntoView;
+	
+	UPROPERTY(BlueprintAssignable, Category = Events, meta = (DisplayName = "On List View Scrolled"))
+	FOnListViewScrolledDynamic BP_OnListViewScrolled;	
 };
