UMG: Fix invalid slot in StackBox introduced with CL 22904491

diff --git a/Engine/Source/Runtime/UMG/Private/Components/StackBoxSlot.cpp b/Engine/Source/Runtime/UMG/Private/Components/StackBoxSlot.cpp
--- a/Engine/Source/Runtime/UMG/Private/Components/StackBoxSlot.cpp
+++ b/Engine/Source/Runtime/UMG/Private/Components/StackBoxSlot.cpp
@@ -34,22 +34,30 @@ void UStackBoxSlot::BuildSlot(TSharedRef<SStackBox> StackBox)
 
 void UStackBoxSlot::ReplaceContent(UWidget* NewContent)
 {
-	if (NewContent)
+	if (Content != NewContent)
 	{
-		NewContent->RemoveFromParent();
+		if (NewContent)
+		{
+			NewContent->RemoveFromParent();
+		}
 
+		if (UWidget* PreviousWidget = Content)
 		{
-			UWidget* PreviousWidget = Content;
+			// Setting Slot=null before RemoveFromParent to prevent destroying this slot
 			PreviousWidget->Slot = nullptr;
 			PreviousWidget->RemoveFromParent();
 		}
 
 		Content = NewContent;
-		NewContent->Slot = this;
+
+		if (Content)
+		{
+			Content->Slot = this;
+		}
 
 		if (Slot)
 		{
-			Slot->AttachWidget(Content->TakeWidget());
+			Slot->AttachWidget(Content == nullptr ? SNullWidget::NullWidget : Content->TakeWidget());
 		}
 	}
 }
