UMG: Update the CircularThrobber with getter and setter

diff --git a/Engine/Source/Runtime/Slate/Private/Widgets/Images/SThrobber.cpp b/Engine/Source/Runtime/Slate/Private/Widgets/Images/SThrobber.cpp
--- a/Engine/Source/Runtime/Slate/Private/Widgets/Images/SThrobber.cpp
+++ b/Engine/Source/Runtime/Slate/Private/Widgets/Images/SThrobber.cpp
@@ -69,6 +69,11 @@ void SThrobber::SetPieceImage(const FSlateBrush* InPieceImage)
 	}
 }
 
+void SThrobber::InvalidatePieceImage()
+{
+	Invalidate(EInvalidateWidgetReason::Paint);
+}
+
 void SThrobber::SetNumPieces(int InNumPieces)
 {
 	if (NumPieces != InNumPieces)
@@ -152,6 +157,11 @@ void SCircularThrobber::SetPieceImage(const FSlateBrush* InPieceImage)
 	}
 }
 
+void SCircularThrobber::InvalidatePieceImage()
+{
+	Invalidate(EInvalidateWidgetReason::Paint);
+}
+
 void SCircularThrobber::SetNumPieces(const int32 InNumPieces)
 {
 	if (NumPieces != InNumPieces)
diff --git a/Engine/Source/Runtime/Slate/Public/Widgets/Images/SThrobber.h b/Engine/Source/Runtime/Slate/Public/Widgets/Images/SThrobber.h
--- a/Engine/Source/Runtime/Slate/Public/Widgets/Images/SThrobber.h
+++ b/Engine/Source/Runtime/Slate/Public/Widgets/Images/SThrobber.h
@@ -54,6 +54,12 @@ public:
 	/** Sets what each segment of the throbber looks like. */
 	void SetPieceImage(const FSlateBrush* InPieceImage);
 
+	/**
+	 * Invalidate the SCircularThrobber with the correct reason.
+	 * You should invalidate when you are changing a property of SlateBrush after the SlateBrush was set to the SCircularThrobber.
+	 */
+	void InvalidatePieceImage();
+
 	/** Sets how many pieces there are */
 	void SetNumPieces(int InNumPieces);
 
@@ -133,6 +139,12 @@ public:
 	/** Sets what each segment of the throbber looks like */
 	void SetPieceImage(const FSlateBrush* InPieceImage);
 
+	/**
+	 * Invalidate the SCircularThrobber with the correct reason. 
+	 * You should invalidate when you are changing a property of SlateBrush after the SlateBrush was set to the SCircularThrobber.
+	 */
+	void InvalidatePieceImage();
+
 	/** Sets how many pieces there are */
 	void SetNumPieces(int32 InNumPieces);
 
diff --git a/Engine/Source/Runtime/UMG/Private/Components/CircularThrobber.cpp b/Engine/Source/Runtime/UMG/Private/Components/CircularThrobber.cpp
--- a/Engine/Source/Runtime/UMG/Private/Components/CircularThrobber.cpp
+++ b/Engine/Source/Runtime/UMG/Private/Components/CircularThrobber.cpp
@@ -32,6 +32,7 @@ UCircularThrobber::UCircularThrobber(const FObjectInitializer& ObjectInitializer
 		DefaultCircularThrobberBrushStyle->UnlinkColors();
 	}
 
+PRAGMA_DISABLE_DEPRECATION_WARNINGS
 	Image = *DefaultCircularThrobberBrushStyle;
 
 #if WITH_EDITOR 
@@ -55,6 +56,7 @@ UCircularThrobber::UCircularThrobber(const FObjectInitializer& ObjectInitializer
 	NumberOfPieces = 6;
 	Period = 0.75f;
 	Radius = 16.f;
+PRAGMA_ENABLE_DEPRECATION_WARNINGS
 }
 
 void UCircularThrobber::ReleaseSlateResources(bool bReleaseChildren)
@@ -66,11 +68,13 @@ void UCircularThrobber::ReleaseSlateResources(bool bReleaseChildren)
 
 TSharedRef<SWidget> UCircularThrobber::RebuildWidget()
 {
+PRAGMA_DISABLE_DEPRECATION_WARNINGS
 	MyCircularThrobber = SNew(SCircularThrobber)
 		.PieceImage(&Image)
 		.NumPieces(FMath::Clamp(NumberOfPieces, 1, 25))
 		.Period(FMath::Max(Period, SCircularThrobber::MinimumPeriodValue))
 		.Radius(Radius);
+PRAGMA_ENABLE_DEPRECATION_WARNINGS
 
 	return MyCircularThrobber.ToSharedRef();
 }
@@ -79,9 +83,13 @@ void UCircularThrobber::SynchronizeProperties()
 {
 	Super::SynchronizeProperties();
 
+PRAGMA_DISABLE_DEPRECATION_WARNINGS
 	MyCircularThrobber->SetNumPieces(FMath::Clamp(NumberOfPieces, 1, 25));
 	MyCircularThrobber->SetPeriod(FMath::Max(Period, SCircularThrobber::MinimumPeriodValue));
 	MyCircularThrobber->SetRadius(Radius);
+	MyCircularThrobber->SetPieceImage(&Image);
+	MyCircularThrobber->InvalidatePieceImage();
+PRAGMA_ENABLE_DEPRECATION_WARNINGS
 
 	// If widget is child of Canvas Panel and 'Size to Content' is enabled, we allow user to modify radius.
 	bEnableRadius = true;
@@ -94,6 +102,7 @@ void UCircularThrobber::SynchronizeProperties()
 	}
 }
 
+PRAGMA_DISABLE_DEPRECATION_WARNINGS
 void UCircularThrobber::SetNumberOfPieces(int32 InNumberOfPieces)
 {
 	NumberOfPieces = InNumberOfPieces;
@@ -103,6 +112,11 @@ void UCircularThrobber::SetNumberOfPieces(int32 InNumberOfPieces)
 	}
 }
 
+int32 UCircularThrobber::GetNumberOfPieces() const
+{
+	return NumberOfPieces;
+}
+
 void UCircularThrobber::SetPeriod(float InPeriod)
 {
 	Period = InPeriod;
@@ -112,6 +126,11 @@ void UCircularThrobber::SetPeriod(float InPeriod)
 	}
 }
 
+float UCircularThrobber::GetPeriod() const
+{
+	return Period;
+}
+
 void UCircularThrobber::SetRadius(float InRadius)
 {
 	Radius = InRadius;
@@ -121,6 +140,26 @@ void UCircularThrobber::SetRadius(float InRadius)
 	}
 }
 
+float UCircularThrobber::GetRadius() const
+{
+	return Radius;
+}
+
+void UCircularThrobber::SetImage(const FSlateBrush& InImage)
+{
+	Image = InImage;
+	if (MyCircularThrobber.IsValid())
+	{
+		MyCircularThrobber->InvalidatePieceImage();
+	}
+}
+
+const FSlateBrush& UCircularThrobber::GetImage() const
+{
+	return Image;
+}
+PRAGMA_ENABLE_DEPRECATION_WARNINGS
+
 #if WITH_EDITOR
 
 const FText UCircularThrobber::GetPaletteCategory()
diff --git a/Engine/Source/Runtime/UMG/Private/Components/Throbber.cpp b/Engine/Source/Runtime/UMG/Private/Components/Throbber.cpp
--- a/Engine/Source/Runtime/UMG/Private/Components/Throbber.cpp
+++ b/Engine/Source/Runtime/UMG/Private/Components/Throbber.cpp
@@ -83,6 +83,7 @@ void UThrobber::SynchronizeProperties()
 
 PRAGMA_DISABLE_DEPRECATION_WARNINGS
 	MyThrobber->SetPieceImage(&Image);
+	MyThrobber->InvalidatePieceImage();
 	MyThrobber->SetNumPieces(FMath::Clamp(NumberOfPieces, 1, 25));
 	MyThrobber->SetAnimate(GetAnimation());
 PRAGMA_ENABLE_DEPRECATION_WARNINGS
@@ -165,9 +166,7 @@ void UThrobber::SetImage(const FSlateBrush& Brush)
 	Image = Brush;
 	if (MyThrobber.IsValid())
 	{
-		// to force a new invalidation
-		MyThrobber->SetPieceImage(nullptr);
-		MyThrobber->SetPieceImage(&Image);
+		MyThrobber->InvalidatePieceImage();
 	}
 }
 
diff --git a/Engine/Source/Runtime/UMG/Public/Components/CircularThrobber.h b/Engine/Source/Runtime/UMG/Public/Components/CircularThrobber.h
--- a/Engine/Source/Runtime/UMG/Public/Components/CircularThrobber.h
+++ b/Engine/Source/Runtime/UMG/Public/Components/CircularThrobber.h
@@ -26,18 +26,23 @@ class UMG_API UCircularThrobber : public UWidget
 public:
 
 	/** How many pieces there are */
-	UPROPERTY(EditAnywhere, BlueprintReadOnly, Category=Appearance, meta=( ClampMin="1", ClampMax="25", UIMin="1", UIMax="25" ))
+	UE_DEPRECATED(5.2, "Direct access to NumberOfPieces is deprecated. Please use the getter or setter.")
+	UPROPERTY(EditAnywhere, BlueprintReadWrite, Getter, Setter, BlueprintSetter="SetNumberOfPieces", Category = Appearance, meta = (ClampMin = "1", ClampMax = "25", UIMin = "1", UIMax = "25"))
 	int32 NumberOfPieces;
 
 	/** The amount of time for a full circle (in seconds) */
-	UPROPERTY(EditAnywhere, BlueprintReadOnly, Category=Appearance, meta=( ClampMin="0", UIMin="0" ))
+	UE_DEPRECATED(5.2, "Direct access to Period is deprecated. Please use the getter or setter.")
+	UPROPERTY(EditAnywhere, BlueprintReadWrite, Getter, Setter, BlueprintSetter="SetPeriod", Category=Appearance, meta=( ClampMin="0", UIMin="0" ))
 	float Period;
 
 	/** The radius of the circle. If the throbber is a child of Canvas Panel, the 'Size to Content' option must be enabled in order to set Radius. */
-	UPROPERTY(EditAnywhere, BlueprintReadOnly, Category=Appearance, meta=(EditCondition="bEnableRadius"))
+	UE_DEPRECATED(5.2, "Direct access to Radius is deprecated. Please use the getter or setter.")
+	UPROPERTY(EditAnywhere, BlueprintReadWrite, Getter, Setter, BlueprintSetter="SetRadius", Category=Appearance, meta=(EditCondition="bEnableRadius"))
 	float Radius;
 
-	UPROPERTY(EditAnywhere, BlueprintReadOnly, Category=Appearance)
+	/** The throbber image. */
+	UE_DEPRECATED(5.2, "Direct access to Image is deprecated. Please use the getter or setter.")
+	UPROPERTY(EditAnywhere, BlueprintReadWrite, Getter, Setter, Category=Appearance)
 	FSlateBrush Image;
 
 public:
@@ -45,34 +50,36 @@ public:
 	/** Sets how many pieces there are. */
 	UFUNCTION(BlueprintCallable, Category = "Appearance")
 	void SetNumberOfPieces(int32 InNumberOfPieces);
+	int32 GetNumberOfPieces() const;
 
 	/** Sets the amount of time for a full circle (in seconds). */
 	UFUNCTION(BlueprintCallable, Category = "Appearance")
 	void SetPeriod(float InPeriod);
+	float GetPeriod() const;
 
 	/** Sets the radius of the circle. */
 	UFUNCTION(BlueprintCallable, Category = "Appearance")
 	void SetRadius(float InRadius);
+	float GetRadius() const;
+
+	/** Sets the throbber image. */
+	void SetImage(const FSlateBrush& InRadius);
+	const FSlateBrush& GetImage() const;
 
 	//~ Begin UWidget Interface
+public:
 	virtual void SynchronizeProperties() override;
-	//~ End UWidget Interface
-
-	//~ Begin UVisual Interface
 	virtual void ReleaseSlateResources(bool bReleaseChildren) override;
-	//~ End UVisual Interface
 
 #if WITH_EDITOR
 	virtual const FText GetPaletteCategory() override;
 #endif
 
 protected:
-	//~ Begin UWidget Interface
 	virtual TSharedRef<SWidget> RebuildWidget() override;
 	//~ End UWidget Interface
 
 private:
-
 	UPROPERTY(Transient, EditAnywhere, Category="Appearance", meta=(InlineEditConditionToggle))
 	bool bEnableRadius;
 
