Add accessors for SpinBox, WidgetSwitcher, Viewport

diff --git a/Engine/Plugins/Experimental/CommonUI/Source/CommonUI/Private/CommonAnimatedSwitcher.cpp b/Engine/Plugins/Experimental/CommonUI/Source/CommonUI/Private/CommonAnimatedSwitcher.cpp
--- a/Engine/Plugins/Experimental/CommonUI/Source/CommonUI/Private/CommonAnimatedSwitcher.cpp
+++ b/Engine/Plugins/Experimental/CommonUI/Source/CommonUI/Private/CommonAnimatedSwitcher.cpp
@@ -53,7 +53,7 @@ void UCommonAnimatedSwitcher::ActivateNextWidget(bool bCanWrap)
 {
 	if (Slots.Num() > 1)
 	{
-		if (ActiveWidgetIndex == Slots.Num() - 1)
+		if (GetActiveWidgetIndex() == Slots.Num() - 1)
 		{
 			if (bCanWrap)
 			{
@@ -62,7 +62,7 @@ void UCommonAnimatedSwitcher::ActivateNextWidget(bool bCanWrap)
 		}
 		else
 		{
-			SetActiveWidgetIndex(ActiveWidgetIndex + 1);
+			SetActiveWidgetIndex(GetActiveWidgetIndex() + 1);
 		}
 	}
 }
@@ -71,7 +71,7 @@ void UCommonAnimatedSwitcher::ActivatePreviousWidget(bool bCanWrap)
 {
 	if (Slots.Num() > 1)
 	{
-		if (ActiveWidgetIndex == 0)
+		if (GetActiveWidgetIndex() == 0)
 		{
 			if (bCanWrap)
 			{
@@ -80,7 +80,7 @@ void UCommonAnimatedSwitcher::ActivatePreviousWidget(bool bCanWrap)
 		}
 		else
 		{
-			SetActiveWidgetIndex(ActiveWidgetIndex - 1);
+			SetActiveWidgetIndex(GetActiveWidgetIndex() - 1);
 		}
 	}
 }
@@ -102,16 +102,16 @@ bool UCommonAnimatedSwitcher::IsCurrentlySwitching() const
 
 void UCommonAnimatedSwitcher::HandleSlateActiveIndexChanged(int32 ActiveIndex)
 {
-	if (Slots.IsValidIndex(ActiveWidgetIndex))
+	if (Slots.IsValidIndex(GetActiveWidgetIndex()))
 	{
-		OnActiveWidgetIndexChanged.Broadcast(GetWidgetAtIndex(ActiveWidgetIndex), ActiveWidgetIndex);
+		OnActiveWidgetIndexChanged.Broadcast(GetWidgetAtIndex(GetActiveWidgetIndex()), GetActiveWidgetIndex());
 	}
 }
 
 TSharedRef<SWidget> UCommonAnimatedSwitcher::RebuildWidget()
 {
 	MyWidgetSwitcher = MyAnimatedSwitcher = SNew(SCommonAnimatedSwitcher)
-		.InitialIndex(ActiveWidgetIndex)
+		.InitialIndex(GetActiveWidgetIndex())
 		.TransitionCurveType(TransitionCurveType)
 		.TransitionDuration(TransitionDuration)
 		.TransitionType(TransitionType)
@@ -159,23 +159,23 @@ void UCommonAnimatedSwitcher::SetActiveWidgetIndex_Internal(int32 Index)
 	
 	TGuardValue<bool> bCurrentlySwitchingGuard(bCurrentlySwitching, true);
 
-	if (Index >= 0 && Index < Slots.Num() && (Index != ActiveWidgetIndex || !bSetOnce))
+	if (Index >= 0 && Index < Slots.Num() && (Index != GetActiveWidgetIndex() || !bSetOnce))
 	{
 		HandleOutgoingWidget();
 
-		ActiveWidgetIndex = Index;
+		SetActiveWidgetIndex(Index);
 
 		if (MyAnimatedSwitcher.IsValid())
 		{
 			// Ensure the index is clamped to a valid range.
-			int32 SafeIndex = FMath::Clamp(ActiveWidgetIndex, 0, FMath::Max(0, Slots.Num() - 1));
+			int32 SafeIndex = FMath::Clamp(GetActiveWidgetIndex(), 0, FMath::Max(0, Slots.Num() - 1));
 			MyAnimatedSwitcher->TransitionToIndex(SafeIndex, bInstantTransition);
 		}
 
 		//When we're setting up, and the index goes 0->0, MyAnimatedSwitcher won't fire its ActiveIndexChanged Event.
 		if (!bSetOnce)
 		{
-			HandleSlateActiveIndexChanged(ActiveWidgetIndex);
+			HandleSlateActiveIndexChanged(GetActiveWidgetIndex());
 		}
 
 		bSetOnce = true;
diff --git a/Engine/Source/Runtime/Slate/Public/Widgets/Input/SSpinBox.h b/Engine/Source/Runtime/Slate/Public/Widgets/Input/SSpinBox.h
--- a/Engine/Source/Runtime/Slate/Public/Widgets/Input/SSpinBox.h
+++ b/Engine/Source/Runtime/Slate/Public/Widgets/Input/SSpinBox.h
@@ -869,6 +869,15 @@ public:
 	float GetMinDesiredWidth() const { return SliderExponent.Get(); }
 	void SetMinDesiredWidth(const TAttribute<float>& InMinDesiredWidth) { MinDesiredWidth = InMinDesiredWidth; }
 
+	const FSpinBoxStyle* GetWidgetStyle() const { return Style; }
+	void SetWidgetStyle(const FSpinBoxStyle* InStyle) { Style = InStyle; }
+	void InvalidateStyle() { Invalidate(EInvalidateWidgetReason::Layout); }
+
+	void SetTextBlockFont(FSlateFontInfo InFont) { EditableText->SetFont(InFont); TextBlock->SetFont(InFont); }
+	void SetTextJustification(ETextJustify::Type InJustification) { EditableText->SetJustification(InJustification); TextBlock->SetJustification(InJustification);  }
+	void SetTextClearKeyboardFocusOnCommit(bool bNewValue) { EditableText->SetClearKeyboardFocusOnCommit(bNewValue); }
+	void SetTextSelectAllTextOnCommit(bool bNewValue) { EditableText->SetSelectAllTextOnCommit(bNewValue); }
+
 protected:
 	/** Make the spinbox switch to keyboard-based input mode. */
 	void EnterTextMode()
diff --git a/Engine/Source/Runtime/UMG/Private/Components/SpinBox.cpp b/Engine/Source/Runtime/UMG/Private/Components/SpinBox.cpp
--- a/Engine/Source/Runtime/UMG/Private/Components/SpinBox.cpp
+++ b/Engine/Source/Runtime/UMG/Private/Components/SpinBox.cpp
@@ -24,6 +24,7 @@ USpinBox::USpinBox(const FObjectInitializer& ObjectInitializer)
 	if (!IsRunningDedicatedServer())
 	{
 		static ConstructorHelpers::FObjectFinder<UFont> RobotoFontObj(*UWidget::GetDefaultFontName());
+PRAGMA_DISABLE_DEPRECATION_WARNINGS
 		Font = FSlateFontInfo(RobotoFontObj.Object, 12, FName("Bold"));
 	}
 
@@ -42,6 +43,7 @@ USpinBox::USpinBox(const FObjectInitializer& ObjectInitializer)
 	ClearKeyboardFocusOnCommit = false;
 	SelectAllTextOnCommit = true;
 	KeyboardType = EVirtualKeyboardType::Number;
+PRAGMA_ENABLE_DEPRECATION_WARNINGS
 
 	if (DefaultSpinBoxStyle == nullptr)
 	{
@@ -50,9 +52,9 @@ USpinBox::USpinBox(const FObjectInitializer& ObjectInitializer)
 		// Unlink UMG default colors.
 		DefaultSpinBoxStyle->UnlinkColors();
 	}
-
+PRAGMA_DISABLE_DEPRECATION_WARNINGS
 	WidgetStyle = *DefaultSpinBoxStyle;
-
+PRAGMA_ENABLE_DEPRECATION_WARNINGS
 #if WITH_EDITOR 
 	if (EditorSpinBoxStyle == nullptr)
 	{
@@ -64,14 +66,17 @@ USpinBox::USpinBox(const FObjectInitializer& ObjectInitializer)
 	
 	if (IsEditorWidget())
 	{
+PRAGMA_DISABLE_DEPRECATION_WARNINGS
 		WidgetStyle = *EditorSpinBoxStyle;
-
+PRAGMA_ENABLE_DEPRECATION_WARNINGS
 		// The CDO isn't an editor widget and thus won't use the editor style, call post edit change to mark difference from CDO
 		PostEditChange();
 	}
 #endif // WITH_EDITOR
 
+PRAGMA_DISABLE_DEPRECATION_WARNINGS
 	ForegroundColor = WidgetStyle.ForegroundColor;
+PRAGMA_ENABLE_DEPRECATION_WARNINGS
 }
 
 void USpinBox::ReleaseSlateResources(bool bReleaseChildren)
@@ -81,6 +86,7 @@ void USpinBox::ReleaseSlateResources(bool bReleaseChildren)
 	MySpinBox.Reset();
 }
 
+PRAGMA_DISABLE_DEPRECATION_WARNINGS
 TSharedRef<SWidget> USpinBox::RebuildWidget()
 {
 	MySpinBox = SNew(SSpinBox<float>)
@@ -98,11 +104,12 @@ TSharedRef<SWidget> USpinBox::RebuildWidget()
 	
 	return MySpinBox.ToSharedRef();
 }
+PRAGMA_ENABLE_DEPRECATION_WARNINGS
 
 void USpinBox::SynchronizeProperties()
 {
 	Super::SynchronizeProperties();
-
+PRAGMA_DISABLE_DEPRECATION_WARNINGS
 	MySpinBox->SetDelta(Delta);
 	MySpinBox->SetSliderExponent(SliderExponent);
 	MySpinBox->SetMinDesiredWidth(MinDesiredWidth);
@@ -120,37 +127,55 @@ void USpinBox::SynchronizeProperties()
 	bOverride_MinSliderValue ? SetMinSliderValue(MinSliderValue) : ClearMinSliderValue();
 	bOverride_MaxSliderValue ? SetMaxSliderValue(MaxSliderValue) : ClearMaxSliderValue();
 
+	MySpinBox->SetWidgetStyle(&WidgetStyle);
+	MySpinBox->InvalidateStyle();
+
+	MySpinBox->SetTextJustification(Justification);
+	MySpinBox->SetTextBlockFont(Font);
+	MySpinBox->SetTextClearKeyboardFocusOnCommit(ClearKeyboardFocusOnCommit);
+	MySpinBox->SetTextSelectAllTextOnCommit(ClearKeyboardFocusOnCommit);
+
 	// Always set the value last so that the max/min values are taken into account.
 	TAttribute<float> ValueBinding = PROPERTY_BINDING(float, Value);
+PRAGMA_ENABLE_DEPRECATION_WARNINGS
 	MySpinBox->SetValue(ValueBinding);
 }
 
+PRAGMA_DISABLE_DEPRECATION_WARNINGS
 float USpinBox::GetValue() const
 {
-	if (MySpinBox.IsValid())
-	{
-		return MySpinBox->GetValue();
-	}
-
 	return Value;
 }
 
 void USpinBox::SetValue(float InValue)
 {
-	Value = InValue;
-	if (MySpinBox.IsValid())
+	if (Value != InValue)
 	{
-		MySpinBox->SetValue(InValue);
+		Value = InValue;
+		BroadcastFieldValueChanged(FFieldNotificationClassDescriptor::Value);
+		if (MySpinBox.IsValid())
+		{
+			MySpinBox->SetValue(InValue);
+		}
 	}
 }
 
-int32 USpinBox::GetMinFractionalDigits() const
+const FSpinBoxStyle& USpinBox::GetWidgetStyle() const
 {
+	return WidgetStyle;
+}
+
+void USpinBox::SetWidgetStyle(const FSpinBoxStyle& InWidgetStyle)
+{
+	WidgetStyle = InWidgetStyle;
 	if (MySpinBox.IsValid())
 	{
-		return MySpinBox->GetMinFractionalDigits();
+		MySpinBox->InvalidateStyle();
 	}
+}
 
+int32 USpinBox::GetMinFractionalDigits() const
+{
 	return MinFractionalDigits;
 }
 
@@ -165,11 +190,6 @@ void USpinBox::SetMinFractionalDigits(int32 NewValue)
 
 int32 USpinBox::GetMaxFractionalDigits() const
 {
-	if (MySpinBox.IsValid())
-	{
-		return MySpinBox->GetMaxFractionalDigits();
-	}
-
 	return MaxFractionalDigits;
 }
 
@@ -202,13 +222,23 @@ void USpinBox::SetAlwaysUsesDeltaSnap(bool bNewValue)
 	}
 }
 
-float USpinBox::GetDelta() const
+bool USpinBox::GetEnableSlider() const
+{
+	return bEnableSlider;
+}
+
+void USpinBox::SetEnableSlider(bool bNewValue)
 {
+	bEnableSlider = bNewValue;
+
 	if (MySpinBox.IsValid())
 	{
-		return MySpinBox->GetDelta();
+		MySpinBox->SetEnableSlider(bNewValue);
 	}
+}
 
+float USpinBox::GetDelta() const
+{
 	return Delta;
 }
 
@@ -221,16 +251,96 @@ void USpinBox::SetDelta(float NewValue)
 	}
 }
 
-// MIN VALUE
-float USpinBox::GetMinValue() const
+float USpinBox::GetSliderExponent() const
 {
-	float ReturnVal = TNumericLimits<float>::Lowest();
+	return SliderExponent;
+}
 
+void USpinBox::SetSliderExponent(float NewValue)
+{
+	SliderExponent = NewValue;
 	if (MySpinBox.IsValid())
 	{
-		ReturnVal = MySpinBox->GetMinValue();
+		MySpinBox->SetSliderExponent(NewValue);
 	}
-	else if (bOverride_MinValue)
+}
+
+const FSlateFontInfo& USpinBox::GetFont() const
+{
+	return Font;
+}
+
+void USpinBox::SetFont(const FSlateFontInfo& InFont)
+{
+	Font = InFont;
+	if (MySpinBox.IsValid())
+	{
+		MySpinBox->SetTextBlockFont(InFont);
+	}
+}
+
+const ETextJustify::Type USpinBox::GetJustification() const
+{
+	return Justification;
+}
+
+void USpinBox::SetJustification(ETextJustify::Type InJustification)
+{
+	Justification = InJustification;
+	if (MySpinBox.IsValid())
+	{
+		MySpinBox->SetTextJustification(InJustification);
+	}
+}
+
+float USpinBox::GetMinDesiredWidth() const
+{
+	return MinDesiredWidth;
+}
+
+void USpinBox::SetMinDesiredWidth(float NewValue)
+{
+	MinDesiredWidth = NewValue;
+	if (MySpinBox.IsValid())
+	{
+		MySpinBox->SetMinDesiredWidth(NewValue);
+	}
+}
+
+bool USpinBox::GetClearKeyboardFocusOnCommit() const
+{
+	return ClearKeyboardFocusOnCommit;
+}
+
+void USpinBox::SetClearKeyboardFocusOnCommit(bool bNewValue)
+{
+	ClearKeyboardFocusOnCommit = bNewValue;
+	if (MySpinBox.IsValid())
+	{
+		MySpinBox->SetTextClearKeyboardFocusOnCommit(bNewValue);
+	}
+}
+
+bool USpinBox::GetSelectAllTextOnCommit() const
+{
+	return SelectAllTextOnCommit;
+}
+
+void USpinBox::SetSelectAllTextOnCommit(bool bNewValue)
+{
+	SelectAllTextOnCommit = bNewValue;
+	if (MySpinBox.IsValid())
+	{
+		MySpinBox->SetTextSelectAllTextOnCommit(bNewValue);
+	}
+}
+
+// MIN VALUE
+float USpinBox::GetMinValue() const
+{
+	float ReturnVal = TNumericLimits<float>::Lowest();
+
+	if (bOverride_MinValue)
 	{
 		ReturnVal = MinValue;
 	}
@@ -262,11 +372,7 @@ float USpinBox::GetMaxValue() const
 {
 	float ReturnVal = TNumericLimits<float>::Max();
 
-	if (MySpinBox.IsValid())
-	{
-		ReturnVal = MySpinBox->GetMaxValue();
-	}
-	else if (bOverride_MaxValue)
+	if (bOverride_MaxValue)
 	{
 		ReturnVal = MaxValue;
 	}
@@ -297,11 +403,7 @@ float USpinBox::GetMinSliderValue() const
 {
 	float ReturnVal = TNumericLimits<float>::Min();
 
-	if (MySpinBox.IsValid())
-	{
-		ReturnVal = MySpinBox->GetMinSliderValue();
-	}
-	else if (bOverride_MinSliderValue)
+	if (bOverride_MinSliderValue)
 	{
 		ReturnVal = MinSliderValue;
 	}
@@ -333,11 +435,7 @@ float USpinBox::GetMaxSliderValue() const
 {
 	float ReturnVal = TNumericLimits<float>::Max();
 
-	if (MySpinBox.IsValid())
-	{
-		ReturnVal = MySpinBox->GetMaxSliderValue();
-	}
-	else if (bOverride_MaxSliderValue)
+	if (bOverride_MaxSliderValue)
 	{
 		ReturnVal = MaxSliderValue;
 	}
@@ -373,6 +471,12 @@ void USpinBox::SetForegroundColor(FSlateColor InForegroundColor)
 	}
 }
 
+FSlateColor USpinBox::GetForegroundColor() const
+{
+	return ForegroundColor;
+}
+PRAGMA_ENABLE_DEPRECATION_WARNINGS
+
 // Event handlers
 void USpinBox::HandleOnValueChanged(float InValue)
 {
diff --git a/Engine/Source/Runtime/UMG/Private/Components/Viewport.cpp b/Engine/Source/Runtime/UMG/Private/Components/Viewport.cpp
--- a/Engine/Source/Runtime/UMG/Private/Components/Viewport.cpp
+++ b/Engine/Source/Runtime/UMG/Private/Components/Viewport.cpp
@@ -365,8 +365,9 @@ UViewport::UViewport(const FObjectInitializer& ObjectInitializer)
 	, ShowFlags(ESFIM_Game)
 {
 	bIsVariable = true;
-
+PRAGMA_DISABLE_DEPRECATION_WARNINGS
 	BackgroundColor = FLinearColor::Black;
+PRAGMA_ENABLE_DEPRECATION_WARNINGS
 	ShowFlags.DisableAdvancedFeatures();
 	//ParentArgs.IgnoreTextureAlpha(false);
 	//ParentArgs.EnableBlending(true);
@@ -411,7 +412,10 @@ void UViewport::SynchronizeProperties()
 
 	if ( ViewportWidget.IsValid() )
 	{
+		check(ViewportWidget->ViewportClient.IsValid());
+PRAGMA_DISABLE_DEPRECATION_WARNINGS
 		ViewportWidget->ViewportClient->SetBackgroundColor(BackgroundColor);
+PRAGMA_ENABLE_DEPRECATION_WARNINGS
 		ViewportWidget->ViewportClient->SetEngineShowFlags(ShowFlags);
 	}
 }
@@ -496,6 +500,23 @@ AActor* UViewport::Spawn(TSubclassOf<AActor> ActorClass)
 	return NULL;
 }
 
+PRAGMA_DISABLE_DEPRECATION_WARNINGS
+const FLinearColor& UViewport::GetBackgroundColor() const
+{
+	return BackgroundColor;
+}
+
+void UViewport::SetBackgroundColor(const FLinearColor& InColor)
+{
+	BackgroundColor = InColor;
+	if (ViewportWidget.IsValid())
+	{
+		check(ViewportWidget->ViewportClient.IsValid());
+		ViewportWidget->ViewportClient->SetBackgroundColor(BackgroundColor);
+	}
+}
+PRAGMA_ENABLE_DEPRECATION_WARNINGS
+
 #if WITH_EDITOR
 
 const FText UViewport::GetPaletteCategory()
diff --git a/Engine/Source/Runtime/UMG/Private/Components/WidgetSwitcher.cpp b/Engine/Source/Runtime/UMG/Private/Components/WidgetSwitcher.cpp
--- a/Engine/Source/Runtime/UMG/Private/Components/WidgetSwitcher.cpp
+++ b/Engine/Source/Runtime/UMG/Private/Components/WidgetSwitcher.cpp
@@ -41,24 +41,37 @@ int32 UWidgetSwitcher::GetActiveWidgetIndex() const
 	{
 		return MyWidgetSwitcher->GetActiveWidgetIndex();
 	}
-
+PRAGMA_DISABLE_DEPRECATION_WARNINGS
 	return ActiveWidgetIndex;
+PRAGMA_ENABLE_DEPRECATION_WARNINGS
 }
 
+PRAGMA_DISABLE_DEPRECATION_WARNINGS
 void UWidgetSwitcher::SetActiveWidgetIndex(int32 Index)
 {
-	ActiveWidgetIndex = Index;
-	if ( MyWidgetSwitcher.IsValid() )
+	if (ActiveWidgetIndex != Index)
 	{
-		// Ensure the index is clamped to a valid range.
-		int32 SafeIndex = FMath::Clamp(ActiveWidgetIndex, 0, FMath::Max(0, Slots.Num() - 1));
-		MyWidgetSwitcher->SetActiveWidgetIndex(SafeIndex);
+		ActiveWidgetIndex = Index;
+		BroadcastFieldValueChanged(FFieldNotificationClassDescriptor::ActiveWidgetIndex);
+
+		if (MyWidgetSwitcher.IsValid())
+		{
+			// Ensure the index is clamped to a valid range.
+			int32 SafeIndex = FMath::Clamp(ActiveWidgetIndex, 0, FMath::Max(0, Slots.Num() - 1));
+			MyWidgetSwitcher->SetActiveWidgetIndex(SafeIndex);
+		}
 	}
 }
 
+
 void UWidgetSwitcher::SetActiveWidget(UWidget* Widget)
 {
-	ActiveWidgetIndex = GetChildIndex(Widget);
+	int32 NewIndex = GetChildIndex(Widget);
+	if (ActiveWidgetIndex != NewIndex)
+	{
+		ActiveWidgetIndex = NewIndex;
+		BroadcastFieldValueChanged(FFieldNotificationClassDescriptor::ActiveWidgetIndex);
+	}
 	if ( MyWidgetSwitcher.IsValid() )
 	{
 		// Ensure the index is clamped to a valid range.
@@ -66,6 +79,7 @@ void UWidgetSwitcher::SetActiveWidget(UWidget* Widget)
 		MyWidgetSwitcher->SetActiveWidgetIndex(SafeIndex);
 	}
 }
+PRAGMA_ENABLE_DEPRECATION_WARNINGS
 
 UWidget* UWidgetSwitcher::GetWidgetAtIndex( int32 Index ) const
 {
@@ -128,8 +142,9 @@ TSharedRef<SWidget> UWidgetSwitcher::RebuildWidget()
 void UWidgetSwitcher::SynchronizeProperties()
 {
 	Super::SynchronizeProperties();
-
+PRAGMA_DISABLE_DEPRECATION_WARNINGS
 	SetActiveWidgetIndex(ActiveWidgetIndex);
+PRAGMA_ENABLE_DEPRECATION_WARNINGS
 }
 
 #if WITH_EDITOR
@@ -154,6 +169,7 @@ void UWidgetSwitcher::OnDescendantSelectedByDesigner(UWidget* DescendantWidget)
 	}
 }
 
+PRAGMA_DISABLE_DEPRECATION_WARNINGS
 void UWidgetSwitcher::OnDescendantDeselectedByDesigner(UWidget* DescendantWidget)
 {
 	SetActiveWidgetIndex(ActiveWidgetIndex);
@@ -165,6 +181,7 @@ void UWidgetSwitcher::PostEditChangeProperty(struct FPropertyChangedEvent& Prope
 
 	Super::PostEditChangeProperty(PropertyChangedEvent);
 }
+PRAGMA_ENABLE_DEPRECATION_WARNINGS
 
 #endif
 
diff --git a/Engine/Source/Runtime/UMG/Public/Components/SpinBox.h b/Engine/Source/Runtime/UMG/Public/Components/SpinBox.h
--- a/Engine/Source/Runtime/UMG/Public/Components/SpinBox.h
+++ b/Engine/Source/Runtime/UMG/Public/Components/SpinBox.h
@@ -28,7 +28,8 @@ public:
 public:
 
 	/** Value stored in this spin box */
-	UPROPERTY(EditAnywhere, Category=Content)
+	UE_DEPRECATED(5.2, "Direct access to Value is deprecated. Please use the getter or setter.")
+	UPROPERTY(EditAnywhere, BlueprintReadWrite, Getter, Setter, BlueprintSetter = "SetValue", BlueprintGetter = "GetValue", FieldNotify, Category = Content)
 	float Value;
 
 	/** A bindable delegate to allow logic to drive the value of the widget */
@@ -37,58 +38,71 @@ public:
 
 public:
 	/** The Style */
-	UPROPERTY(EditAnywhere, BlueprintReadWrite, Category="Style", meta=( DisplayName="Style" ))
+	UE_DEPRECATED(5.2, "Direct access to WidgetStyle is deprecated. Please use the getter or setter.")
+	UPROPERTY(EditAnywhere, BlueprintReadWrite, Getter, Setter, Category="Style", meta=( DisplayName="Style" ))
 	FSpinBoxStyle WidgetStyle;
 
 	/** The minimum required fractional digits - default 1 */
-	UPROPERTY(EditAnywhere, BlueprintReadWrite, BlueprintSetter = SetMinFractionalDigits, BlueprintGetter = GetMinFractionalDigits, Category = "Slider", meta = (ClampMin = 0, UIMin = 0))
+	UE_DEPRECATED(5.2, "Direct access to MinFractionalDigits is deprecated. Please use the getter or setter.")
+	UPROPERTY(EditAnywhere, BlueprintReadWrite, Getter, Setter, BlueprintSetter = SetMinFractionalDigits, BlueprintGetter = GetMinFractionalDigits, Category = "Slider", meta = (ClampMin = 0, UIMin = 0))
 	int32 MinFractionalDigits;
 
 	/** The maximume required fractional digits - default 6 */
-	UPROPERTY(EditAnywhere, BlueprintReadWrite, BlueprintSetter = SetMaxFractionalDigits, BlueprintGetter = GetMaxFractionalDigits, Category = "Slider", meta = (ClampMin = 0, UIMin = 0))
+	UE_DEPRECATED(5.2, "Direct access to MaxFractionalDigits is deprecated. Please use the getter or setter.")
+	UPROPERTY(EditAnywhere, BlueprintReadWrite, Getter, Setter, BlueprintSetter = SetMaxFractionalDigits, BlueprintGetter = GetMaxFractionalDigits, Category = "Slider", meta = (ClampMin = 0, UIMin = 0))
 	int32 MaxFractionalDigits;
 
 	/** Whether this spin box should use the delta snapping logic for typed values - default false */
-	UPROPERTY(EditAnywhere, BlueprintReadWrite, BlueprintSetter = SetAlwaysUsesDeltaSnap, BlueprintGetter = GetAlwaysUsesDeltaSnap, Category = "Slider")
+	UE_DEPRECATED(5.2, "Direct access to bAlwaysUsesDeltaSnap is deprecated. Please use the getter or setter.")
+	UPROPERTY(EditAnywhere, BlueprintReadWrite, Getter = GetAlwaysUsesDeltaSnap, Setter = SetAlwaysUsesDeltaSnap, BlueprintSetter = SetAlwaysUsesDeltaSnap, BlueprintGetter = GetAlwaysUsesDeltaSnap, Category = "Slider")
 	bool bAlwaysUsesDeltaSnap;
 
 	/** Whether this spin box should have slider feature enabled */
-	UPROPERTY(EditAnywhere, BlueprintReadWrite, Category = "Slider")
+	UE_DEPRECATED(5.2, "Direct access to bEnableSlider is deprecated. Please use the getter or setter.")
+	UPROPERTY(EditAnywhere, BlueprintReadWrite, Getter = "GetEnableSlider", Setter = "SetEnableSlider", Category = "Slider")
 	bool bEnableSlider = true;
 
 	/** The amount by which to change the spin box value as the slider moves. */
-	UPROPERTY(EditAnywhere, BlueprintReadWrite, BlueprintSetter = SetDelta, BlueprintGetter = GetDelta, Category = "Slider")
+	UE_DEPRECATED(5.2, "Direct access to Delta is deprecated. Please use the getter or setter.")
+	UPROPERTY(EditAnywhere, BlueprintReadWrite, Getter, Setter, BlueprintSetter = SetDelta, BlueprintGetter = GetDelta, Category = "Slider")
 	float Delta;
 
 	/** The exponent by which to increase the delta as the mouse moves. 1 is constant (never increases the delta). */
-	UPROPERTY(EditAnywhere, BlueprintReadOnly, Category="Slider")
+	UE_DEPRECATED(5.2, "Direct access to SliderExponent is deprecated. Please use the getter or setter.")
+	UPROPERTY(EditAnywhere, BlueprintReadWrite, Getter, Setter, Category="Slider")
 	float SliderExponent;
 	
 	/** Font color and opacity (overrides style) */
-	UPROPERTY(EditAnywhere, BlueprintReadOnly, Category="Display")
+	UE_DEPRECATED(5.2, "Direct access to Font is deprecated. Please use the getter or setter.")
+	UPROPERTY(EditAnywhere, BlueprintReadWrite, Getter, Setter, Category="Display")
 	FSlateFontInfo Font;
 
 	/** The justification the value text should appear as. */
-	UPROPERTY(EditAnywhere, BlueprintReadOnly, Category = "Display")
+	UE_DEPRECATED(5.2, "Direct access to Justification is deprecated. Please use the getter or setter.")
+	UPROPERTY(EditAnywhere, BlueprintReadWrite, Getter, Setter, Category = "Display")
 	TEnumAsByte<ETextJustify::Type> Justification;
 
 	/** The minimum width of the spin box */
-	UPROPERTY(EditAnywhere, BlueprintReadOnly, Category="Display", AdvancedDisplay, DisplayName = "Minimum Desired Width")
+	UE_DEPRECATED(5.2, "Direct access to MinDesiredWidth is deprecated. Please use the getter or setter.")
+	UPROPERTY(EditAnywhere, BlueprintReadWrite, Getter, Setter, Category="Display", AdvancedDisplay, DisplayName = "Minimum Desired Width")
 	float MinDesiredWidth;
 	
 	/** If we're on a platform that requires a virtual keyboard, what kind of keyboard should this widget use? */
-	UPROPERTY(EditAnywhere, BlueprintReadOnly, Category = "Input", AdvancedDisplay)
+	UPROPERTY(EditAnywhere, Category = "Input", AdvancedDisplay)
 	TEnumAsByte<EVirtualKeyboardType::Type> KeyboardType;
 
 	/** Whether to remove the keyboard focus from the spin box when the value is committed */
-	UPROPERTY(EditAnywhere, BlueprintReadOnly, Category="Input", AdvancedDisplay)
+	UE_DEPRECATED(5.2, "Direct access to ClearKeyboardFocusOnCommit is deprecated. Please use the getter or setter.")
+	UPROPERTY(EditAnywhere, BlueprintReadWrite, Getter, Setter, Category="Input", AdvancedDisplay)
 	bool ClearKeyboardFocusOnCommit;
 
 	/** Whether to select the text in the spin box when the value is committed */
-	UPROPERTY(EditAnywhere, BlueprintReadOnly, Category="Input", AdvancedDisplay)
+	UE_DEPRECATED(5.2, "Direct access to SelectAllTextOnCommit is deprecated. Please use the getter or setter.")
+	UPROPERTY(EditAnywhere, BlueprintReadWrite, Getter, Setter, Category="Input", AdvancedDisplay)
 	bool SelectAllTextOnCommit;
 
-	UPROPERTY(EditAnywhere, BlueprintReadOnly, Category="Style")
+	UE_DEPRECATED(5.2, "Direct access to ForegroundColor is deprecated. Please use the getter or setter.")
+	UPROPERTY(EditAnywhere, BlueprintReadWrite, Getter, Setter, BlueprintSetter = "SetForegroundColor", Category = "Style")
 	FSlateColor ForegroundColor;
 
 public:
@@ -120,6 +134,12 @@ public:
 
 public:
 
+	/** Set the style for the spin box. */
+	const FSpinBoxStyle& GetWidgetStyle() const;
+
+	/** Get the style for the spin box. */
+	void SetWidgetStyle(const FSpinBoxStyle& InWidgetStyle);
+
 	/** Get the current Min Fractional Digits for the spin box. */
 	UFUNCTION(BlueprintCallable, BlueprintGetter, Category = "Behavior")
 	int32 GetMinFractionalDigits() const;
@@ -144,6 +164,12 @@ public:
 	UFUNCTION(BlueprintCallable, BlueprintSetter, Category = "Behavior")
 	void SetAlwaysUsesDeltaSnap(bool bNewValue);
 
+	/** Get whether the spin box uses the slider feature. */
+	bool GetEnableSlider() const;
+
+	/** Set whether the spin box uses the slider feature. */
+	void SetEnableSlider(bool bNewValue);
+
 	/** Get the current delta for the spin box. */
 	UFUNCTION(BlueprintCallable, BlueprintGetter, Category = "Behavior")
 	float GetDelta() const;
@@ -152,6 +178,42 @@ public:
 	UFUNCTION(BlueprintCallable, BlueprintSetter, Category = "Behavior")
 	void SetDelta(float NewValue);
 
+	/** Get the current slider exponent for the spin box. */
+	float GetSliderExponent() const;
+
+	/** Set the slider exponent for the spin box. */
+	void SetSliderExponent(float NewValue);
+
+	/**  Get the font color and opacity that overrides the style font. */
+	const FSlateFontInfo& GetFont() const;
+
+	/** Set the font color and opacity that overrides the style font. */
+	void SetFont(const FSlateFontInfo& InFont);
+
+	/** Get the justification for value text. */
+	const ETextJustify::Type GetJustification() const;
+
+	/** Set the justification for value text. */
+	void SetJustification(ETextJustify::Type InJustification);
+
+	/** Get the minimum width of the spin box. */
+	float GetMinDesiredWidth() const;
+
+	/** Set the minimum width of the spin box. */
+	void SetMinDesiredWidth(float NewValue);
+
+	/** Get whether the keyboard focus is removed from the spin box when the value is committed. */
+	bool GetClearKeyboardFocusOnCommit() const;
+
+	/** Set whether the keyboard focus is removed from the spin box when the value is committed. */
+	void SetClearKeyboardFocusOnCommit(bool bNewValue);
+
+	/** Get whether to select the text in the spin box when the value is committed. */
+	bool GetSelectAllTextOnCommit() const;
+
+	/** Set whether to select the text in the spin box when the value is committed. */
+	void SetSelectAllTextOnCommit(bool bNewValue);
+
 	/** Get the current minimum value that can be manually set in the spin box. */
 	UFUNCTION(BlueprintCallable, Category="Behavior")
 	float GetMinValue() const;
@@ -204,6 +266,9 @@ public:
 	UFUNCTION(BlueprintCallable, Category = "Appearance")
 	void SetForegroundColor(FSlateColor InForegroundColor);
 
+	/** Get the foreground color of the spin box. */
+	FSlateColor GetForegroundColor() const;
+
 public:
 
 	//~ Begin UWidget Interface
@@ -231,7 +296,7 @@ protected:
 
 protected:
 	/** Whether the optional MinValue attribute of the widget is set */
-	UPROPERTY(EditAnywhere, Category = Content, meta=(InlineEditConditionToggle))
+	UPROPERTY(EditAnywhere, Category = Content, meta = (InlineEditConditionToggle))
 	uint32 bOverride_MinValue : 1;
 
 	/** Whether the optional MaxValue attribute of the widget is set */
@@ -239,7 +304,7 @@ protected:
 	uint32 bOverride_MaxValue : 1;
 
 	/** Whether the optional MinSliderValue attribute of the widget is set */
-	UPROPERTY(EditAnywhere, Category = Content, meta=(InlineEditConditionToggle))
+	UPROPERTY(EditAnywhere, Category = Content, meta = (InlineEditConditionToggle))
 	uint32 bOverride_MinSliderValue : 1;
 
 	/** Whether the optional MaxSliderValue attribute of the widget is set */
@@ -247,19 +312,23 @@ protected:
 	uint32 bOverride_MaxSliderValue : 1;
 
 	/** The minimum allowable value that can be manually entered into the spin box */
-	UPROPERTY(EditAnywhere, Category = Content, DisplayName = "Minimum Value", meta = (editcondition = "bOverride_MinValue"))
+	UE_DEPRECATED(5.2, "Direct access to MinValue is deprecated. Please use the getter or setter.")
+	UPROPERTY(EditAnywhere, BlueprintReadWrite, Getter, Setter, BlueprintSetter = "SetMinValue", BlueprintGetter = "GetMinValue", Category = Content, DisplayName = "Minimum Value", meta = (editcondition = "bOverride_MinValue"))
 	float MinValue;
 
 	/** The maximum allowable value that can be manually entered into the spin box */
-	UPROPERTY(EditAnywhere, Category = Content, DisplayName = "Maximum Value", meta = (editcondition = "bOverride_MaxValue"))
+	UE_DEPRECATED(5.2, "Direct access to MaxValue is deprecated. Please use the getter or setter.")
+	UPROPERTY(EditAnywhere, BlueprintReadWrite, Getter, Setter, BlueprintSetter = "SetMaxValue", BlueprintGetter = "GetMaxValue", Category = Content, DisplayName = "Maximum Value", meta = (editcondition = "bOverride_MaxValue"))
 	float MaxValue;
 
 	/** The minimum allowable value that can be specified using the slider */
-	UPROPERTY(EditAnywhere, Category = Content, DisplayName = "Minimum Slider Value", meta = (editcondition = "bOverride_MinSliderValue"))
+	UE_DEPRECATED(5.2, "Direct access to MinSliderValue is deprecated. Please use the getter or setter.")
+	UPROPERTY(EditAnywhere, BlueprintReadWrite, Getter, Setter, BlueprintSetter = "SetMinSliderValue", BlueprintGetter = "GetMinSliderValue", Category = Content, DisplayName = "Minimum Slider Value", meta = (editcondition = "bOverride_MinSliderValue"))
 	float MinSliderValue;
 
 	/** The maximum allowable value that can be specified using the slider */
-	UPROPERTY(EditAnywhere, Category = Content, DisplayName = "Maximum Slider Value", meta = (editcondition = "bOverride_MaxSliderValue"))
+	UE_DEPRECATED(5.2, "Direct access to MaxSliderValue is deprecated. Please use the getter or setter.")
+	UPROPERTY(EditAnywhere, BlueprintReadWrite, Getter, Setter, BlueprintSetter = "SetMaxSliderValue", BlueprintGetter = "GetMaxSliderValue", Category = Content, DisplayName = "Maximum Slider Value", meta = (editcondition = "bOverride_MaxSliderValue"))
 	float MaxSliderValue;
 
 protected:
diff --git a/Engine/Source/Runtime/UMG/Public/Components/Viewport.h b/Engine/Source/Runtime/UMG/Public/Components/Viewport.h
--- a/Engine/Source/Runtime/UMG/Public/Components/Viewport.h
+++ b/Engine/Source/Runtime/UMG/Public/Components/Viewport.h
@@ -229,7 +229,8 @@ class UMG_API UViewport : public UContentWidget
 {
 	GENERATED_UCLASS_BODY()
 
-	UPROPERTY(EditAnywhere, Category=Appearance)
+	UE_DEPRECATED(5.2, "Direct access to BackgroundColor is deprecated. Please use the getter or setter.")
+	UPROPERTY(EditAnywhere, BlueprintReadWrite, Getter, Setter, Category = Appearance)
 	FLinearColor BackgroundColor;
 
 	UFUNCTION(BlueprintCallable, Category="Viewport")
@@ -250,6 +251,10 @@ class UMG_API UViewport : public UContentWidget
 	UFUNCTION(BlueprintCallable, Category="Viewport")
 	AActor* Spawn(TSubclassOf<AActor> ActorClass);
 
+	void SetBackgroundColor(const FLinearColor& InColor);
+
+	const FLinearColor& GetBackgroundColor() const;
+
 	// UWidget interface
 	virtual void SynchronizeProperties() override;
 	// End of UWidget interface
diff --git a/Engine/Source/Runtime/UMG/Public/Components/WidgetSwitcher.h b/Engine/Source/Runtime/UMG/Public/Components/WidgetSwitcher.h
--- a/Engine/Source/Runtime/UMG/Public/Components/WidgetSwitcher.h
+++ b/Engine/Source/Runtime/UMG/Public/Components/WidgetSwitcher.h
@@ -19,7 +19,8 @@ class UMG_API UWidgetSwitcher : public UPanelWidget
 
 public:
 	/** The slot index to display */
-	UPROPERTY(EditAnywhere, BlueprintReadOnly, Category="Switcher", meta=( UIMin=0, ClampMin=0 ))
+	UE_DEPRECATED(5.2, "Direct access to ActiveWidgetIndex is deprecated. Please use the getter or setter.")
+	UPROPERTY(EditAnywhere, BlueprintReadWrite, Getter, Setter, BlueprintSetter = "SetActiveWidgetIndex", BlueprintGetter = "GetActiveWidgetIndex", FieldNotify, Category = "Switcher", meta = (UIMin = 0, ClampMin = 0))
 	int32 ActiveWidgetIndex;
 
 public:
