UMG - You can now style the scrollbar for the comboboxes.
UMG - Adding a DefaultStyleCache class to centralize the places we're storing static copies of the default styles.  If we need to late initialize them, we may as well do it all in one place and avoid the duplicants.

diff --git a/Engine/Source/Runtime/Slate/Public/Widgets/Input/SComboBox.h b/Engine/Source/Runtime/Slate/Public/Widgets/Input/SComboBox.h
--- a/Engine/Source/Runtime/Slate/Public/Widgets/Input/SComboBox.h
+++ b/Engine/Source/Runtime/Slate/Public/Widgets/Input/SComboBox.h
@@ -111,6 +111,7 @@ public:
 		, _ComboBoxStyle(&FAppStyle::Get().GetWidgetStyle< FComboBoxStyle >("ComboBox"))
 		, _ButtonStyle(nullptr)
 		, _ItemStyle(&FAppStyle::Get().GetWidgetStyle< FTableRowStyle >("ComboBox.Row"))
+		, _ScrollBarStyle(&FAppStyle::Get().GetWidgetStyle<FScrollBarStyle>("ScrollBar"))
 		, _ContentPadding(_ComboBoxStyle->ContentPadding)
 		, _ForegroundColor(FSlateColor::UseStyle())
 		, _OptionsSource()
@@ -133,6 +134,8 @@ public:
 		SLATE_STYLE_ARGUMENT( FButtonStyle, ButtonStyle )
 
 		SLATE_STYLE_ARGUMENT(FTableRowStyle, ItemStyle)
+		
+		SLATE_STYLE_ARGUMENT( FScrollBarStyle, ScrollBarStyle )
 
 		SLATE_ATTRIBUTE( FMargin, ContentPadding )
 		SLATE_ATTRIBUTE( FSlateColor, ForegroundColor )
@@ -218,7 +221,9 @@ public:
 				.OnSelectionChanged(this, &SComboBox< OptionType >::OnSelectionChanged_Internal)
 				.OnKeyDownHandler(this, &SComboBox< OptionType >::OnKeyDownHandler)
 				.SelectionMode(ESelectionMode::Single)
+				.ScrollBarStyle(InArgs._ScrollBarStyle)
 				.ExternalScrollbar(InArgs._CustomScrollbar)
+				
 			];
 
 		// Set up content
diff --git a/Engine/Source/Runtime/UMG/Private/Components/ComboBox.cpp b/Engine/Source/Runtime/UMG/Private/Components/ComboBox.cpp
--- a/Engine/Source/Runtime/UMG/Private/Components/ComboBox.cpp
+++ b/Engine/Source/Runtime/UMG/Private/Components/ComboBox.cpp
@@ -1,6 +1,7 @@
 // Copyright Epic Games, Inc. All Rights Reserved.
 
 #include "Components/ComboBox.h"
+#include "DefaultStyleCache.h"
 
 #include UE_INLINE_GENERATED_CPP_BY_NAME(ComboBox)
 
@@ -9,10 +10,13 @@
 /////////////////////////////////////////////////////
 // UComboBox
 
+static FScrollBarStyle* DefaultScrollBoxBarStyle = nullptr;
+
 UComboBox::UComboBox(const FObjectInitializer& ObjectInitializer)
 	: Super(ObjectInitializer)
 {
 	bIsFocusable = true;
+	ScrollBarStyle = FDefaultStyleCache::Get().GetScrollBarStyle();
 }
 
 void UComboBox::ReleaseSlateResources(bool bReleaseChildren)
@@ -26,12 +30,14 @@ TSharedRef<SWidget> UComboBox::RebuildWidget()
 {
 	TSet<UObject*> UniqueItems(Items);
 	Items = UniqueItems.Array();
-
+	
 	MyComboBox =
 		SNew(SComboBox<UObject*>)
 		.OptionsSource(&ToRawPtrTArrayUnsafe(Items))
 		.OnGenerateWidget(BIND_UOBJECT_DELEGATE(SComboBox<UObject*>::FOnGenerateWidget, HandleGenerateWidget))
-		.IsFocusable(bIsFocusable);
+		.IsFocusable(bIsFocusable)
+		.ScrollBarStyle(&ScrollBarStyle)
+	;
 
 	return MyComboBox.ToSharedRef();
 }
@@ -42,7 +48,7 @@ TSharedRef<SWidget> UComboBox::HandleGenerateWidget(UObject* Item) const
 	if ( OnGenerateWidgetEvent.IsBound() )
 	{
 		UWidget* Widget = OnGenerateWidgetEvent.Execute(Item);
-		if ( Widget != NULL )
+		if ( Widget != nullptr )
 		{
 			return Widget->TakeWidget();
 		}
diff --git a/Engine/Source/Runtime/UMG/Private/Components/ComboBoxKey.cpp b/Engine/Source/Runtime/UMG/Private/Components/ComboBoxKey.cpp
--- a/Engine/Source/Runtime/UMG/Private/Components/ComboBoxKey.cpp
+++ b/Engine/Source/Runtime/UMG/Private/Components/ComboBoxKey.cpp
@@ -2,6 +2,7 @@
 
 #include "Components/ComboBoxKey.h"
 
+#include "DefaultStyleCache.h"
 #include "Styling/UMGCoreStyle.h"
 #include "UObject/ConstructorHelpers.h"
 #include "Widgets/Layout/SBox.h"
@@ -21,9 +22,9 @@ UComboBoxKey::UComboBoxKey()
 	if (!IsRunningDedicatedServer())
 	{
 #if WITH_EDITOR 
-		const ISlateStyle& SlateStyle = IsEditorWidget() ? FCoreStyle::Get() : FUMGCoreStyle::Get();
-		WidgetStyle = SlateStyle.GetWidgetStyle<FComboBoxStyle>(IsEditorWidget() ? "EditorUtilityComboBox" : "ComboBox");
-		ItemStyle = SlateStyle.GetWidgetStyle<FTableRowStyle>("TableView.Row");
+		WidgetStyle = IsEditorWidget() ? FDefaultStyleCache::Get().GetEditorComboBoxStyle() : FDefaultStyleCache::Get().GetComboBoxStyle();
+		ItemStyle = IsEditorWidget() ? FDefaultStyleCache::Get().GetEditorComboBoxRowStyle() : FDefaultStyleCache::Get().GetComboBoxRowStyle();
+		ScrollBarStyle = IsEditorWidget() ? FDefaultStyleCache::Get().GetEditorScrollBarStyle() : FDefaultStyleCache::Get().GetScrollBarStyle();
 
 		if (IsEditorWidget())
 		{
@@ -31,8 +32,9 @@ UComboBoxKey::UComboBoxKey()
 			PostEditChange();
 		}
 #else
-		WidgetStyle = FUMGCoreStyle::Get().GetWidgetStyle<FComboBoxStyle>("ComboBox");
-		ItemStyle = FUMGCoreStyle::Get().GetWidgetStyle<FTableRowStyle>("TableView.Row");
+		WidgetStyle = FDefaultStyleCache::Get().GetComboBoxStyle();
+		ItemStyle = FDefaultStyleCache::Get().GetComboBoxRowStyle();
+		ScrollBarStyle = FDefaultStyleCache::Get().GetScrollBarStyle();
 #endif // WITH_EDITOR
 	}
 
@@ -72,6 +74,7 @@ TSharedRef<SWidget> UComboBoxKey::RebuildWidget()
 		.OnSelectionChanged_UObject(this, &UComboBoxKey::HandleSelectionChanged)
 		.OnComboBoxOpening_UObject(this, &UComboBoxKey::HandleOpening)
 		.IsFocusable(bIsFocusable)
+		.ScrollBarStyle(&ScrollBarStyle)
 		[
 			SAssignNew(ComboBoxContent, SBox)
 		];
diff --git a/Engine/Source/Runtime/UMG/Private/Components/ComboBoxString.cpp b/Engine/Source/Runtime/UMG/Private/Components/ComboBoxString.cpp
--- a/Engine/Source/Runtime/UMG/Private/Components/ComboBoxString.cpp
+++ b/Engine/Source/Runtime/UMG/Private/Components/ComboBoxString.cpp
@@ -1,6 +1,8 @@
 // Copyright Epic Games, Inc. All Rights Reserved.
 
 #include "Components/ComboBoxString.h"
+
+#include "DefaultStyleCache.h"
 #include "Widgets/SNullWidget.h"
 #include "UObject/EditorObjectVersion.h"
 #include "UObject/ConstructorHelpers.h"
@@ -14,66 +16,25 @@
 /////////////////////////////////////////////////////
 // UComboBoxString
 
-static FComboBoxStyle* DefaultComboBoxStyle = nullptr;
-static FTableRowStyle* DefaultComboBoxRowStyle = nullptr;
-
-#if WITH_EDITOR
-static FComboBoxStyle* EditorComboBoxStyle = nullptr;
-static FTableRowStyle* EditorComboBoxRowStyle = nullptr;
-#endif 
-
 UComboBoxString::UComboBoxString(const FObjectInitializer& ObjectInitializer)
 	: Super(ObjectInitializer)
 {
-	if (DefaultComboBoxStyle == nullptr)
-	{
-		DefaultComboBoxStyle = new FComboBoxStyle(FUMGCoreStyle::Get().GetWidgetStyle<FComboBoxStyle>("ComboBox"));
-
-		// Unlink UMG default colors.
-		DefaultComboBoxStyle->UnlinkColors();
-	}
-
-	if (DefaultComboBoxRowStyle == nullptr)
-	{
-		DefaultComboBoxRowStyle = new FTableRowStyle(FUMGCoreStyle::Get().GetWidgetStyle<FTableRowStyle>("TableView.Row"));
-
-		// Unlink UMG default colors.
-		DefaultComboBoxRowStyle->UnlinkColors();
-	}
-
-	WidgetStyle = *DefaultComboBoxStyle;
-	ItemStyle = *DefaultComboBoxRowStyle;
+	WidgetStyle = FDefaultStyleCache::Get().GetComboBoxStyle();
+	ItemStyle = FDefaultStyleCache::Get().GetComboBoxRowStyle();
+	ScrollBarStyle = FDefaultStyleCache::Get().GetScrollBarStyle();
 
 #if WITH_EDITOR 
-	if (EditorComboBoxStyle == nullptr)
-	{
-		EditorComboBoxStyle = new FComboBoxStyle(FCoreStyle::Get().GetWidgetStyle<FComboBoxStyle>("EditorUtilityComboBox"));
-
-		// Unlink UMG Editor colors from the editor settings colors.
-		EditorComboBoxStyle->UnlinkColors();
-	}
-
-	if (EditorComboBoxRowStyle == nullptr)
-	{
-		EditorComboBoxRowStyle = new FTableRowStyle(FCoreStyle::Get().GetWidgetStyle<FTableRowStyle>("TableView.Row"));
-
-		// Unlink UMG Editor colors from the editor settings colors.
-		EditorComboBoxRowStyle->UnlinkColors();
-	}
-
 	if (IsEditorWidget())
 	{
-		WidgetStyle = *EditorComboBoxStyle;
-		ItemStyle = *EditorComboBoxRowStyle;
+		WidgetStyle = FDefaultStyleCache::Get().GetEditorComboBoxStyle();
+		ItemStyle = FDefaultStyleCache::Get().GetEditorComboBoxRowStyle();
+		ScrollBarStyle = FDefaultStyleCache::Get().GetEditorScrollBarStyle();
 
 		// The CDO isn't an editor widget and thus won't use the editor style, call post edit change to mark difference from CDO
 		PostEditChange();
 	}
-#endif // WITH_EDITOR
-
-	WidgetStyle.UnlinkColors();
-	ItemStyle.UnlinkColors();
-
+#endif
+	
 	ForegroundColor = ItemStyle.TextColor;
 	bIsFocusable = true;
 
@@ -133,7 +94,7 @@ void UComboBoxString::PostLoad()
 
 TSharedRef<SWidget> UComboBoxString::RebuildWidget()
 {
-	int32 InitialIndex = FindOptionIndex(SelectedOption);
+	const int32 InitialIndex = FindOptionIndex(SelectedOption);
 	if ( InitialIndex != -1 )
 	{
 		CurrentOptionPtr = Options[InitialIndex];
@@ -153,6 +114,7 @@ TSharedRef<SWidget> UComboBoxString::RebuildWidget()
 		.OnGenerateWidget(BIND_UOBJECT_DELEGATE(SComboBox< TSharedPtr<FString> >::FOnGenerateWidget, HandleGenerateWidget))
 		.OnSelectionChanged(BIND_UOBJECT_DELEGATE(SComboBox< TSharedPtr<FString> >::FOnSelectionChanged, HandleSelectionChanged))
 		.OnComboBoxOpening(BIND_UOBJECT_DELEGATE(FOnComboBoxOpening, HandleOpening))
+	    .ScrollBarStyle(&ScrollBarStyle)
 		.IsFocusable(bIsFocusable)
 		[
 			SAssignNew(ComboBoxContent, SBox)
@@ -176,7 +138,7 @@ void UComboBoxString::AddOption(const FString& Option)
 
 bool UComboBoxString::RemoveOption(const FString& Option)
 {
-	int32 OptionIndex = FindOptionIndex(Option);
+	const int32 OptionIndex = FindOptionIndex(Option);
 
 	if ( OptionIndex != -1 )
 	{
@@ -256,7 +218,7 @@ void UComboBoxString::RefreshOptions()
 
 void UComboBoxString::SetSelectedOption(FString Option)
 {
-	int32 InitialIndex = FindOptionIndex(Option);
+	const int32 InitialIndex = FindOptionIndex(Option);
 	SetSelectedIndex(InitialIndex);
 }
 
@@ -335,7 +297,7 @@ void UComboBoxString::UpdateOrGenerateWidget(TSharedPtr<FString> Item)
 
 TSharedRef<SWidget> UComboBoxString::HandleGenerateWidget(TSharedPtr<FString> Item) const
 {
-	FString StringItem = Item.IsValid() ? *Item : FString();
+	const FString StringItem = Item.IsValid() ? *Item : FString();
 
 	// Call the user's delegate to see if they want to generate a custom widget bound to the data source.
 	if ( !IsDesignTime() && OnGenerateWidgetEvent.IsBound() )
@@ -387,4 +349,3 @@ const FText UComboBoxString::GetPaletteCategory()
 /////////////////////////////////////////////////////
 
 #undef LOCTEXT_NAMESPACE
-
diff --git a/Engine/Source/Runtime/UMG/Private/DefaultStyleCache.cpp b/Engine/Source/Runtime/UMG/Private/DefaultStyleCache.cpp
new file mode 100644
--- /dev/null
+++ b/Engine/Source/Runtime/UMG/Private/DefaultStyleCache.cpp
@@ -0,0 +1,43 @@
+// Copyright Epic Games, Inc. All Rights Reserved.
+
+#include "DefaultStyleCache.h"
+#include "Styling/SlateTypes.h"
+#include "Styling/UMGCoreStyle.h"
+
+#if WITH_EDITOR
+#include "Styling/CoreStyle.h"
+#endif
+
+FDefaultStyleCache* FDefaultStyleCache::Instance = nullptr;
+
+FDefaultStyleCache& FDefaultStyleCache::Get()
+{
+    if (!Instance)
+	{
+		Instance = new FDefaultStyleCache();
+	}
+	
+	return *Instance;
+}
+
+FDefaultStyleCache::FDefaultStyleCache()
+{
+	if (!IsRunningDedicatedServer())
+	{
+		Runtime.ScrollBarStyle = FUMGCoreStyle::Get().GetWidgetStyle<FScrollBarStyle>("ScrollBar");
+		Runtime.ScrollBarStyle.UnlinkColors();
+		Runtime.ComboBoxStyle = FUMGCoreStyle::Get().GetWidgetStyle<FComboBoxStyle>("ComboBox");
+		Runtime.ComboBoxStyle.UnlinkColors();
+		Runtime.ComboBoxRowStyle = FUMGCoreStyle::Get().GetWidgetStyle<FTableRowStyle>("TableView.Row");
+		Runtime.ComboBoxRowStyle.UnlinkColors();
+
+#if WITH_EDITOR
+		Editor.ScrollBarStyle = FCoreStyle::Get().GetWidgetStyle<FScrollBarStyle>("ScrollBar");
+		Editor.ScrollBarStyle.UnlinkColors();
+		Editor.ComboBoxStyle = FCoreStyle::Get().GetWidgetStyle<FComboBoxStyle>("EditorUtilityComboBox");
+		Editor.ComboBoxStyle.UnlinkColors();
+		Editor.ComboBoxRowStyle = FCoreStyle::Get().GetWidgetStyle<FTableRowStyle>("TableView.Row");
+		Editor.ComboBoxRowStyle.UnlinkColors();
+#endif
+	}
+}
diff --git a/Engine/Source/Runtime/UMG/Private/DefaultStyleCache.h b/Engine/Source/Runtime/UMG/Private/DefaultStyleCache.h
new file mode 100644
--- /dev/null
+++ b/Engine/Source/Runtime/UMG/Private/DefaultStyleCache.h
@@ -0,0 +1,39 @@
+// Copyright Epic Games, Inc. All Rights Reserved.
+
+#pragma once
+
+#include "CoreMinimal.h"
+#include "Styling/SlateTypes.h"
+
+struct FDefaultStyleCache
+{
+	static FDefaultStyleCache& Get();
+	
+	const FScrollBarStyle& GetScrollBarStyle() const { return Runtime.ScrollBarStyle; }
+	const FComboBoxStyle&  GetComboBoxStyle() const { return Runtime.ComboBoxStyle; }
+	const FTableRowStyle&  GetComboBoxRowStyle() const { return Runtime.ComboBoxRowStyle; }
+
+#if WITH_EDITOR
+	const FScrollBarStyle& GetEditorScrollBarStyle() const { return Editor.ScrollBarStyle; }
+	const FComboBoxStyle&  GetEditorComboBoxStyle() const { return Editor.ComboBoxStyle; }
+	const FTableRowStyle&  GetEditorComboBoxRowStyle() const { return Editor.ComboBoxRowStyle; }
+#endif
+	
+private:
+	FDefaultStyleCache();
+	
+    static FDefaultStyleCache* Instance;
+
+	struct FStyles
+	{
+		FScrollBarStyle ScrollBarStyle;
+		FComboBoxStyle ComboBoxStyle;
+		FTableRowStyle ComboBoxRowStyle;
+	};
+
+	FStyles Runtime;
+
+#if WITH_EDITOR
+	FStyles Editor;
+#endif
+};
\ No newline at end of file
diff --git a/Engine/Source/Runtime/UMG/Public/Components/ComboBox.h b/Engine/Source/Runtime/UMG/Public/Components/ComboBox.h
--- a/Engine/Source/Runtime/UMG/Public/Components/ComboBox.h
+++ b/Engine/Source/Runtime/UMG/Public/Components/ComboBox.h
@@ -19,6 +19,9 @@ class UMG_API UComboBox : public UWidget
 
 public:
 
+	UPROPERTY(EditAnywhere, BlueprintReadWrite, Category="Style")
+	FScrollBarStyle ScrollBarStyle;
+
 	/** The list of items to be displayed on the combobox. */
 	UPROPERTY(EditAnywhere, BlueprintReadWrite, Category=Content)
 	TArray<TObjectPtr<UObject>> Items;
diff --git a/Engine/Source/Runtime/UMG/Public/Components/ComboBoxKey.h b/Engine/Source/Runtime/UMG/Public/Components/ComboBoxKey.h
--- a/Engine/Source/Runtime/UMG/Public/Components/ComboBoxKey.h
+++ b/Engine/Source/Runtime/UMG/Public/Components/ComboBoxKey.h
@@ -44,6 +44,10 @@ public:
 	UPROPERTY(EditAnywhere, BlueprintReadOnly, Category = Style)
 	FTableRowStyle ItemStyle;
 
+	/** The scroll bar style. */
+    UPROPERTY(EditAnywhere, BlueprintReadWrite, Category="Style")
+    FScrollBarStyle ScrollBarStyle;
+
 	/** The foreground color to pass through the hierarchy. */
 	UPROPERTY(EditAnywhere, BlueprintReadOnly, Category = Style, meta = (DesignerRebuild))
 	FSlateColor ForegroundColor;
diff --git a/Engine/Source/Runtime/UMG/Public/Components/ComboBoxString.h b/Engine/Source/Runtime/UMG/Public/Components/ComboBoxString.h
--- a/Engine/Source/Runtime/UMG/Public/Components/ComboBoxString.h
+++ b/Engine/Source/Runtime/UMG/Public/Components/ComboBoxString.h
@@ -44,6 +44,10 @@ public:
 	/** The item row style. */
 	UPROPERTY(EditAnywhere, BlueprintReadWrite, Category=Style)
 	FTableRowStyle ItemStyle;
+	
+	/** The scroll bar style. */
+	UPROPERTY(EditAnywhere, BlueprintReadWrite, Category="Style")
+    FScrollBarStyle ScrollBarStyle;
 
 	UPROPERTY(EditAnywhere, BlueprintReadOnly, Category=Content)
 	FMargin ContentPadding;
