UMG - Removing a local static that isn't needed that collides in unity builds.

diff --git a/Engine/Source/Runtime/UMG/Private/Components/ComboBox.cpp b/Engine/Source/Runtime/UMG/Private/Components/ComboBox.cpp
--- a/Engine/Source/Runtime/UMG/Private/Components/ComboBox.cpp
+++ b/Engine/Source/Runtime/UMG/Private/Components/ComboBox.cpp
@@ -10,8 +10,6 @@
 /////////////////////////////////////////////////////
 // UComboBox
 
-static FScrollBarStyle* DefaultScrollBoxBarStyle = nullptr;
-
 UComboBox::UComboBox(const FObjectInitializer& ObjectInitializer)
 	: Super(ObjectInitializer)
 {
