UMG - Using TLazySingleton for the DefaultStyleCache.

diff --git a/Engine/Source/Runtime/UMG/Private/DefaultStyleCache.cpp b/Engine/Source/Runtime/UMG/Private/DefaultStyleCache.cpp
--- a/Engine/Source/Runtime/UMG/Private/DefaultStyleCache.cpp
+++ b/Engine/Source/Runtime/UMG/Private/DefaultStyleCache.cpp
@@ -1,6 +1,8 @@
 // Copyright Epic Games, Inc. All Rights Reserved.
 
 #include "DefaultStyleCache.h"
+
+#include "Misc/LazySingleton.h"
 #include "Styling/SlateTypes.h"
 #include "Styling/UMGCoreStyle.h"
 
@@ -8,16 +10,9 @@
 #include "Styling/CoreStyle.h"
 #endif
 
-FDefaultStyleCache* FDefaultStyleCache::Instance = nullptr;
-
 FDefaultStyleCache& FDefaultStyleCache::Get()
 {
-    if (!Instance)
-	{
-		Instance = new FDefaultStyleCache();
-	}
-	
-	return *Instance;
+	return TLazySingleton<FDefaultStyleCache>::Get();
 }
 
 FDefaultStyleCache::FDefaultStyleCache()
diff --git a/Engine/Source/Runtime/UMG/Private/DefaultStyleCache.h b/Engine/Source/Runtime/UMG/Private/DefaultStyleCache.h
--- a/Engine/Source/Runtime/UMG/Private/DefaultStyleCache.h
+++ b/Engine/Source/Runtime/UMG/Private/DefaultStyleCache.h
@@ -20,9 +20,9 @@ struct FDefaultStyleCache
 #endif
 	
 private:
+	friend class FLazySingleton;
+
 	FDefaultStyleCache();
-	
-    static FDefaultStyleCache* Instance;
 
 	struct FStyles
 	{
