UMG: Added the ability to emit frame markers when animations are started and stopped

CVar to control this is UMG.AnimationMarkers 1/0

diff --git a/Engine/Source/Runtime/UMG/Private/Animation/UMGSequencePlayer.cpp b/Engine/Source/Runtime/UMG/Private/Animation/UMGSequencePlayer.cpp
--- a/Engine/Source/Runtime/UMG/Private/Animation/UMGSequencePlayer.cpp
+++ b/Engine/Source/Runtime/UMG/Private/Animation/UMGSequencePlayer.cpp
@@ -9,6 +9,7 @@
 #include "EntitySystem/MovieSceneEntitySystemRunner.h"
 #include "EntitySystem/MovieSceneEntitySystemLinker.h"
 #include "Animation/UMGSequenceTickManager.h"
+#include "ProfilingDebugging/CsvProfiler.h"
 
 #include UE_INLINE_GENERATED_CPP_BY_NAME(UMGSequencePlayer)
 
@@ -24,6 +25,13 @@ namespace UE::UMG
 		TEXT("(Default: true) Whether to perform animation control flow functions (Play, Pause, Stop etc) asynchronously.")
 	);
 
+	bool GVarAnimationMarkers = false;
+	FAutoConsoleVariableRef CVarAnimationMarkers(
+		TEXT("UMG.AnimationMarkers"),
+		GVarAnimationMarkers,
+		TEXT("(Default: false) Whether to emit profiling frame markers for starting and stopping UMG animations.")
+	);
+
 } // namespace UE::UMG
 
 UUMGSequencePlayer::UUMGSequencePlayer(const FObjectInitializer& ObjectInitializer)
@@ -220,6 +228,11 @@ void UUMGSequencePlayer::PlayInternal(double StartAtTime, double EndAtTime, int3
 	UUserWidget* Widget = UserWidget.Get();
 	UUMGSequenceTickManager* TickManager = Widget ? ToRawPtr(Widget->AnimationTickManager) : nullptr;
 
+	if (UE::UMG::GVarAnimationMarkers && Animation && Widget)
+	{
+		CSV_EVENT_GLOBAL(TEXT("Play Animation [%s::%s]"), *Widget->GetName(), *Animation->GetName());
+	}
+
 	TSharedPtr<FMovieSceneEntitySystemRunner> RunnerToUse = TickManager ? TickManager->GetRunner() : nullptr;
 	if (EnumHasAnyFlags(Animation->GetFlags(), EMovieSceneSequenceFlags::BlockingEvaluation))
 	{
@@ -431,9 +444,15 @@ void UUMGSequencePlayer::HandleLatentStop()
 
 	bIsStopping = false;
 
-	if (UserWidget.IsValid())
+	UUserWidget* Widget = UserWidget.Get();
+	if (Widget)
 	{
-		UserWidget->OnAnimationFinishedPlaying(*this);
+		if (UE::UMG::GVarAnimationMarkers && Animation)
+		{
+			CSV_EVENT_GLOBAL(TEXT("Stop Animation [%s::%s]"), *Widget->GetName(), *Animation->GetName());
+		}
+
+		Widget->OnAnimationFinishedPlaying(*this);
 	}
 
 	OnSequenceFinishedPlayingEvent.Broadcast(*this);
