UMG: Do not keep tracking widget tick state for widgets that are no longer ticking

diff --git a/Engine/Source/Runtime/UMG/Private/Animation/UMGSequenceTickManager.cpp b/Engine/Source/Runtime/UMG/Private/Animation/UMGSequenceTickManager.cpp
--- a/Engine/Source/Runtime/UMG/Private/Animation/UMGSequenceTickManager.cpp
+++ b/Engine/Source/Runtime/UMG/Private/Animation/UMGSequenceTickManager.cpp
@@ -71,6 +71,10 @@ void UUMGSequenceTickManager::OnWidgetTicked(UUserWidget* InWidget)
 	{
 		WidgetData->bIsTicking = true;
 	}
+	else
+	{
+		WeakUserWidgetData.Add(InWidget, FSequenceTickManagerWidgetData());
+	}
 }
 
 void UUMGSequenceTickManager::BeginDestroy()
@@ -160,9 +164,13 @@ void UUMGSequenceTickManager::TickWidgetAnimations(float DeltaSeconds)
 					if (WidgetData.bLastKnownTickState)
 					{
 						UserWidget->DisableAnimations();
+						WidgetData.bLastKnownTickState = false;
+					}
+					else
+					{
+						WidgetIter.RemoveCurrent();
+						continue;
 					}
-
-					WidgetData.bLastKnownTickState = false;
 				}
 			}
 			else
