Add accessor for Retainerbox and Scrollbar and ScaleBox

diff --git a/Engine/Plugins/Experimental/CommonUI/Source/CommonUI/Private/Input/CommonAnalogCursor.cpp b/Engine/Plugins/Experimental/CommonUI/Source/CommonUI/Private/Input/CommonAnalogCursor.cpp
--- a/Engine/Plugins/Experimental/CommonUI/Source/CommonUI/Private/Input/CommonAnalogCursor.cpp
+++ b/Engine/Plugins/Experimental/CommonUI/Source/CommonUI/Private/Input/CommonAnalogCursor.cpp
@@ -436,7 +436,7 @@ EOrientation FCommonAnalogCursor::DetermineScrollOrientation(const UWidget& Widg
 	}
 	else if (const UScrollBar* AsScrollBar = Cast<const UScrollBar>(&Widget))
 	{
-		return AsScrollBar->Orientation;
+		return AsScrollBar->GetOrientation();
 	}
 	else if (const UScrollBox* AsScrollBox = Cast<const UScrollBox>(&Widget))
 	{
diff --git a/Engine/Source/Runtime/UMG/Private/Components/RetainerBox.cpp b/Engine/Source/Runtime/UMG/Private/Components/RetainerBox.cpp
--- a/Engine/Source/Runtime/UMG/Private/Components/RetainerBox.cpp
+++ b/Engine/Source/Runtime/UMG/Private/Components/RetainerBox.cpp
@@ -19,13 +19,16 @@ URetainerBox::URetainerBox(const FObjectInitializer& ObjectInitializer)
 	: Super(ObjectInitializer)
 {
 	SetVisibilityInternal(ESlateVisibility::Visible);
+PRAGMA_DISABLE_DEPRECATION_WARNINGS
 	Phase = 0;
 	PhaseCount = 1;
 	RenderOnPhase = true;
 	RenderOnInvalidation = false;
 	TextureParameter = DefaultTextureParameterName;
+PRAGMA_ENABLE_DEPRECATION_WARNINGS
 }
 
+PRAGMA_DISABLE_DEPRECATION_WARNINGS
 void URetainerBox::SetRenderingPhase(int32 PhaseToRenderOn, int32 TotalRenderingPhases)
 {
 	Phase = PhaseToRenderOn;
@@ -41,6 +44,7 @@ void URetainerBox::SetRenderingPhase(int32 PhaseToRenderOn, int32 TotalRendering
 		MyRetainerWidget->SetRenderingPhase(Phase, PhaseCount);
 	}
 }
+PRAGMA_ENABLE_DEPRECATION_WARNINGS
 
 void URetainerBox::RequestRender()
 {
@@ -50,6 +54,13 @@ void URetainerBox::RequestRender()
 	}
 }
 
+PRAGMA_DISABLE_DEPRECATION_WARNINGS
+const UMaterialInterface* URetainerBox::GetEffectMaterialInterface() const
+{
+	return EffectMaterial;
+}
+PRAGMA_ENABLE_DEPRECATION_WARNINGS
+
 UMaterialInstanceDynamic* URetainerBox::GetEffectMaterial() const
 {
 	if ( MyRetainerWidget.IsValid() )
@@ -60,6 +71,7 @@ UMaterialInstanceDynamic* URetainerBox::GetEffectMaterial() const
 	return nullptr;
 }
 
+PRAGMA_DISABLE_DEPRECATION_WARNINGS
 void URetainerBox::SetEffectMaterial(UMaterialInterface* InEffectMaterial)
 {
 	EffectMaterial = InEffectMaterial;
@@ -78,6 +90,11 @@ void URetainerBox::SetTextureParameter(FName InTextureParameter)
 	}
 }
 
+const FName& URetainerBox::GetTextureParameter() const
+{
+	return TextureParameter;
+}
+
 void URetainerBox::SetRetainRendering(bool bInRetainRendering)
 {
 	bRetainRender = bInRetainRendering;
@@ -88,6 +105,32 @@ void URetainerBox::SetRetainRendering(bool bInRetainRendering)
 	}
 }
 
+bool URetainerBox::IsRetainRendering() const
+{
+	return bRetainRender;
+}
+
+int32 URetainerBox::GetPhase() const
+{
+	return Phase;
+}
+
+int32 URetainerBox::GetPhaseCount() const
+{
+	return PhaseCount;
+}
+
+bool URetainerBox::IsRenderOnInvalidation() const
+{
+	return RenderOnInvalidation;
+}
+
+bool URetainerBox::IsRenderOnPhase() const
+{
+	return RenderOnPhase;
+}
+PRAGMA_ENABLE_DEPRECATION_WARNINGS
+
 void URetainerBox::ReleaseSlateResources(bool bReleaseChildren)
 {
 	Super::ReleaseSlateResources(bReleaseChildren);
@@ -97,6 +140,7 @@ void URetainerBox::ReleaseSlateResources(bool bReleaseChildren)
 
 TSharedRef<SWidget> URetainerBox::RebuildWidget()
 {
+PRAGMA_DISABLE_DEPRECATION_WARNINGS
 	MyRetainerWidget =
 		SNew(SRetainerWidget)
 		.RenderOnInvalidation(RenderOnInvalidation)
@@ -107,7 +151,7 @@ TSharedRef<SWidget> URetainerBox::RebuildWidget()
 		.StatId( FName( *FString::Printf(TEXT("%s [%s]"), *GetFName().ToString(), *GetClass()->GetName() ) ) )
 #endif//STATS
 	;
-
+PRAGMA_ENABLE_DEPRECATION_WARNINGS
 	if ( GetChildrenCount() > 0 )
 	{
 		MyRetainerWidget->SetContent(GetContentSlot()->Content ? GetContentSlot()->Content->TakeWidget() : SNullWidget::NullWidget);
@@ -120,10 +164,12 @@ void URetainerBox::SynchronizeProperties()
 {
 	Super::SynchronizeProperties();
 
+PRAGMA_DISABLE_DEPRECATION_WARNINGS
 	MyRetainerWidget->SetRetainedRendering(IsDesignTime() ? false : bRetainRender);
 	MyRetainerWidget->SetEffectMaterial(EffectMaterial);
 	MyRetainerWidget->SetTextureParameter(TextureParameter);
 	MyRetainerWidget->SetWorld(GetWorld());
+PRAGMA_ENABLE_DEPRECATION_WARNINGS
 }
 
 void URetainerBox::OnSlotAdded(UPanelSlot* InSlot)
@@ -172,10 +218,12 @@ bool URetainerBox::CanEditChange(const FProperty* InProperty) const
 		return false;
 	}
 
+PRAGMA_DISABLE_DEPRECATION_WARNINGS
 	if (InProperty->GetFName() == GET_MEMBER_NAME_CHECKED(URetainerBox, Phase)
 		|| InProperty->GetFName() == GET_MEMBER_NAME_CHECKED(URetainerBox, PhaseCount))
 	{
 		return RenderOnPhase && bRetainRender;
+PRAGMA_ENABLE_DEPRECATION_WARNINGS
 	}
 	return true;
 }
diff --git a/Engine/Source/Runtime/UMG/Private/Components/ScaleBox.cpp b/Engine/Source/Runtime/UMG/Private/Components/ScaleBox.cpp
--- a/Engine/Source/Runtime/UMG/Private/Components/ScaleBox.cpp
+++ b/Engine/Source/Runtime/UMG/Private/Components/ScaleBox.cpp
@@ -17,10 +17,12 @@ UScaleBox::UScaleBox(const FObjectInitializer& ObjectInitializer)
 	bIsVariable = false;
 	SetVisibilityInternal(ESlateVisibility::SelfHitTestInvisible);
 
+PRAGMA_DISABLE_DEPRECATION_WARNINGS
 	StretchDirection = EStretchDirection::Both;
 	Stretch = EStretch::ScaleToFit;
 	UserSpecifiedScale = 1.0f;
 	IgnoreInheritedScale = false;
+PRAGMA_ENABLE_DEPRECATION_WARNINGS
 }
 
 void UScaleBox::ReleaseSlateResources(bool bReleaseChildren)
@@ -46,6 +48,7 @@ TSharedRef<SWidget> UScaleBox::RebuildWidget()
 	return MyScaleBox.ToSharedRef();
 }
 
+PRAGMA_DISABLE_DEPRECATION_WARNINGS
 void UScaleBox::SetStretch(EStretch::Type InStretch)
 {
 	Stretch = InStretch;
@@ -55,6 +58,11 @@ void UScaleBox::SetStretch(EStretch::Type InStretch)
 	}
 }
 
+EStretch::Type UScaleBox::GetStretch() const
+{
+	return Stretch;
+}
+
 void UScaleBox::SetStretchDirection(EStretchDirection::Type InStretchDirection)
 {
 	StretchDirection = InStretchDirection;
@@ -64,6 +72,11 @@ void UScaleBox::SetStretchDirection(EStretchDirection::Type InStretchDirection)
 	}
 }
 
+EStretchDirection::Type UScaleBox::GetStretchDirection() const
+{
+	return StretchDirection;
+}
+
 void UScaleBox::SetUserSpecifiedScale(float InUserSpecifiedScale)
 {
 	UserSpecifiedScale = InUserSpecifiedScale;
@@ -73,6 +86,11 @@ void UScaleBox::SetUserSpecifiedScale(float InUserSpecifiedScale)
 	}
 }
 
+float UScaleBox::GetUserSpecifiedScale() const
+{
+	return UserSpecifiedScale;
+}
+
 void UScaleBox::SetIgnoreInheritedScale(bool bInIgnoreInheritedScale)
 {
 	IgnoreInheritedScale = bInIgnoreInheritedScale;
@@ -82,14 +100,22 @@ void UScaleBox::SetIgnoreInheritedScale(bool bInIgnoreInheritedScale)
 	}
 }
 
+bool UScaleBox::IsIgnoreInheritedScale() const
+{
+	return IgnoreInheritedScale;
+}
+PRAGMA_ENABLE_DEPRECATION_WARNINGS
+
 void UScaleBox::SynchronizeProperties()
 {
 	Super::SynchronizeProperties();
 
+PRAGMA_DISABLE_DEPRECATION_WARNINGS
 	MyScaleBox->SetStretchDirection(StretchDirection);
 	MyScaleBox->SetStretch(Stretch);
 	MyScaleBox->SetUserSpecifiedScale(UserSpecifiedScale);
 	MyScaleBox->SetIgnoreInheritedScale(IgnoreInheritedScale);
+PRAGMA_ENABLE_DEPRECATION_WARNINGS
 }
 
 UClass* UScaleBox::GetSlotClass() const
@@ -135,7 +161,7 @@ bool UScaleBox::CanEditChange(const FProperty* InProperty) const
 	if (bIsEditable && InProperty)
 	{
 		const FName PropertyName = InProperty->GetFName();
-
+		PRAGMA_DISABLE_DEPRECATION_WARNINGS
 		if (PropertyName == GET_MEMBER_NAME_CHECKED(UScaleBox, StretchDirection))
 		{
 			return Stretch != EStretch::None && Stretch != EStretch::ScaleBySafeZone &&
@@ -146,6 +172,7 @@ bool UScaleBox::CanEditChange(const FProperty* InProperty) const
 		{
 			return Stretch == EStretch::UserSpecified || Stretch == EStretch::UserSpecifiedWithClipping;
 		}
+		PRAGMA_ENABLE_DEPRECATION_WARNINGS
 	}
 
 	return bIsEditable;
diff --git a/Engine/Source/Runtime/UMG/Private/Components/ScrollBar.cpp b/Engine/Source/Runtime/UMG/Private/Components/ScrollBar.cpp
--- a/Engine/Source/Runtime/UMG/Private/Components/ScrollBar.cpp
+++ b/Engine/Source/Runtime/UMG/Private/Components/ScrollBar.cpp
@@ -25,12 +25,13 @@ UScrollBar::UScrollBar(const FObjectInitializer& ObjectInitializer)
 {
 	bIsVariable = false;
 
+PRAGMA_DISABLE_DEPRECATION_WARNINGS
 	bAlwaysShowScrollbar = true;
 	bAlwaysShowScrollbarTrack = true;
 	Orientation = Orient_Vertical;
 	Thickness = FVector2D(16.0f, 16.0f);
 	Padding = FMargin(2.0f);
-
+PRAGMA_ENABLE_DEPRECATION_WARNINGS
 	if (DefaultScrollBarStyle == nullptr)
 	{
 		DefaultScrollBarStyle = new FScrollBarStyle(FUMGCoreStyle::Get().GetWidgetStyle<FScrollBarStyle>("Scrollbar"));
@@ -38,9 +39,9 @@ UScrollBar::UScrollBar(const FObjectInitializer& ObjectInitializer)
 		// Unlink UMG default colors.
 		DefaultScrollBarStyle->UnlinkColors();
 	}
-	
+PRAGMA_DISABLE_DEPRECATION_WARNINGS
 	WidgetStyle = *DefaultScrollBarStyle;
-
+PRAGMA_ENABLE_DEPRECATION_WARNINGS
 #if WITH_EDITOR 
 	if (EditorScrollBarStyle == nullptr)
 	{
@@ -52,8 +53,9 @@ UScrollBar::UScrollBar(const FObjectInitializer& ObjectInitializer)
 
 	if (IsEditorWidget())
 	{
+PRAGMA_DISABLE_DEPRECATION_WARNINGS
 		WidgetStyle = *EditorScrollBarStyle;
-
+PRAGMA_ENABLE_DEPRECATION_WARNINGS
 		// The CDO isn't an editor widget and thus won't use the editor style, call post edit change to mark difference from CDO
 		PostEditChange();
 	}
@@ -69,6 +71,7 @@ void UScrollBar::ReleaseSlateResources(bool bReleaseChildren)
 
 TSharedRef<SWidget> UScrollBar::RebuildWidget()
 {
+PRAGMA_DISABLE_DEPRECATION_WARNINGS
 	MyScrollBar = SNew(SScrollBar)
 		.Style(&WidgetStyle)
 		.AlwaysShowScrollbar(bAlwaysShowScrollbar)
@@ -76,7 +79,7 @@ TSharedRef<SWidget> UScrollBar::RebuildWidget()
 		.Orientation(Orientation)
 		.Thickness(Thickness)
 		.Padding(Padding);
-
+PRAGMA_ENABLE_DEPRECATION_WARNINGS
 	//SLATE_EVENT(FOnUserScrolled, OnUserScrolled)
 
 	return MyScrollBar.ToSharedRef();
@@ -106,21 +109,103 @@ void UScrollBar::Serialize(FArchive& Ar)
 	const bool bDeprecateThickness = Ar.IsLoading() && Ar.CustomVer(FEditorObjectVersion::GUID) < FEditorObjectVersion::ScrollBarThicknessChange;
 	if (bDeprecateThickness)
 	{
+PRAGMA_DISABLE_DEPRECATION_WARNINGS
 		// Set Thickness property to previous default value.
 		Thickness.Set(12.0f, 12.0f);
+PRAGMA_ENABLE_DEPRECATION_WARNINGS
 	}
 
 	Super::Serialize(Ar);
 
 	if (bDeprecateThickness)
 	{
+PRAGMA_DISABLE_DEPRECATION_WARNINGS
 		// Implicit padding of 2 was removed, so Thickness value must be incremented by 4.
 		Thickness += FVector2D(4.0f, 4.0f);
+PRAGMA_ENABLE_DEPRECATION_WARNINGS
 	}
 }
 
 #endif // if WITH_EDITORONLY_DATA
 
+PRAGMA_DISABLE_DEPRECATION_WARNINGS
+
+const FScrollBarStyle& UScrollBar::GetWidgetStyle() const
+{
+	return WidgetStyle;
+}
+void UScrollBar::SetWidgetStyle(const FScrollBarStyle& InWidgetStyle)
+{
+	WidgetStyle = InWidgetStyle;
+	if (MyScrollBar.IsValid())
+	{
+		MyScrollBar->SetStyle(&InWidgetStyle);
+	}
+}
+
+bool UScrollBar::IsAlwaysShowScrollbar() const
+{
+	return bAlwaysShowScrollbar;
+}
+
+void UScrollBar::SetAlwaysShowScrollbar(bool bNewValue)
+{
+	bAlwaysShowScrollbar = bNewValue;
+	if (MyScrollBar.IsValid())
+	{
+		MyScrollBar->SetScrollBarAlwaysVisible(bAlwaysShowScrollbar);
+	}
+}
+
+bool UScrollBar::IsAlwaysShowScrollbarTrack() const
+{
+	return bAlwaysShowScrollbarTrack;
+}
+
+void UScrollBar::SetAlwaysShowScrollbarTrack(bool bNewValue)
+{
+	bAlwaysShowScrollbarTrack = bNewValue;
+	if (MyScrollBar.IsValid())
+	{
+		MyScrollBar->SetScrollBarTrackAlwaysVisible(bAlwaysShowScrollbarTrack);
+	}
+}
+
+EOrientation UScrollBar::GetOrientation() const
+{
+	return Orientation;
+}
+
+FVector2D UScrollBar::GetThickness() const
+{
+	return Thickness;
+}
+
+void UScrollBar::SetThickness(const FVector2D& InThickness)
+{
+	Thickness = InThickness;
+	if (MyScrollBar.IsValid())
+	{
+		MyScrollBar->SetThickness(Thickness);
+	}
+}
+
+FMargin UScrollBar::GetPadding() const
+{
+	return Padding;
+}
+
+void UScrollBar::SetPadding(const FMargin& InPadding)
+{
+	Padding = InPadding;
+	if (MyScrollBar.IsValid())
+	{
+		MyScrollBar->SetPadding(InPadding);
+	}
+}
+
+PRAGMA_ENABLE_DEPRECATION_WARNINGS
+
 #if WITH_EDITOR
 
 const FText UScrollBar::GetPaletteCategory()
diff --git a/Engine/Source/Runtime/UMG/Public/Components/RetainerBox.h b/Engine/Source/Runtime/UMG/Public/Components/RetainerBox.h
--- a/Engine/Source/Runtime/UMG/Public/Components/RetainerBox.h
+++ b/Engine/Source/Runtime/UMG/Public/Components/RetainerBox.h
@@ -28,7 +28,8 @@ class UMG_API URetainerBox : public UContentWidget
 	GENERATED_UCLASS_BODY()
 
 protected:
-	UPROPERTY(EditAnywhere, Category="Render Rules")
+	UE_DEPRECATED(5.2, "Direct access to bRetainRender is deprecated. Please use the getter or setter.")
+	UPROPERTY(EditAnywhere, BlueprintReadWrite, Setter = "SetRetainRendering", Getter = "IsRetainRendering", BlueprintSetter = "SetRetainRendering", Category = "Render Rules")
 	bool bRetainRender = true;
 
 public:
@@ -36,13 +37,15 @@ public:
 	 * Should this widget redraw the contents it has every time it receives an invalidation request
 	 * from it's children, similar to the invalidation panel.
 	 */
-	UPROPERTY(EditAnywhere, BlueprintReadOnly, Category="Render Rules", meta=(EditCondition=bRetainRender))
+	UE_DEPRECATED(5.2, "Direct access to RenderOnInvalidation is deprecated. Please use the getter.")
+	UPROPERTY(EditAnywhere, BlueprintReadOnly, Getter = "IsRenderOnInvalidation", Category = "Render Rules", meta = (EditCondition = bRetainRender))
 	bool RenderOnInvalidation;
 
 	/**
 	 * Should this widget redraw the contents it has every time the phase occurs.
 	 */
-	UPROPERTY(EditAnywhere, BlueprintReadOnly, Category="Render Rules", meta=(EditCondition=bRetainRender))
+	UE_DEPRECATED(5.2, "Direct access to RenderOnPhase is deprecated. Please use the getter.")
+	UPROPERTY(EditAnywhere, BlueprintReadOnly, Getter = "IsRenderOnPhase", Category = "Render Rules", meta = (EditCondition = bRetainRender))
 	bool RenderOnPhase;
 
 	/**
@@ -52,7 +55,8 @@ public:
 	 * If the Phase were 0, and the PhaseCount were 2, this retainer would draw a fresh frame every
 	 * other frame.  So in a 60Hz game, the UI would render at 30Hz.
 	 */
-	UPROPERTY(EditAnywhere, BlueprintReadOnly, Category="Render Rules", meta=(UIMin=0, ClampMin=0))
+	UE_DEPRECATED(5.2, "Direct access to Phase is deprecated. Please use the getter.")
+	UPROPERTY(EditAnywhere, BlueprintReadOnly, Getter, Category="Render Rules", meta=(UIMin=0, ClampMin=0))
 	int32 Phase;
 
 	/**
@@ -63,7 +67,8 @@ public:
 	 * If the Phase were 0, and the PhaseCount were 2, this retainer would draw a fresh frame every 
 	 * other frame.  So in a 60Hz game, the UI would render at 30Hz.
 	 */
-	UPROPERTY(EditAnywhere, BlueprintReadOnly, Category="Render Rules", meta=(UIMin=1, ClampMin=1))
+	UE_DEPRECATED(5.2, "Direct access to PhaseCount is deprecated. Please use the getter.")
+	UPROPERTY(EditAnywhere, BlueprintReadOnly, Getter, Category="Render Rules", meta=(UIMin=1, ClampMin=1))
 	int32 PhaseCount;
 
 public:
@@ -80,6 +85,11 @@ public:
 	UFUNCTION(BlueprintCallable, Category="Retainer")
 	void RequestRender();
 
+	/**
+	 * Get the current dynamic effect material applied to the retainer box.
+	 */
+	const UMaterialInterface* GetEffectMaterialInterface() const;
+
 	/**
 	 * Get the current dynamic effect material applied to the retainer box.
 	 */
@@ -97,12 +107,43 @@ public:
 	 */
 	UFUNCTION(BlueprintCallable, Category="Retainer|Effect")
 	void SetTextureParameter(FName TextureParameter);
+
+	/**
+	 * Gets the name of the texture parameter to set the render target to on the material.
+	 */
+	const FName& GetTextureParameter() const;
+
 	/**
 	* Set the flag for if we retain the render or pass-through
 	*/
 	UFUNCTION(BlueprintCallable, Category = "Retainer")
 	void SetRetainRendering(bool bInRetainRendering);
 
+	/**
+	 * Get the flag for if we retain the render or pass-through.
+	 */
+	bool IsRetainRendering() const;
+
+	/**
+	 * Get the phase to render on.
+	 */
+	int32 GetPhase() const;
+
+	/**
+	 * Get the total number of phases.
+	 */
+	int32 GetPhaseCount() const;
+
+	/**
+	 * Get whether this widget should redraw the contents it has every time it receives an invalidation request.
+	 */
+	bool IsRenderOnInvalidation() const;
+
+	/**
+	 * Get whether this widget should redraw the contents it has every time the phase occurs.
+	 */
+	bool IsRenderOnPhase() const;
+
 	virtual void ReleaseSlateResources(bool bReleaseChildren) override;
 
 #if WITH_EDITOR
@@ -121,13 +162,15 @@ protected:
 	 * and make sure to multiply the alpha you're apply across the surface to the color and the alpha of the render target, otherwise
 	 * you won't see the expected color.
 	 */
-	UPROPERTY(EditAnywhere, BlueprintReadOnly, Category="Effect")
+	UE_DEPRECATED(5.2, "Direct access to EffectMaterial is deprecated. Please use the getter or setter.")
+	UPROPERTY(EditAnywhere, BlueprintReadWrite, Getter = "GetEffectMaterialInterface", Setter, BlueprintSetter = "SetEffectMaterial", Category = "Effect")
 	TObjectPtr<UMaterialInterface> EffectMaterial;
 
 	/**
 	 * The texture sampler parameter of the @EffectMaterial, that we'll set to the render target.
 	 */
-	UPROPERTY(EditAnywhere, BlueprintReadOnly, Category="Effect")
+	UE_DEPRECATED(5.2, "Direct access to TextureParameter is deprecated. Please use the getter or setter.")
+	UPROPERTY(EditAnywhere, BlueprintReadWrite, Getter, Setter, BlueprintSetter = "SetTextureParameter", Category="Effect")
 	FName TextureParameter;
 
 	//~ Begin UPanelWidget interface
diff --git a/Engine/Source/Runtime/UMG/Public/Components/ScaleBox.h b/Engine/Source/Runtime/UMG/Public/Components/ScaleBox.h
--- a/Engine/Source/Runtime/UMG/Public/Components/ScaleBox.h
+++ b/Engine/Source/Runtime/UMG/Public/Components/ScaleBox.h
@@ -25,33 +25,46 @@ class UMG_API UScaleBox : public UContentWidget
 public:
 
 	/** The stretching rule to apply when content is stretched */
-	UPROPERTY(EditAnywhere, BlueprintReadOnly, Category="Stretching")
+	UE_DEPRECATED(5.2, "Direct access to Stretch is deprecated. Please use the getter or setter.")
+	UPROPERTY(EditAnywhere, BlueprintReadWrite, Setter, Getter, BlueprintSetter = "SetStretch", Category = "Stretching")
 	TEnumAsByte<EStretch::Type> Stretch;
 
 	/** Controls in what direction content can be scaled */
-	UPROPERTY(EditAnywhere, BlueprintReadOnly, Category="Stretching")
+	UE_DEPRECATED(5.2, "Direct access to StretchDirection is deprecated. Please use the getter or setter.")
+	UPROPERTY(EditAnywhere, BlueprintReadWrite, Setter, Getter, BlueprintSetter = "SetStretchDirection", Category = "Stretching")
 	TEnumAsByte<EStretchDirection::Type> StretchDirection;
 
 	/** Optional scale that can be specified by the User. Used only for UserSpecified stretching. */
-	UPROPERTY(EditAnywhere, BlueprintReadOnly, Category="Stretching")
+	UE_DEPRECATED(5.2, "Direct access to UserSpecifiedScale is deprecated. Please use the getter or setter.")
+	UPROPERTY(EditAnywhere, BlueprintReadWrite, Setter, Getter, BlueprintSetter = "SetUserSpecifiedScale", Category = "Stretching")
 	float UserSpecifiedScale;
 
 	/** Optional bool to ignore the inherited scale. Applies inverse scaling to counteract parents before applying the local scale operation. */
-	UPROPERTY(EditAnywhere, BlueprintReadOnly, Category = "Stretching")
+	UE_DEPRECATED(5.2, "Direct access to IgnoreInheritedScale is deprecated. Please use the getter or setter.")
+	UPROPERTY(EditAnywhere, BlueprintReadWrite, Setter, Getter = "IsIgnoreInheritedScale", BlueprintSetter = "SetIgnoreInheritedScale", Category = "Stretching")
 	bool IgnoreInheritedScale;
 
 public:
 	UFUNCTION(BlueprintCallable, Category = "Appearance")
 	void SetStretch(EStretch::Type InStretch);
 
+	EStretch::Type GetStretch() const;
+
 	UFUNCTION(BlueprintCallable, Category = "Appearance")
 	void SetStretchDirection(EStretchDirection::Type InStretchDirection);
 
+	EStretchDirection::Type GetStretchDirection() const;
+
 	UFUNCTION(BlueprintCallable, Category = "Appearance")
 	void SetUserSpecifiedScale(float InUserSpecifiedScale);
 
+	float GetUserSpecifiedScale() const;
+
 	UFUNCTION(BlueprintCallable, Category = "Appearance")
 	void SetIgnoreInheritedScale(bool bInIgnoreInheritedScale);
+
+	bool IsIgnoreInheritedScale() const;
+
 public:
 
 	// UWidget interface
diff --git a/Engine/Source/Runtime/UMG/Public/Components/ScrollBar.h b/Engine/Source/Runtime/UMG/Public/Components/ScrollBar.h
--- a/Engine/Source/Runtime/UMG/Public/Components/ScrollBar.h
+++ b/Engine/Source/Runtime/UMG/Public/Components/ScrollBar.h
@@ -19,27 +19,33 @@ class UMG_API UScrollBar : public UWidget
 public:
 
 	/** Style of the scrollbar */
-	UPROPERTY(EditAnywhere, BlueprintReadWrite, Category="Style", meta=( DisplayName="Style" ))
+	UE_DEPRECATED(5.2, "Direct access to WidgetStyle is deprecated. Please use the getter or setter.")
+	UPROPERTY(EditAnywhere, BlueprintReadWrite, Getter, Setter, Category="Style", meta=( DisplayName="Style" ))
 	FScrollBarStyle WidgetStyle;
 
 	/**  */
-	UPROPERTY(EditAnywhere, Category="Behavior")
+	UE_DEPRECATED(5.2, "Direct access to bAlwaysShowScrollbar is deprecated. Please use the getter or setter.")
+	UPROPERTY(EditAnywhere, BlueprintReadWrite, Getter = "IsAlwaysShowScrollbar", Setter = "SetAlwaysShowScrollbar", Category = "Behavior")
 	bool bAlwaysShowScrollbar;
 
 	/**  */
-	UPROPERTY(EditAnywhere, Category = "Behavior")
+	UE_DEPRECATED(5.2, "Direct access to bAlwaysShowScrollbarTrack is deprecated. Please use the getter or setter.")
+	UPROPERTY(EditAnywhere, BlueprintReadWrite, Getter = "IsAlwaysShowScrollbarTrack", Setter = "SetAlwaysShowScrollbarTrack", Category = "Behavior")
 	bool bAlwaysShowScrollbarTrack;
 
 	/**  */
-	UPROPERTY(EditAnywhere, Category="Behavior")
+	UE_DEPRECATED(5.2, "Direct access to Orientation is deprecated. Please use the getter. Note that the orientation of a scrollbar is only set at construction and is not modifiable at runtime.")
+	UPROPERTY(EditAnywhere, BlueprintReadOnly, Getter, Category="Behavior")
 	TEnumAsByte<EOrientation> Orientation;
 
 	/** The thickness of the scrollbar thumb */
-	UPROPERTY(EditAnywhere, Category="Behavior")
+	UE_DEPRECATED(5.2, "Direct access to Thickness is deprecated. Please use the getter or setter.")
+	UPROPERTY(EditAnywhere, BlueprintReadWrite, Getter, Setter, Category="Behavior")
 	FVector2D Thickness;
 
 	/** The margin around the scrollbar */
-	UPROPERTY(EditAnywhere, Category = "Behavior")
+	UE_DEPRECATED(5.2, "Direct access to Padding is deprecated. Please use the getter or setter.")
+	UPROPERTY(EditAnywhere, BlueprintReadWrite, Getter, Setter, Category = "Behavior")
 	FMargin Padding;
 
 public:
@@ -90,6 +96,34 @@ public:
 	//~ End UWidget Interface
 #endif
 
+	/** @return the style of scrollbar. */
+	const FScrollBarStyle& GetWidgetStyle() const;
+	/** Sets the style of scrollbar. */
+	void SetWidgetStyle(const FScrollBarStyle& InWidgetStyle);
+
+	/** @return True if the scrollbar should always show. */
+	bool IsAlwaysShowScrollbar() const;
+	/** Sets whether the scrollbar should always show. */
+	void SetAlwaysShowScrollbar(bool bNewValue);
+
+	/** @return True if the scrollbar track should always show. */
+	bool IsAlwaysShowScrollbarTrack() const;
+	/** Sets whether the scrollbar track should always show */
+	void SetAlwaysShowScrollbarTrack(bool bNewValue);
+
+	/** @return the orientation of the scrollbar. */
+	EOrientation GetOrientation() const;
+
+	/** @return the thickness of the scrollbar. */
+	FVector2D GetThickness() const;
+	/** Sets the thickness of the scrollbar. */
+	void SetThickness(const FVector2D& InThickness);
+
+	/** Sets the padding of the scrollbar. */
+	FMargin GetPadding() const;
+	/** @return the padding of the scrollbar. */
+	void SetPadding(const FMargin& InPadding);
+
 protected:
 
 	TSharedPtr<class SScrollBar> MyScrollBar;
