Fix retainerbox accessor order for UHTCompactEditor

diff --git a/Engine/Source/Runtime/UMG/Public/Components/RetainerBox.h b/Engine/Source/Runtime/UMG/Public/Components/RetainerBox.h
--- a/Engine/Source/Runtime/UMG/Public/Components/RetainerBox.h
+++ b/Engine/Source/Runtime/UMG/Public/Components/RetainerBox.h
@@ -85,11 +85,6 @@ public:
 	UFUNCTION(BlueprintCallable, Category="Retainer")
 	void RequestRender();
 
-	/**
-	 * Get the current dynamic effect material applied to the retainer box.
-	 */
-	const UMaterialInterface* GetEffectMaterialInterface() const;
-
 	/**
 	 * Get the current dynamic effect material applied to the retainer box.
 	 */
@@ -108,11 +103,6 @@ public:
 	UFUNCTION(BlueprintCallable, Category="Retainer|Effect")
 	void SetTextureParameter(FName TextureParameter);
 
-	/**
-	 * Gets the name of the texture parameter to set the render target to on the material.
-	 */
-	const FName& GetTextureParameter() const;
-
 	/**
 	* Set the flag for if we retain the render or pass-through
 	*/
@@ -189,6 +179,17 @@ protected:
 #endif
 	//~ End UObject interface
 
+public:
+	/**
+	* Gets the name of the texture parameter to set the render target to on the material.
+	*/
+	const FName& GetTextureParameter() const;
+
+	/**
+	* Gets the current dynamic effect material applied to the retainer box.
+	*/
+	const UMaterialInterface* GetEffectMaterialInterface() const;
+
 protected:
 	TSharedPtr<class SRetainerWidget> MyRetainerWidget;
 };
