Add accessors for UImage.

diff --git a/Engine/Plugins/Experimental/CommonUI/Source/CommonUI/Private/CommonLazyImage.cpp b/Engine/Plugins/Experimental/CommonUI/Source/CommonUI/Private/CommonLazyImage.cpp
--- a/Engine/Plugins/Experimental/CommonUI/Source/CommonUI/Private/CommonLazyImage.cpp
+++ b/Engine/Plugins/Experimental/CommonUI/Source/CommonUI/Private/CommonLazyImage.cpp
@@ -88,7 +88,7 @@ bool UCommonLazyImage::IsLoading() const
 
 void UCommonLazyImage::SetMaterialTextureParamName(FName TextureParamName)
 {
-	if (Cast<UMaterialInterface>(Brush.GetResourceObject()) != nullptr)
+	if (Cast<UMaterialInterface>(GetBrush().GetResourceObject()) != nullptr)
 	{
 		MaterialTextureParamName = TextureParamName;
 	}
@@ -185,7 +185,7 @@ bool UCommonLazyImage::CanEditChange(const FProperty* InProperty) const
 	if (InProperty->GetFName().IsEqual(GET_MEMBER_NAME_CHECKED(UCommonLazyImage, MaterialTextureParamName)))
 	{
 		// The param name is only relevant when the brush uses a material
-		return Cast<UMaterialInterface>(Brush.GetResourceObject()) != nullptr;
+		return Cast<UMaterialInterface>(GetBrush().GetResourceObject()) != nullptr;
 	}
 	return Super::CanEditChange(InProperty);
 }
@@ -237,16 +237,18 @@ void UCommonLazyImage::SetBrushObjectInternal(UTexture* Texture, bool bMatchSize
 			BrushMID->SetTextureParameterValue(MaterialTextureParamName, Texture);
 			if (bMatchSize)
 			{
+				FSlateBrush CurrentBrush = GetBrush();
 				if (UTexture2DDynamic* AsDynamicTexture = Cast<UTexture2DDynamic>(Texture))
 				{
-					Brush.ImageSize.X = AsDynamicTexture->SizeX;
-					Brush.ImageSize.Y = AsDynamicTexture->SizeY;
+					CurrentBrush.ImageSize.X = AsDynamicTexture->SizeX;
+					CurrentBrush.ImageSize.Y = AsDynamicTexture->SizeY;
 				}
 				else if (UTexture2D* AsTexture2D = Cast<UTexture2D>(Texture))
 				{
-					Brush.ImageSize.X = AsTexture2D->GetSizeX();
-					Brush.ImageSize.Y = AsTexture2D->GetSizeY();
+					CurrentBrush.ImageSize.X = AsTexture2D->GetSizeX();
+					CurrentBrush.ImageSize.Y = AsTexture2D->GetSizeY();
 				}
+				Super::SetBrush(CurrentBrush);
 			}
 		}
 	}
diff --git a/Engine/Source/Editor/UMGEditor/Private/Templates/WidgetTemplateImageClass.cpp b/Engine/Source/Editor/UMGEditor/Private/Templates/WidgetTemplateImageClass.cpp
--- a/Engine/Source/Editor/UMGEditor/Private/Templates/WidgetTemplateImageClass.cpp
+++ b/Engine/Source/Editor/UMGEditor/Private/Templates/WidgetTemplateImageClass.cpp
@@ -23,7 +23,7 @@ UWidget* FWidgetTemplateImageClass::Create(UWidgetTree* WidgetTree)
 		UImage* ImageWidget = Cast<UImage>(Widget);
 
 		UObject* ImageResource = FindObject<UObject>(nullptr, *WidgetAssetData.GetObjectPathString());
-		ImageWidget->Brush.SetResourceObject(ImageResource);
+		ImageWidget->SetBrushResourceObject(ImageResource);
 	}
 
 	return Widget;
diff --git a/Engine/Source/Runtime/SlateCore/Private/Widgets/Images/SImage.cpp b/Engine/Source/Runtime/SlateCore/Private/Widgets/Images/SImage.cpp
--- a/Engine/Source/Runtime/SlateCore/Private/Widgets/Images/SImage.cpp
+++ b/Engine/Source/Runtime/SlateCore/Private/Widgets/Images/SImage.cpp
@@ -91,11 +91,26 @@ void SImage::SetImage(TAttribute<const FSlateBrush*> InImage)
 	ImageAttribute.Assign(*this, MoveTemp(InImage));
 }
 
+
+void SImage::InvalidateImage()
+{
+	Invalidate(EInvalidateWidgetReason::Layout);
+}
+
 void SImage::SetDesiredSizeOverride(TAttribute<TOptional<FVector2D>> InDesiredSizeOverride)
 {
 	DesiredSizeOverrideAttribute.Assign(*this, MoveTemp(InDesiredSizeOverride));
 }
 
+void SImage::FlipForRightToLeftFlowDirection(bool InbFlipForRightToLeftFlowDirection)
+{
+	if (InbFlipForRightToLeftFlowDirection != bFlipForRightToLeftFlowDirection)
+	{
+		bFlipForRightToLeftFlowDirection = InbFlipForRightToLeftFlowDirection;
+		Invalidate(EInvalidateWidgetReason::Paint);
+	}
+}
+
 #if WITH_ACCESSIBILITY
 TSharedRef<FSlateAccessibleWidget> SImage::CreateAccessibleWidget()
 {
diff --git a/Engine/Source/Runtime/SlateCore/Public/Widgets/Images/SImage.h b/Engine/Source/Runtime/SlateCore/Public/Widgets/Images/SImage.h
--- a/Engine/Source/Runtime/SlateCore/Public/Widgets/Images/SImage.h
+++ b/Engine/Source/Runtime/SlateCore/Public/Widgets/Images/SImage.h
@@ -66,9 +66,15 @@ public:
 	/** Set the Image attribute */
 	void SetImage(TAttribute<const FSlateBrush*> InImage);
 
+	/** Invalidate the Image */
+	void InvalidateImage();
+
 	/** Set SizeOverride attribute */
 	void SetDesiredSizeOverride(TAttribute<TOptional<FVector2D>> InDesiredSizeOverride);
 
+	/** Set FlipForRightToLeftFlowDirection */
+	void FlipForRightToLeftFlowDirection(bool InbFlipForRightToLeftFlowDirection);
+
 public:
 
 	// SWidget overrides
diff --git a/Engine/Source/Runtime/UMG/Private/Components/Image.cpp b/Engine/Source/Runtime/UMG/Private/Components/Image.cpp
--- a/Engine/Source/Runtime/UMG/Private/Components/Image.cpp
+++ b/Engine/Source/Runtime/UMG/Private/Components/Image.cpp
@@ -33,9 +33,10 @@ void UImage::ReleaseSlateResources(bool bReleaseChildren)
 
 TSharedRef<SWidget> UImage::RebuildWidget()
 {
+PRAGMA_DISABLE_DEPRECATION_WARNINGS
 	MyImage = SNew(SImage)
 			.FlipForRightToLeftFlowDirection(bFlipForRightToLeftFlowDirection);
-
+PRAGMA_ENABLE_DEPRECATION_WARNINGS
 	return MyImage.ToSharedRef();
 }
 
@@ -43,17 +44,20 @@ void UImage::SynchronizeProperties()
 {
 	Super::SynchronizeProperties();
 
+PRAGMA_DISABLE_DEPRECATION_WARNINGS
 	TAttribute<FSlateColor> ColorAndOpacityBinding = PROPERTY_BINDING(FSlateColor, ColorAndOpacity);
 	TAttribute<const FSlateBrush*> ImageBinding = OPTIONAL_BINDING_CONVERT(FSlateBrush, Brush, const FSlateBrush*, ConvertImage);
-
+PRAGMA_ENABLE_DEPRECATION_WARNINGS
 	if (MyImage.IsValid())
 	{
 		MyImage->SetImage(ImageBinding);
+		MyImage->InvalidateImage();
 		MyImage->SetColorAndOpacity(ColorAndOpacityBinding);
 		MyImage->SetOnMouseButtonDown(BIND_UOBJECT_DELEGATE(FPointerEventHandler, HandleMouseButtonDown));
 	}
 }
 
+PRAGMA_DISABLE_DEPRECATION_WARNINGS
 void UImage::SetColorAndOpacity(FLinearColor InColorAndOpacity)
 {
 	ColorAndOpacity = InColorAndOpacity;
@@ -63,6 +67,12 @@ void UImage::SetColorAndOpacity(FLinearColor InColorAndOpacity)
 	}
 }
 
+
+const FLinearColor& UImage::GetColorAndOpacity() const
+{
+	return ColorAndOpacity;
+}
+
 void UImage::SetOpacity(float InOpacity)
 {
 	ColorAndOpacity.A = InOpacity;
@@ -75,24 +85,31 @@ void UImage::SetOpacity(float InOpacity)
 const FSlateBrush* UImage::ConvertImage(TAttribute<FSlateBrush> InImageAsset) const
 {
 	UImage* MutableThis = const_cast<UImage*>( this );
-	MutableThis->Brush = InImageAsset.Get();
+	MutableThis->SetBrush(InImageAsset.Get());
 
 	return &Brush;
 }
 
 void UImage::SetBrush(const FSlateBrush& InBrush)
 {
-	if(Brush != InBrush)
+	if (Brush != InBrush)
 	{
 		Brush = InBrush;
-
+		BroadcastFieldValueChanged(FFieldNotificationClassDescriptor::Brush);
 		if (MyImage.IsValid())
 		{
-			MyImage->SetImage(&Brush);
+			MyImage->InvalidateImage();
 		}
 	}
 }
 
+const FSlateBrush& UImage::GetBrush() const
+{
+	return Brush;
+
+}
+PRAGMA_ENABLE_DEPRECATION_WARNINGS
+
 void UImage::SetBrushSize(FVector2D DesiredSize)
 {
 	SetDesiredSizeOverride(DesiredSize);
@@ -106,15 +123,17 @@ void UImage::SetDesiredSizeOverride(FVector2D DesiredSize)
 	}
 }
 
+PRAGMA_DISABLE_DEPRECATION_WARNINGS
 void UImage::SetBrushTintColor(FSlateColor TintColor)
 {
 	if(Brush.TintColor != TintColor)
 	{
 		Brush.TintColor = TintColor;
+		BroadcastFieldValueChanged(FFieldNotificationClassDescriptor::Brush);
 
 		if (MyImage.IsValid())
 		{
-			MyImage->SetImage(&Brush);
+			MyImage->InvalidateImage();;
 		}
 	}
 }
@@ -124,10 +143,11 @@ void UImage::SetBrushResourceObject(UObject* ResourceObject)
 	if (Brush.GetResourceObject() != ResourceObject)
 	{
 		Brush.SetResourceObject(ResourceObject);
+		BroadcastFieldValueChanged(FFieldNotificationClassDescriptor::Brush);
 
 		if (MyImage.IsValid())
 		{
-			MyImage->SetImage(&Brush);
+			MyImage->InvalidateImage();
 		}
 	}
 }
@@ -137,11 +157,13 @@ void UImage::SetBrushFromAsset(USlateBrushAsset* Asset)
 	if(!Asset || Brush != Asset->Brush)
 	{
 		CancelImageStreaming();
-		Brush = Asset ? Asset->Brush : FSlateBrush();
-
-		if (MyImage.IsValid())
+		if (Asset)
+		{
+			SetBrush(Asset->Brush);
+		}
+		else
 		{
-			MyImage->SetImage(&Brush);
+			SetBrush(FSlateBrush());
 		}
 	}
 }
@@ -152,7 +174,7 @@ void UImage::SetBrushFromTexture(UTexture2D* Texture, bool bMatchSize)
 
 	if(Brush.GetResourceObject() != Texture)
 	{
-		Brush.SetResourceObject(Texture);
+		SetBrushResourceObject(Texture);
 
 		if (Texture) // Since this texture is used as UI, don't allow it affected by budget.
 		{
@@ -171,11 +193,12 @@ void UImage::SetBrushFromTexture(UTexture2D* Texture, bool bMatchSize)
 			{
 				Brush.ImageSize = FVector2D(0, 0);
 			}
+			BroadcastFieldValueChanged(FFieldNotificationClassDescriptor::Brush);
 		}
 
 		if (MyImage.IsValid())
 		{
-			MyImage->SetImage(&Brush);
+			MyImage->InvalidateImage();
 		}
 	}
 }
@@ -185,7 +208,7 @@ void UImage::SetBrushFromAtlasInterface(TScriptInterface<ISlateTextureAtlasInter
 	if(Brush.GetResourceObject() != AtlasRegion.GetObject())
 	{
 		CancelImageStreaming();
-		Brush.SetResourceObject(AtlasRegion.GetObject());
+		SetBrushResourceObject(AtlasRegion.GetObject());
 
 		if (bMatchSize)
 		{
@@ -198,11 +221,12 @@ void UImage::SetBrushFromAtlasInterface(TScriptInterface<ISlateTextureAtlasInter
 			{
 				Brush.ImageSize = FVector2D(0, 0);
 			}
+			BroadcastFieldValueChanged(FFieldNotificationClassDescriptor::Brush);
 		}
 
 		if (MyImage.IsValid())
 		{
-			MyImage->SetImage(&Brush);
+			MyImage->InvalidateImage();
 		}
 	}
 }
@@ -212,17 +236,18 @@ void UImage::SetBrushFromTextureDynamic(UTexture2DDynamic* Texture, bool bMatchS
 	if(Brush.GetResourceObject() != Texture)
 	{
 		CancelImageStreaming();
-		Brush.SetResourceObject(Texture);
+		SetBrushResourceObject(Texture);
 
 		if (bMatchSize && Texture)
 		{
 			Brush.ImageSize.X = Texture->SizeX;
 			Brush.ImageSize.Y = Texture->SizeY;
+			BroadcastFieldValueChanged(FFieldNotificationClassDescriptor::Brush);
 		}
 
 		if (MyImage.IsValid())
 		{
-			MyImage->SetImage(&Brush);
+			MyImage->InvalidateImage();
 		}
 	}
 }
@@ -232,17 +257,37 @@ void UImage::SetBrushFromMaterial(UMaterialInterface* Material)
 	if(Brush.GetResourceObject() != Material)
 	{
 		CancelImageStreaming();
-		Brush.SetResourceObject(Material);
+		SetBrushResourceObject(Material);
 
 		//TODO UMG Check if the material can be used with the UI
 
 		if (MyImage.IsValid())
 		{
-			MyImage->SetImage(&Brush);
+			MyImage->InvalidateImage();
+		}
+	}
+}
+
+void UImage::SetFlipForRightToLeftFlowDirection(bool InbFlipForRightToLeftFlowDirection)
+{
+	if (bFlipForRightToLeftFlowDirection != InbFlipForRightToLeftFlowDirection)
+	{
+		bFlipForRightToLeftFlowDirection = InbFlipForRightToLeftFlowDirection;
+
+		if (MyImage.IsValid())
+		{
+			MyImage->FlipForRightToLeftFlowDirection(InbFlipForRightToLeftFlowDirection);
 		}
 	}
 }
 
+bool UImage::ShouldFlipForRightToLeftFlowDirection() const
+{
+	return bFlipForRightToLeftFlowDirection;
+
+}
+PRAGMA_ENABLE_DEPRECATION_WARNINGS
+
 void UImage::CancelImageStreaming()
 {
 	if (StreamingHandle.IsValid())
@@ -335,6 +380,7 @@ void UImage::SetBrushFromSoftMaterial(TSoftObjectPtr<UMaterialInterface> SoftMat
 	);
 }
 
+PRAGMA_DISABLE_DEPRECATION_WARNINGS
 UMaterialInstanceDynamic* UImage::GetDynamicMaterial()
 {
 	UMaterialInterface* Material = NULL;
@@ -349,14 +395,13 @@ UMaterialInstanceDynamic* UImage::GetDynamicMaterial()
 		if ( !DynamicMaterial )
 		{
 			DynamicMaterial = UMaterialInstanceDynamic::Create(Material, this);
-			Brush.SetResourceObject(DynamicMaterial);
+			SetBrushResourceObject(DynamicMaterial);
 
 			if ( MyImage.IsValid() )
 			{
-				MyImage->SetImage(&Brush);
+				MyImage->InvalidateImage();
 			}
 		}
-		
 		return DynamicMaterial;
 	}
 
@@ -364,6 +409,7 @@ UMaterialInstanceDynamic* UImage::GetDynamicMaterial()
 
 	return NULL;
 }
+PRAGMA_ENABLE_DEPRECATION_WARNINGS
 
 FReply UImage::HandleMouseButtonDown(const FGeometry& Geometry, const FPointerEvent& MouseEvent)
 {
diff --git a/Engine/Source/Runtime/UMG/Public/Components/Image.h b/Engine/Source/Runtime/UMG/Public/Components/Image.h
--- a/Engine/Source/Runtime/UMG/Public/Components/Image.h
+++ b/Engine/Source/Runtime/UMG/Public/Components/Image.h
@@ -34,7 +34,8 @@ class UMG_API UImage : public UWidget
 public:
 
 	/** Image to draw */
-	UPROPERTY(EditAnywhere, BlueprintReadOnly, Category=Appearance)
+	UE_DEPRECATED(5.2, "Direct access to Brush is deprecated. Please use the getter or setter.")
+	UPROPERTY(EditAnywhere, BlueprintReadWrite, Getter, Setter, BlueprintSetter = "SetBrush", FieldNotify, Category=Appearance)
 	FSlateBrush Brush;
 
 	/** A bindable delegate for the Image. */
@@ -42,7 +43,8 @@ public:
 	FGetSlateBrush BrushDelegate;
 
 	/** Color and opacity */
-	UPROPERTY(EditAnywhere, BlueprintReadOnly, Category=Appearance, meta=( sRGB="true") )
+	UE_DEPRECATED(5.2, "Direct access to ColorAndOpacity is deprecated. Please use the getter or setter.")
+	UPROPERTY(EditAnywhere, BlueprintReadWrite, Getter, Setter, BlueprintSetter = "SetColorAndOpacity", Category=Appearance, meta=( sRGB="true") )
 	FLinearColor ColorAndOpacity;
 
 	/** A bindable delegate for the ColorAndOpacity. */
@@ -50,7 +52,8 @@ public:
 	FGetLinearColor ColorAndOpacityDelegate;
 
 	/** Flips the image if the localization's flow direction is RightToLeft */
-	UPROPERTY(EditAnywhere, Category = "Localization")
+	UE_DEPRECATED(5.2, "Direct access to bFlipForRightToLeftFlowDirection is deprecated. Please use the getter or setter.")
+	UPROPERTY(EditAnywhere, BlueprintReadWrite, Getter = "ShouldFlipForRightToLeftFlowDirection", Setter = "SetFlipForRightToLeftFlowDirection", Category = "Localization")
 	bool bFlipForRightToLeftFlowDirection;
 
 public:
@@ -64,6 +67,8 @@ public:
 	UFUNCTION(BlueprintCallable, Category="Appearance")
 	void SetColorAndOpacity(FLinearColor InColorAndOpacity);
 
+	const FLinearColor& GetColorAndOpacity() const;
+
 	/**  */
 	UFUNCTION(BlueprintCallable, Category="Appearance")
 	void SetOpacity(float InOpacity);
@@ -87,6 +92,8 @@ public:
 	UFUNCTION(BlueprintCallable, Category="Appearance")
 	virtual void SetBrush(const FSlateBrush& InBrush);
 
+	const FSlateBrush& GetBrush() const;
+
 	/**  */
 	UFUNCTION(BlueprintCallable, Category="Appearance")
 	virtual void SetBrushFromAsset(USlateBrushAsset* Asset);
@@ -139,6 +146,10 @@ public:
 	UFUNCTION(BlueprintCallable, Category="Appearance")
 	UMaterialInstanceDynamic* GetDynamicMaterial();
 
+	void SetFlipForRightToLeftFlowDirection(bool InbFlipForRightToLeftFlowDirection);
+
+	bool ShouldFlipForRightToLeftFlowDirection() const;
+
 	//~ Begin UWidget Interface
 	virtual void SynchronizeProperties() override;
 	//~ End UWidget Interface
@@ -189,5 +200,7 @@ protected:
 
 protected:
 
+PRAGMA_DISABLE_DEPRECATION_WARNINGS
 	PROPERTY_BINDING_IMPLEMENTATION(FSlateColor, ColorAndOpacity);
+PRAGMA_ENABLE_DEPRECATION_WARNINGS
 };
