Fix to WidgetSwitcher to set Slate ActiveWidget Index even if it did not change to properly initialize it.

diff --git a/Engine/Source/Runtime/UMG/Private/Components/WidgetSwitcher.cpp b/Engine/Source/Runtime/UMG/Private/Components/WidgetSwitcher.cpp
--- a/Engine/Source/Runtime/UMG/Private/Components/WidgetSwitcher.cpp
+++ b/Engine/Source/Runtime/UMG/Private/Components/WidgetSwitcher.cpp
@@ -53,13 +53,13 @@ void UWidgetSwitcher::SetActiveWidgetIndex(int32 Index)
 	{
 		ActiveWidgetIndex = Index;
 		BroadcastFieldValueChanged(FFieldNotificationClassDescriptor::ActiveWidgetIndex);
+	}
 
-		if (MyWidgetSwitcher.IsValid())
-		{
-			// Ensure the index is clamped to a valid range.
-			int32 SafeIndex = FMath::Clamp(ActiveWidgetIndex, 0, FMath::Max(0, Slots.Num() - 1));
-			MyWidgetSwitcher->SetActiveWidgetIndex(SafeIndex);
-		}
+	if (MyWidgetSwitcher.IsValid())
+	{
+		// Ensure the index is clamped to a valid range.
+		int32 SafeIndex = FMath::Clamp(ActiveWidgetIndex, 0, FMath::Max(0, Slots.Num() - 1));
+		MyWidgetSwitcher->SetActiveWidgetIndex(SafeIndex);
 	}
 }
 
