Initialize active index for slate widget switcher in SyncronizeProperties

diff --git a/Engine/Source/Runtime/UMG/Private/Components/WidgetSwitcher.cpp b/Engine/Source/Runtime/UMG/Private/Components/WidgetSwitcher.cpp
--- a/Engine/Source/Runtime/UMG/Private/Components/WidgetSwitcher.cpp
+++ b/Engine/Source/Runtime/UMG/Private/Components/WidgetSwitcher.cpp
@@ -53,13 +53,7 @@ void UWidgetSwitcher::SetActiveWidgetIndex(int32 Index)
 	{
 		ActiveWidgetIndex = Index;
 		BroadcastFieldValueChanged(FFieldNotificationClassDescriptor::ActiveWidgetIndex);
-	}
-
-	if (MyWidgetSwitcher.IsValid())
-	{
-		// Ensure the index is clamped to a valid range.
-		int32 SafeIndex = FMath::Clamp(ActiveWidgetIndex, 0, FMath::Max(0, Slots.Num() - 1));
-		MyWidgetSwitcher->SetActiveWidgetIndex(SafeIndex);
+		SetActiveWidgetIndexForSlateWidget();
 	}
 }
 
@@ -71,8 +65,13 @@ void UWidgetSwitcher::SetActiveWidget(UWidget* Widget)
 	{
 		ActiveWidgetIndex = NewIndex;
 		BroadcastFieldValueChanged(FFieldNotificationClassDescriptor::ActiveWidgetIndex);
+		SetActiveWidgetIndexForSlateWidget();
 	}
-	if ( MyWidgetSwitcher.IsValid() )
+}
+
+void UWidgetSwitcher::SetActiveWidgetIndexForSlateWidget()
+{
+	if (MyWidgetSwitcher.IsValid())
 	{
 		// Ensure the index is clamped to a valid range.
 		int32 SafeIndex = FMath::Clamp(ActiveWidgetIndex, 0, FMath::Max(0, Slots.Num() - 1));
@@ -142,9 +141,7 @@ TSharedRef<SWidget> UWidgetSwitcher::RebuildWidget()
 void UWidgetSwitcher::SynchronizeProperties()
 {
 	Super::SynchronizeProperties();
-PRAGMA_DISABLE_DEPRECATION_WARNINGS
-	SetActiveWidgetIndex(ActiveWidgetIndex);
-PRAGMA_ENABLE_DEPRECATION_WARNINGS
+	SetActiveWidgetIndexForSlateWidget();
 }
 
 #if WITH_EDITOR
diff --git a/Engine/Source/Runtime/UMG/Public/Components/WidgetSwitcher.h b/Engine/Source/Runtime/UMG/Public/Components/WidgetSwitcher.h
--- a/Engine/Source/Runtime/UMG/Public/Components/WidgetSwitcher.h
+++ b/Engine/Source/Runtime/UMG/Public/Components/WidgetSwitcher.h
@@ -72,6 +72,7 @@ protected:
 	virtual void OnSlotRemoved(UPanelSlot* Slot) override;
 	// End UPanelWidget
 
+	void SetActiveWidgetIndexForSlateWidget();
 protected:
 
 	TSharedPtr<class SWidgetSwitcher> MyWidgetSwitcher;
