Fix property tooltip regression in UMG designer

diff --git a/Engine/Source/Runtime/UMG/Public/Components/InputKeySelector.h b/Engine/Source/Runtime/UMG/Public/Components/InputKeySelector.h
--- a/Engine/Source/Runtime/UMG/Public/Components/InputKeySelector.h
+++ b/Engine/Source/Runtime/UMG/Public/Components/InputKeySelector.h
@@ -27,18 +27,18 @@ public:
 	DECLARE_DYNAMIC_MULTICAST_DELEGATE( FOnIsSelectingKeyChanged );
 
 public:
-	/** The button style used at runtime */
 	UE_DEPRECATED(5.2, "Direct access to WidgetStyle is now deprecated. Use the setter or getter.")
+	/** The button style used at runtime */
 	UPROPERTY(EditAnywhere, BlueprintReadWrite, Getter = "GetButtonStyle", Setter = "SetButtonStyle", Category = Appearance, meta = (DisplayName = "Style"))
 	FButtonStyle WidgetStyle;
 
-	/** The button style used at runtime */
 	UE_DEPRECATED(5.2, "Direct access to TextStyle is now deprecated. Use the setter or getter.")
+	/** The button style used at runtime */
 	UPROPERTY(EditAnywhere, BlueprintReadWrite, Getter, Setter, Category = Appearance, meta = (DisplayName = "Text Style"))
 	FTextBlockStyle TextStyle;
 
-	/** The currently selected key chord. */
 	UE_DEPRECATED(5.2, "Direct access to SelectedKey is now deprecated. Use the setter or getter.")
+	/** The currently selected key chord. */
 	UPROPERTY( BlueprintReadWrite, Getter, Setter, BlueprintSetter = "SetSelectedKey", FieldNotify, Category = "Key Selection")
 	FInputChord SelectedKey;
 
@@ -47,24 +47,26 @@ public:
 	UPROPERTY( EditAnywhere, BlueprintReadWrite, Getter, Setter, Category = Appearance )
 	FMargin Margin;
 
-	/** Sets the text which is displayed while selecting keys. */
 	UE_DEPRECATED(5.2, "Direct access to KeySelectionText is now deprecated. Use the setter or getter.")
+	/** Sets the text which is displayed while selecting keys. */
 	UPROPERTY( EditAnywhere, BlueprintReadWrite, Getter, Setter, BlueprintSetter = "SetKeySelectionText", Category = Appearance )
 	FText KeySelectionText;
 
-	/** Sets the text to display when no key text is available or not selecting a key. */
 	UE_DEPRECATED(5.2, "Direct access to NoKeySpecifiedText is now deprecated. Use the setter or getter.")
+	/** Sets the text to display when no key text is available or not selecting a key. */
 	UPROPERTY( EditAnywhere, BlueprintReadWrite, Getter, Setter, BlueprintSetter = "SetNoKeySpecifiedText", Category = Appearance )
 	FText NoKeySpecifiedText;
 
-	/** When true modifier keys such as control and alt are allowed in the */
-	/** input chord representing the selected key, if false modifier keys are ignored. */
 	UE_DEPRECATED(5.2, "Direct access to bAllowModifierKeys is now deprecated. Use the setter or getter.")
+	/**
+	 * When true modifier keys such as control and alt are allowed in the
+	 * input chord representing the selected key, if false modifier keys are ignored.
+	 */
 	UPROPERTY( EditAnywhere, BlueprintReadWrite, Getter = "AllowModifierKeys", Setter = "SetAllowModifierKeys", BlueprintSetter = "SetAllowModifierKeys", Category = "Key Selection")
 	bool bAllowModifierKeys;
 
-	/** When true gamepad keys are allowed in the input chord representing the selected key, otherwise they are ignored. */
 	UE_DEPRECATED(5.2, "Direct access to bAllowGamepadKeys is now deprecated. Use the setter or getter.")
+	/** When true gamepad keys are allowed in the input chord representing the selected key, otherwise they are ignored. */
 	UPROPERTY( EditAnywhere, BlueprintReadWrite, Getter = "AllowGamepadKeys", Setter = "SetAllowGamepadKeys", BlueprintSetter = "SetAllowGamepadKeys", Category = "Key Selection")
 	bool bAllowGamepadKeys;
 
