MVVM: Add event to GameViewportSubsystem when a widget is added and removed.

diff --git a/Engine/Source/Runtime/UMG/Private/Blueprint/GameViewportSubsystem.cpp b/Engine/Source/Runtime/UMG/Private/Blueprint/GameViewportSubsystem.cpp
--- a/Engine/Source/Runtime/UMG/Private/Blueprint/GameViewportSubsystem.cpp
+++ b/Engine/Source/Runtime/UMG/Private/Blueprint/GameViewportSubsystem.cpp
@@ -179,6 +179,7 @@ void UGameViewportSubsystem::AddToScreen(UWidget* Widget, ULocalPlayer* Player,
 		// displaying below any built-in controls, like the virtual joysticks on mobile builds.
 		ViewportClient->AddViewportWidgetContent(FullScreenCanvas, Slot.ZOrder + 10);
 	}
+	OnWidgetAdded.Broadcast(Widget, Player);
 }
 
 void UGameViewportSubsystem::RemoveWidget(UWidget* Widget)
@@ -188,6 +189,8 @@ void UGameViewportSubsystem::RemoveWidget(UWidget* Widget)
 		FSlotInfo SlotInfo;
 		ViewportWidgets.RemoveAndCopyValue(Widget, SlotInfo);
 		RemoveWidgetInternal(Widget, SlotInfo.FullScreenWidget, SlotInfo.LocalPlayer);
+
+		OnWidgetRemoved.Broadcast(Widget);
 	}
 }
 
diff --git a/Engine/Source/Runtime/UMG/Public/Blueprint/GameViewportSubsystem.h b/Engine/Source/Runtime/UMG/Public/Blueprint/GameViewportSubsystem.h
--- a/Engine/Source/Runtime/UMG/Public/Blueprint/GameViewportSubsystem.h
+++ b/Engine/Source/Runtime/UMG/Public/Blueprint/GameViewportSubsystem.h
@@ -81,6 +81,14 @@ public:
 	UFUNCTION(BlueprintCallable, BlueprintCosmetic, Category = "User Interface")
 	void SetWidgetSlot(UWidget* Widget, FGameViewportWidgetSlot Slot);
 
+	/** */
+	DECLARE_EVENT_TwoParams(UGameViewportSubsystem, FWidgetAddedEvent, UWidget*, ULocalPlayer*);
+	FWidgetAddedEvent OnWidgetAdded;
+
+	/** */
+	DECLARE_EVENT_OneParam(UGameViewportSubsystem, FWidgetRemovedEvent, UWidget*);
+	FWidgetRemovedEvent OnWidgetRemoved;
+
 public:
 	/**
 	 * Helper function to set the position in the viewport for the Slot.
