Fix to UMGSequenceTickManager to stop off-screen animations from becoming active. Fixes an issue in the battlepass reward screen where the intro animation was blending with old animations from leaving the screen - results in incorrect final animation values (50% opacity widgets).

Original fix made by [REDACTED_PERSON] and shelved in CL 23409988

#rnx

diff --git a/Engine/Source/Runtime/UMG/Private/Animation/UMGSequenceTickManager.cpp b/Engine/Source/Runtime/UMG/Private/Animation/UMGSequenceTickManager.cpp
--- a/Engine/Source/Runtime/UMG/Private/Animation/UMGSequenceTickManager.cpp
+++ b/Engine/Source/Runtime/UMG/Private/Animation/UMGSequenceTickManager.cpp
@@ -169,6 +169,7 @@ void UUMGSequenceTickManager::TickWidgetAnimations(float DeltaSeconds)
 					else
 					{
 						WidgetIter.RemoveCurrent();
+						UserWidget->AnimationTickManager = nullptr;
 						continue;
 					}
 				}
