UMG: Backout CL 23136588 and 23421191 to return code to its previous state

diff --git a/Engine/Source/Runtime/UMG/Private/Animation/UMGSequenceTickManager.cpp b/Engine/Source/Runtime/UMG/Private/Animation/UMGSequenceTickManager.cpp
--- a/Engine/Source/Runtime/UMG/Private/Animation/UMGSequenceTickManager.cpp
+++ b/Engine/Source/Runtime/UMG/Private/Animation/UMGSequenceTickManager.cpp
@@ -71,10 +71,6 @@ void UUMGSequenceTickManager::OnWidgetTicked(UUserWidget* InWidget)
 	{
 		WidgetData->bIsTicking = true;
 	}
-	else
-	{
-		WeakUserWidgetData.Add(InWidget, FSequenceTickManagerWidgetData());
-	}
 }
 
 void UUMGSequenceTickManager::BeginDestroy()
@@ -164,14 +160,9 @@ void UUMGSequenceTickManager::TickWidgetAnimations(float DeltaSeconds)
 					if (WidgetData.bLastKnownTickState)
 					{
 						UserWidget->DisableAnimations();
-						WidgetData.bLastKnownTickState = false;
-					}
-					else
-					{
-						WidgetIter.RemoveCurrent();
-						UserWidget->AnimationTickManager = nullptr;
-						continue;
 					}
+
+					WidgetData.bLastKnownTickState = false;
 				}
 			}
 			else
