UEditableText/UEditableTextBox/UMultiLineEditableText:
- fix SetText that may fail due to lack of synchro with the underlying SWidget state.

tests (for all types of editable texts) :
-repro from the UDN (entering a text as a user, then setting an empty text via blueprint).
-same as UDN, but with a not empty string.
-validated that assignment was still only done in case of string change.

diff --git a/Engine/Source/Runtime/UMG/Private/Components/EditableText.cpp b/Engine/Source/Runtime/UMG/Private/Components/EditableText.cpp
--- a/Engine/Source/Runtime/UMG/Private/Components/EditableText.cpp
+++ b/Engine/Source/Runtime/UMG/Private/Components/EditableText.cpp
@@ -161,7 +161,7 @@ FText UEditableText::GetText() const
 void UEditableText::SetText(FText InText)
 {
 	// We detect if the Text is internal pointing to the same thing if so, nothing to do.
-	if (Text.IdenticalTo(InText))
+	if (GetText().IdenticalTo(InText))
 	{
 		return;
 	}
diff --git a/Engine/Source/Runtime/UMG/Private/Components/EditableTextBox.cpp b/Engine/Source/Runtime/UMG/Private/Components/EditableTextBox.cpp
--- a/Engine/Source/Runtime/UMG/Private/Components/EditableTextBox.cpp
+++ b/Engine/Source/Runtime/UMG/Private/Components/EditableTextBox.cpp
@@ -176,7 +176,7 @@ FText UEditableTextBox::GetText() const
 void UEditableTextBox::SetText(FText InText)
 {
 	// We detect if the Text is internal pointing to the same thing if so, nothing to do.
-	if (Text.IdenticalTo(InText))
+	if (GetText().IdenticalTo(InText))
 	{
 		return;
 	}
diff --git a/Engine/Source/Runtime/UMG/Private/Components/MultiLineEditableText.cpp b/Engine/Source/Runtime/UMG/Private/Components/MultiLineEditableText.cpp
--- a/Engine/Source/Runtime/UMG/Private/Components/MultiLineEditableText.cpp
+++ b/Engine/Source/Runtime/UMG/Private/Components/MultiLineEditableText.cpp
@@ -153,7 +153,7 @@ FText UMultiLineEditableText::GetText() const
 void UMultiLineEditableText::SetText(FText InText)
 {
 	// We detect if the Text is internal pointing to the same thing if so, nothing to do.
-	if (Text.IdenticalTo(InText))
+	if (GetText().IdenticalTo(InText))
 	{
 		return;
 	}
