All EditableTexts:
-prevent MVVM notifications from being called to often.
-ensure MVVM notifications are called when the text is edited from the SWidget.

MultiLineEditableTextBox: Like other members from the (MultiLine)EditableText(Box) family, enable MVVM notification on the Text property.

Tests: same than for the bug fix for UE-171644
-repro from the UDN (entering a text as a user, then setting an empty text via blueprint).
-same as UDN, but with a not empty string.
-validated that assignment was still only done in case of string change.

diff --git a/Engine/Source/Runtime/UMG/Private/Components/EditableText.cpp b/Engine/Source/Runtime/UMG/Private/Components/EditableText.cpp
--- a/Engine/Source/Runtime/UMG/Private/Components/EditableText.cpp
+++ b/Engine/Source/Runtime/UMG/Private/Components/EditableText.cpp
@@ -160,20 +160,22 @@ FText UEditableText::GetText() const
 
 void UEditableText::SetText(FText InText)
 {
-	// We detect if the Text is internal pointing to the same thing if so, nothing to do.
-	if (GetText().IdenticalTo(InText))
+	if (SetTextInternal(InText) && MyEditableText.IsValid() )
 	{
-		return;
+		MyEditableText->SetText(Text);
 	}
+}
 
-	Text = InText;
-
-	BroadcastFieldValueChanged(FFieldNotificationClassDescriptor::Text);
-
-	if ( MyEditableText.IsValid() )
+bool UEditableText::SetTextInternal(const FText& InText)
+{
+	if (!Text.IdenticalTo(InText, ETextIdenticalModeFlags::DeepCompare | ETextIdenticalModeFlags::LexicalCompareInvariants))
 	{
-		MyEditableText->SetText(Text);
+		Text = InText;
+		BroadcastFieldValueChanged(FFieldNotificationClassDescriptor::Text);
+		return true;
 	}
+
+	return false;
 }
 
 void UEditableText::SetIsPassword(bool InbIsPassword)
@@ -389,11 +391,15 @@ void UEditableText::SetFontOutlineMaterial(UMaterialInterface* InMaterial)
 
 void UEditableText::HandleOnTextChanged(const FText& InText)
 {
-	OnTextChanged.Broadcast(InText);
+	if (SetTextInternal(InText))
+	{
+		OnTextChanged.Broadcast(InText);
+	}
 }
 
 void UEditableText::HandleOnTextCommitted(const FText& InText, ETextCommit::Type CommitMethod)
 {
+	SetTextInternal(InText);
 	OnTextCommitted.Broadcast(InText, CommitMethod);
 }
 
diff --git a/Engine/Source/Runtime/UMG/Private/Components/EditableTextBox.cpp b/Engine/Source/Runtime/UMG/Private/Components/EditableTextBox.cpp
--- a/Engine/Source/Runtime/UMG/Private/Components/EditableTextBox.cpp
+++ b/Engine/Source/Runtime/UMG/Private/Components/EditableTextBox.cpp
@@ -175,20 +175,22 @@ FText UEditableTextBox::GetText() const
 
 void UEditableTextBox::SetText(FText InText)
 {
-	// We detect if the Text is internal pointing to the same thing if so, nothing to do.
-	if (GetText().IdenticalTo(InText))
+	if (SetTextInternal(InText) && MyEditableTextBlock.IsValid() )
 	{
-		return;
+		MyEditableTextBlock->SetText(Text);
 	}
+}
 
-	Text = InText;
-
-	BroadcastFieldValueChanged(FFieldNotificationClassDescriptor::Text);
-
-	if ( MyEditableTextBlock.IsValid() )
+bool UEditableTextBox::SetTextInternal(const FText& InText)
+{
+	if (!Text.IdenticalTo(InText, ETextIdenticalModeFlags::DeepCompare | ETextIdenticalModeFlags::LexicalCompareInvariants))
 	{
-		MyEditableTextBlock->SetText(Text);
+		Text = InText;
+		BroadcastFieldValueChanged(FFieldNotificationClassDescriptor::Text);
+		return true;
 	}
+
+	return false;
 }
 
 FText UEditableTextBox::GetHintText() const
@@ -389,13 +391,15 @@ void UEditableTextBox::SetTextOverflowPolicy(ETextOverflowPolicy InOverflowPolic
 
 void UEditableTextBox::HandleOnTextChanged(const FText& InText)
 {
-	Text = InText;
-	OnTextChanged.Broadcast(InText);
+	if (SetTextInternal(InText))
+	{
+		OnTextChanged.Broadcast(InText);
+	}
 }
 
 void UEditableTextBox::HandleOnTextCommitted(const FText& InText, ETextCommit::Type CommitMethod)
 {
-	Text = InText;
+	SetTextInternal(InText);
 	OnTextCommitted.Broadcast(InText, CommitMethod);
 }
 PRAGMA_ENABLE_DEPRECATION_WARNINGS
diff --git a/Engine/Source/Runtime/UMG/Private/Components/MultiLineEditableText.cpp b/Engine/Source/Runtime/UMG/Private/Components/MultiLineEditableText.cpp
--- a/Engine/Source/Runtime/UMG/Private/Components/MultiLineEditableText.cpp
+++ b/Engine/Source/Runtime/UMG/Private/Components/MultiLineEditableText.cpp
@@ -152,20 +152,22 @@ FText UMultiLineEditableText::GetText() const
 
 void UMultiLineEditableText::SetText(FText InText)
 {
-	// We detect if the Text is internal pointing to the same thing if so, nothing to do.
-	if (GetText().IdenticalTo(InText))
+	if (SetTextInternal(InText) && MyMultiLineEditableText.IsValid() )
 	{
-		return;
+		MyMultiLineEditableText->SetText(Text);
 	}
+}
 
-	Text = InText;
-
-	BroadcastFieldValueChanged(FFieldNotificationClassDescriptor::Text);
-
-	if ( MyMultiLineEditableText.IsValid() )
+bool UMultiLineEditableText::SetTextInternal(const FText& InText)
+{
+	if (!Text.IdenticalTo(InText, ETextIdenticalModeFlags::DeepCompare | ETextIdenticalModeFlags::LexicalCompareInvariants))
 	{
-		MyMultiLineEditableText->SetText(Text);
+		Text = InText;
+		BroadcastFieldValueChanged(FFieldNotificationClassDescriptor::Text);
+		return true;
 	}
+
+	return false;
 }
 
 FText UMultiLineEditableText::GetHintText() const
@@ -312,11 +314,15 @@ void UMultiLineEditableText::SetFontOutlineMaterial(UMaterialInterface* InMateri
 
 void UMultiLineEditableText::HandleOnTextChanged(const FText& InText)
 {
-	OnTextChanged.Broadcast(InText);
+	if (SetTextInternal(InText))
+	{
+		OnTextChanged.Broadcast(InText);
+	}
 }
 
 void UMultiLineEditableText::HandleOnTextCommitted(const FText& InText, ETextCommit::Type CommitMethod)
 {
+	SetTextInternal(InText);
 	OnTextCommitted.Broadcast(InText, CommitMethod);
 }
 
diff --git a/Engine/Source/Runtime/UMG/Private/Components/MultiLineEditableTextBox.cpp b/Engine/Source/Runtime/UMG/Private/Components/MultiLineEditableTextBox.cpp
--- a/Engine/Source/Runtime/UMG/Private/Components/MultiLineEditableTextBox.cpp
+++ b/Engine/Source/Runtime/UMG/Private/Components/MultiLineEditableTextBox.cpp
@@ -163,13 +163,24 @@ FText UMultiLineEditableTextBox::GetText() const
 
 void UMultiLineEditableTextBox::SetText(FText InText)
 {
-	Text = InText;
-	if ( MyEditableTextBlock.IsValid() )
+	if (SetTextInternal(InText) && MyEditableTextBlock.IsValid() )
 	{
 		MyEditableTextBlock->SetText(Text);
 	}
 }
 
+bool UMultiLineEditableTextBox::SetTextInternal(const FText& InText)
+{
+	if (!Text.IdenticalTo(InText, ETextIdenticalModeFlags::DeepCompare | ETextIdenticalModeFlags::LexicalCompareInvariants))
+	{
+		Text = InText;
+		BroadcastFieldValueChanged(FFieldNotificationClassDescriptor::Text);
+		return true;
+	}
+
+	return false;
+}
+
 FText UMultiLineEditableTextBox::GetHintText() const
 {
 	if (MyEditableTextBlock.IsValid())
@@ -236,11 +247,15 @@ void UMultiLineEditableTextBox::SetForegroundColor(FLinearColor color)
 
 void UMultiLineEditableTextBox::HandleOnTextChanged(const FText& InText)
 {
-	OnTextChanged.Broadcast(InText);
+	if (SetTextInternal(InText))
+	{
+		OnTextChanged.Broadcast(InText);
+	}
 }
 
 void UMultiLineEditableTextBox::HandleOnTextCommitted(const FText& InText, ETextCommit::Type CommitMethod)
 {
+	SetTextInternal(InText);
 	OnTextCommitted.Broadcast(InText, CommitMethod);
 }
 
diff --git a/Engine/Source/Runtime/UMG/Public/Components/EditableText.h b/Engine/Source/Runtime/UMG/Public/Components/EditableText.h
--- a/Engine/Source/Runtime/UMG/Public/Components/EditableText.h
+++ b/Engine/Source/Runtime/UMG/Public/Components/EditableText.h
@@ -273,4 +273,8 @@ protected:
 
 	PROPERTY_BINDING_IMPLEMENTATION(FText, Text);
 	PROPERTY_BINDING_IMPLEMENTATION(FText, HintText);
+
+private:
+	/** @return true if the text was changed, or false if identical. */
+	bool SetTextInternal(const FText& InText);
 };
diff --git a/Engine/Source/Runtime/UMG/Public/Components/EditableTextBox.h b/Engine/Source/Runtime/UMG/Public/Components/EditableTextBox.h
--- a/Engine/Source/Runtime/UMG/Public/Components/EditableTextBox.h
+++ b/Engine/Source/Runtime/UMG/Public/Components/EditableTextBox.h
@@ -284,6 +284,9 @@ protected:
 	PROPERTY_BINDING_IMPLEMENTATION(FText, HintText);
 
 private:
+	/** @return true if the text was changed, or false if identical. */
+	bool SetTextInternal(const FText& InText);
+
 #if WITH_EDITORONLY_DATA
 	UPROPERTY()
 	bool bIsFontDeprecationDone;
diff --git a/Engine/Source/Runtime/UMG/Public/Components/MultiLineEditableText.h b/Engine/Source/Runtime/UMG/Public/Components/MultiLineEditableText.h
--- a/Engine/Source/Runtime/UMG/Public/Components/MultiLineEditableText.h
+++ b/Engine/Source/Runtime/UMG/Public/Components/MultiLineEditableText.h
@@ -193,4 +193,8 @@ protected:
 	TSharedPtr<SMultiLineEditableText> MyMultiLineEditableText;
 
 	PROPERTY_BINDING_IMPLEMENTATION(FText, HintText);
+
+private:
+	/** @return true if the text was changed, or false if identical. */
+	bool SetTextInternal(const FText& InText);
 };
diff --git a/Engine/Source/Runtime/UMG/Public/Components/MultiLineEditableTextBox.h b/Engine/Source/Runtime/UMG/Public/Components/MultiLineEditableTextBox.h
--- a/Engine/Source/Runtime/UMG/Public/Components/MultiLineEditableTextBox.h
+++ b/Engine/Source/Runtime/UMG/Public/Components/MultiLineEditableTextBox.h
@@ -31,7 +31,7 @@ public:
 	
 	UE_DEPRECATED(5.1, "Direct access to Text is deprecated. Please use the getter or setter.")
 	/** The text content for this editable text box widget */
-	UPROPERTY(EditAnywhere, BlueprintReadWrite, Getter, Setter, BlueprintGetter = "GetText", BlueprintSetter = "SetText", Category = Content, meta=(MultiLine="true"))
+	UPROPERTY(EditAnywhere, BlueprintReadWrite, Getter, Setter, BlueprintGetter = "GetText", BlueprintSetter = "SetText", FieldNotify, Category = Content, meta=(MultiLine="true"))
 	FText Text;
 
 	UE_DEPRECATED(5.1, "Direct access to HintText is deprecated. Please use the getter or setter.")
@@ -163,6 +163,9 @@ protected:
 	PROPERTY_BINDING_IMPLEMENTATION(FText, HintText);
 
 private:
+	/** @return true if the text was changed, or false if identical. */
+	bool SetTextInternal(const FText& InText);
+
 #if WITH_EDITORONLY_DATA
 	UPROPERTY()
 	bool bIsFontDeprecationDone;
