UMG: Offscreen widgets with animations are no longer tracked by the tick manager

  - When a previously ticking widget no longer reports itself as ticking, all animation data will be disabled, and the widget removed from the tick manager
  - Since the widget is now removed from the tick manager, there is no need to keep track of its last known tick state
  - If the widget starts ticking in future, it will see that its animation tick manager is null, re-initialize it and start updating as normal
  - Added safety to AddWidget to ensure that flags do not get overwritten if it is called when the widget is already being tracked (this is not a real problem currently, but I don't want it to become one)

diff --git a/Engine/Source/Runtime/UMG/Private/Animation/UMGSequenceTickManager.cpp b/Engine/Source/Runtime/UMG/Private/Animation/UMGSequenceTickManager.cpp
--- a/Engine/Source/Runtime/UMG/Private/Animation/UMGSequenceTickManager.cpp
+++ b/Engine/Source/Runtime/UMG/Private/Animation/UMGSequenceTickManager.cpp
@@ -55,7 +55,10 @@ UUMGSequenceTickManager::UUMGSequenceTickManager(const FObjectInitializer& Init)
 void UUMGSequenceTickManager::AddWidget(UUserWidget* InWidget)
 {
 	TWeakObjectPtr<UUserWidget> WeakWidget = InWidget;
-	WeakUserWidgetData.Add(WeakWidget, FSequenceTickManagerWidgetData());
+	if (!WeakUserWidgetData.Contains(WeakWidget))
+	{
+		WeakUserWidgetData.Add(WeakWidget, FSequenceTickManagerWidgetData());
+	}
 }
 
 void UUMGSequenceTickManager::RemoveWidget(UUserWidget* InWidget)
@@ -71,6 +74,10 @@ void UUMGSequenceTickManager::OnWidgetTicked(UUserWidget* InWidget)
 	{
 		WidgetData->bIsTicking = true;
 	}
+	else
+	{
+		WeakUserWidgetData.Add(InWidget, FSequenceTickManagerWidgetData());
+	}
 }
 
 void UUMGSequenceTickManager::BeginDestroy()
@@ -148,21 +155,18 @@ void UUMGSequenceTickManager::TickWidgetAnimations(float DeltaSeconds)
 			}
 			else if (!WidgetData.bIsTicking)
 			{
-				// If this widget has not told us it is ticking, and its last known state was
-				// ticking, we disable animations for that widget. Once it ticks again, the animation
-				// will be updated naturally, and doesn't need anything re-enabling.
+				// If this widget has not told us it is ticking, we disable animations for that widget.
+				// Once it ticks again, the animation will be updated naturally, and doesn't need anything re-enabling.
 				// 
 				// @todo: There is a chance that relative animations hitting this code path will resume with
 				// different relative bases due to the way the ecs data is destroyed and re-created.
 				// In order to fix this we would have to annex that data instead of destroying it.
 				if (!bIsCurrentlyEvaluating)
 				{
-					if (WidgetData.bLastKnownTickState)
-					{
-						UserWidget->DisableAnimations();
-					}
+					UserWidget->DisableAnimations();
+					UserWidget->AnimationTickManager = nullptr;
 
-					WidgetData.bLastKnownTickState = false;
+					WidgetIter.RemoveCurrent();
 				}
 			}
 			else
@@ -182,7 +186,6 @@ void UUMGSequenceTickManager::TickWidgetAnimations(float DeltaSeconds)
 
 				// Assume this widget will no longer tick, until we're told otherwise by way of OnWidgetTicked
 				WidgetData.bIsTicking = false;
-				WidgetData.bLastKnownTickState = true;
 			}
 		}
 	}
@@ -212,6 +215,8 @@ void UUMGSequenceTickManager::TickWidgetAnimations(float DeltaSeconds)
 			}
 		}
 	}
+
+	WeakUserWidgetData.Shrink();
 }
 
 void UUMGSequenceTickManager::ForceFlush()
diff --git a/Engine/Source/Runtime/UMG/Private/UserWidget.cpp b/Engine/Source/Runtime/UMG/Private/UserWidget.cpp
--- a/Engine/Source/Runtime/UMG/Private/UserWidget.cpp
+++ b/Engine/Source/Runtime/UMG/Private/UserWidget.cpp
@@ -1465,6 +1465,13 @@ void UUserWidget::NativeTick(const FGeometry& MyGeometry, float InDeltaTime)
 #endif
 		if (bTickAnimations)
 		{
+			// If we have active animations but no tick manager this means we were offscreen
+			// but are now being actively ticked again so we need to resume animations
+			if (ActiveSequencePlayers.Num() != 0 && AnimationTickManager == nullptr)
+			{
+				AnimationTickManager = UUMGSequenceTickManager::Get(this);
+			}
+
 			if (AnimationTickManager)
 			{
 				AnimationTickManager->OnWidgetTicked(this);
diff --git a/Engine/Source/Runtime/UMG/Public/Animation/UMGSequenceTickManager.h b/Engine/Source/Runtime/UMG/Public/Animation/UMGSequenceTickManager.h
--- a/Engine/Source/Runtime/UMG/Public/Animation/UMGSequenceTickManager.h
+++ b/Engine/Source/Runtime/UMG/Public/Animation/UMGSequenceTickManager.h
@@ -26,9 +26,6 @@ struct FSequenceTickManagerWidgetData
 	/** The widget was ticked. */
 	bool bIsTicking = true;
 
-	/** The widget was ticked last update. */
-	bool bLastKnownTickState = true;
-
 	/** The widget's animation was ticked this frame. */
 	bool bActionsAndAnimationTicked = false;
 };
