Backout CL 23477010

diff --git a/Engine/Source/Runtime/UMG/Private/Animation/UMGSequenceTickManager.cpp b/Engine/Source/Runtime/UMG/Private/Animation/UMGSequenceTickManager.cpp
--- a/Engine/Source/Runtime/UMG/Private/Animation/UMGSequenceTickManager.cpp
+++ b/Engine/Source/Runtime/UMG/Private/Animation/UMGSequenceTickManager.cpp
@@ -55,10 +55,7 @@ UUMGSequenceTickManager::UUMGSequenceTickManager(const FObjectInitializer& Init)
 void UUMGSequenceTickManager::AddWidget(UUserWidget* InWidget)
 {
 	TWeakObjectPtr<UUserWidget> WeakWidget = InWidget;
-	if (!WeakUserWidgetData.Contains(WeakWidget))
-	{
-		WeakUserWidgetData.Add(WeakWidget, FSequenceTickManagerWidgetData());
-	}
+	WeakUserWidgetData.Add(WeakWidget, FSequenceTickManagerWidgetData());
 }
 
 void UUMGSequenceTickManager::RemoveWidget(UUserWidget* InWidget)
@@ -74,10 +71,6 @@ void UUMGSequenceTickManager::OnWidgetTicked(UUserWidget* InWidget)
 	{
 		WidgetData->bIsTicking = true;
 	}
-	else
-	{
-		WeakUserWidgetData.Add(InWidget, FSequenceTickManagerWidgetData());
-	}
 }
 
 void UUMGSequenceTickManager::BeginDestroy()
@@ -155,18 +148,21 @@ void UUMGSequenceTickManager::TickWidgetAnimations(float DeltaSeconds)
 			}
 			else if (!WidgetData.bIsTicking)
 			{
-				// If this widget has not told us it is ticking, we disable animations for that widget.
-				// Once it ticks again, the animation will be updated naturally, and doesn't need anything re-enabling.
+				// If this widget has not told us it is ticking, and its last known state was
+				// ticking, we disable animations for that widget. Once it ticks again, the animation
+				// will be updated naturally, and doesn't need anything re-enabling.
 				// 
 				// @todo: There is a chance that relative animations hitting this code path will resume with
 				// different relative bases due to the way the ecs data is destroyed and re-created.
 				// In order to fix this we would have to annex that data instead of destroying it.
 				if (!bIsCurrentlyEvaluating)
 				{
-					UserWidget->DisableAnimations();
-					UserWidget->AnimationTickManager = nullptr;
+					if (WidgetData.bLastKnownTickState)
+					{
+						UserWidget->DisableAnimations();
+					}
 
-					WidgetIter.RemoveCurrent();
+					WidgetData.bLastKnownTickState = false;
 				}
 			}
 			else
@@ -186,6 +182,7 @@ void UUMGSequenceTickManager::TickWidgetAnimations(float DeltaSeconds)
 
 				// Assume this widget will no longer tick, until we're told otherwise by way of OnWidgetTicked
 				WidgetData.bIsTicking = false;
+				WidgetData.bLastKnownTickState = true;
 			}
 		}
 	}
@@ -215,8 +212,6 @@ void UUMGSequenceTickManager::TickWidgetAnimations(float DeltaSeconds)
 			}
 		}
 	}
-
-	WeakUserWidgetData.Shrink();
 }
 
 void UUMGSequenceTickManager::ForceFlush()
diff --git a/Engine/Source/Runtime/UMG/Private/UserWidget.cpp b/Engine/Source/Runtime/UMG/Private/UserWidget.cpp
--- a/Engine/Source/Runtime/UMG/Private/UserWidget.cpp
+++ b/Engine/Source/Runtime/UMG/Private/UserWidget.cpp
@@ -1465,13 +1465,6 @@ void UUserWidget::NativeTick(const FGeometry& MyGeometry, float InDeltaTime)
 #endif
 		if (bTickAnimations)
 		{
-			// If we have active animations but no tick manager this means we were offscreen
-			// but are now being actively ticked again so we need to resume animations
-			if (ActiveSequencePlayers.Num() != 0 && AnimationTickManager == nullptr)
-			{
-				AnimationTickManager = UUMGSequenceTickManager::Get(this);
-			}
-
 			if (AnimationTickManager)
 			{
 				AnimationTickManager->OnWidgetTicked(this);
diff --git a/Engine/Source/Runtime/UMG/Public/Animation/UMGSequenceTickManager.h b/Engine/Source/Runtime/UMG/Public/Animation/UMGSequenceTickManager.h
--- a/Engine/Source/Runtime/UMG/Public/Animation/UMGSequenceTickManager.h
+++ b/Engine/Source/Runtime/UMG/Public/Animation/UMGSequenceTickManager.h
@@ -26,6 +26,9 @@ struct FSequenceTickManagerWidgetData
 	/** The widget was ticked. */
 	bool bIsTicking = true;
 
+	/** The widget was ticked last update. */
+	bool bLastKnownTickState = true;
+
 	/** The widget's animation was ticked this frame. */
 	bool bActionsAndAnimationTicked = false;
 };
