UMG: Offscreen widgets with animations are no longer tracked by the tick manager

  - When a previously ticking widget no longer reports itself as ticking, all animation data will be disabled, and the widget removed from the tick manager
  - Since the widget is now removed from the tick manager, there is no need to keep track of its last known tick state
  - If the widget starts ticking in future OnWidgetTicked will be called and the animation updated
  - Any call to GetOrAddSequencePlayer will always add the widget back onto its tick manager to ensure that the initial update is made. UUserWidget::AnimationTickManager should always be valid if a widget has any ActiveSequencePlayers

diff --git a/Engine/Source/Runtime/UMG/Private/Animation/UMGSequenceTickManager.cpp b/Engine/Source/Runtime/UMG/Private/Animation/UMGSequenceTickManager.cpp
--- a/Engine/Source/Runtime/UMG/Private/Animation/UMGSequenceTickManager.cpp
+++ b/Engine/Source/Runtime/UMG/Private/Animation/UMGSequenceTickManager.cpp
@@ -54,8 +54,18 @@ UUMGSequenceTickManager::UUMGSequenceTickManager(const FObjectInitializer& Init)
 
 void UUMGSequenceTickManager::AddWidget(UUserWidget* InWidget)
 {
+	// This is functionally the same as OnWidgetTicked, but they remain
+	// separate functions to convey the semantic difference
 	TWeakObjectPtr<UUserWidget> WeakWidget = InWidget;
-	WeakUserWidgetData.Add(WeakWidget, FSequenceTickManagerWidgetData());
+
+	if (FSequenceTickManagerWidgetData* WidgetData = WeakUserWidgetData.Find(WeakWidget))
+	{
+		WidgetData->bIsTicking = true;
+	}
+	else
+	{
+		WeakUserWidgetData.Add(WeakWidget, FSequenceTickManagerWidgetData());
+	}
 }
 
 void UUMGSequenceTickManager::RemoveWidget(UUserWidget* InWidget)
@@ -71,6 +81,10 @@ void UUMGSequenceTickManager::OnWidgetTicked(UUserWidget* InWidget)
 	{
 		WidgetData->bIsTicking = true;
 	}
+	else
+	{
+		WeakUserWidgetData.Add(InWidget, FSequenceTickManagerWidgetData());
+	}
 }
 
 void UUMGSequenceTickManager::BeginDestroy()
@@ -141,6 +155,9 @@ void UUMGSequenceTickManager::TickWidgetAnimations(float DeltaSeconds)
 				{
 					// Wait until the runner has finished evaluating before we tear stuff down
 					UserWidget->TearDownAnimations();
+					UserWidget->UpdateCanTick();
+
+					// Resetting the animation tick manager is ok here because TearDownAnimations will always clear out all animations
 					UserWidget->AnimationTickManager = nullptr;
 
 					WidgetIter.RemoveCurrent();
@@ -148,21 +165,21 @@ void UUMGSequenceTickManager::TickWidgetAnimations(float DeltaSeconds)
 			}
 			else if (!WidgetData.bIsTicking)
 			{
-				// If this widget has not told us it is ticking, and its last known state was
-				// ticking, we disable animations for that widget. Once it ticks again, the animation
-				// will be updated naturally, and doesn't need anything re-enabling.
+				// If this widget has not told us it is ticking, we disable animations for that widget.
+				// Once it ticks again, the animation will be updated naturally, and doesn't need anything re-enabling.
 				// 
 				// @todo: There is a chance that relative animations hitting this code path will resume with
 				// different relative bases due to the way the ecs data is destroyed and re-created.
 				// In order to fix this we would have to annex that data instead of destroying it.
 				if (!bIsCurrentlyEvaluating)
 				{
-					if (WidgetData.bLastKnownTickState)
-					{
-						UserWidget->DisableAnimations();
-					}
+					UserWidget->DisableAnimations();
 
-					WidgetData.bLastKnownTickState = false;
+					// Do not null out UUserWidget::AnimationTickManager because although we removed animation _data_
+					// the animations themselves are still playing. As such any UUMGSequencePlayers may hold a reference to this
+					// tick manager's linker, and therefore also need to keep this tick manager alive since the linker is not outered to this tick manager
+
+					WidgetIter.RemoveCurrent();
 				}
 			}
 			else
@@ -182,7 +199,6 @@ void UUMGSequenceTickManager::TickWidgetAnimations(float DeltaSeconds)
 
 				// Assume this widget will no longer tick, until we're told otherwise by way of OnWidgetTicked
 				WidgetData.bIsTicking = false;
-				WidgetData.bLastKnownTickState = true;
 			}
 		}
 	}
@@ -212,6 +228,8 @@ void UUMGSequenceTickManager::TickWidgetAnimations(float DeltaSeconds)
 			}
 		}
 	}
+
+	WeakUserWidgetData.Shrink();
 }
 
 void UUMGSequenceTickManager::ForceFlush()
diff --git a/Engine/Source/Runtime/UMG/Private/UserWidget.cpp b/Engine/Source/Runtime/UMG/Private/UserWidget.cpp
--- a/Engine/Source/Runtime/UMG/Private/UserWidget.cpp
+++ b/Engine/Source/Runtime/UMG/Private/UserWidget.cpp
@@ -415,9 +415,12 @@ UUMGSequencePlayer* UUserWidget::GetOrAddSequencePlayer(UWidgetAnimation* InAnim
 		if (!AnimationTickManager)
 		{
 			AnimationTickManager = UUMGSequenceTickManager::Get(this);
-			AnimationTickManager->AddWidget(this);
 		}
 
+		// Always ensure that this widget's animations are ticked for at least the first frame
+		// If this widget is currently offscreen it would very well not be being tracked by the tick manager
+		AnimationTickManager->AddWidget(this);
+
 		// @todo UMG sequencer - Restart animations which have had Play called on them?
 		UUMGSequencePlayer* FoundPlayer = nullptr;
 		for (UUMGSequencePlayer* Player : ActiveSequencePlayers)
diff --git a/Engine/Source/Runtime/UMG/Public/Animation/UMGSequenceTickManager.h b/Engine/Source/Runtime/UMG/Public/Animation/UMGSequenceTickManager.h
--- a/Engine/Source/Runtime/UMG/Public/Animation/UMGSequenceTickManager.h
+++ b/Engine/Source/Runtime/UMG/Public/Animation/UMGSequenceTickManager.h
@@ -26,9 +26,6 @@ struct FSequenceTickManagerWidgetData
 	/** The widget was ticked. */
 	bool bIsTicking = true;
 
-	/** The widget was ticked last update. */
-	bool bLastKnownTickState = true;
-
 	/** The widget's animation was ticked this frame. */
 	bool bActionsAndAnimationTicked = false;
 };
