Auto-cast the result of User Widget Extension utilities in Blueprints

diff --git a/Engine/Source/Runtime/UMG/Public/Blueprint/UserWidget.h b/Engine/Source/Runtime/UMG/Public/Blueprint/UserWidget.h
--- a/Engine/Source/Runtime/UMG/Public/Blueprint/UserWidget.h
+++ b/Engine/Source/Runtime/UMG/Public/Blueprint/UserWidget.h
@@ -1077,11 +1077,11 @@ public:
 	}
 
 	/** Find the first extension of the requested type. */
-	UFUNCTION(BlueprintCallable, Category = "User Interface|Extension")
+	UFUNCTION(BlueprintCallable, Category = "User Interface|Extension", Meta = (DeterminesOutputType = "ExtensionType"))
 	UUserWidgetExtension* GetExtension(TSubclassOf<UUserWidgetExtension> ExtensionType) const;
 
 	/** Find the extensions of the requested type. */
-	UFUNCTION(BlueprintCallable, Category = "User Interface|Extension")
+	UFUNCTION(BlueprintCallable, Category = "User Interface|Extension", Meta = (DeterminesOutputType = "ExtensionType"))
 	TArray<UUserWidgetExtension*> GetExtensions(TSubclassOf<UUserWidgetExtension> ExtensionType) const;
 
 	/** Add the extension of the requested type. */
@@ -1092,7 +1092,7 @@ public:
 	}
 
 	/** Add the extension of the requested type. */
-	UFUNCTION(BlueprintCallable, Category = "User Interface|Extension")
+	UFUNCTION(BlueprintCallable, Category = "User Interface|Extension", Meta = (DeterminesOutputType = "InExtensionType"))
 	UUserWidgetExtension* AddExtension(TSubclassOf<UUserWidgetExtension> InExtensionType);
 
 	/** Remove the extension. */
