Add Accessor for MenuAnchor, ExpandableArea and DynamicEntryBase

diff --git a/Engine/Source/Editor/UMGEditor/Private/Customizations/DynamicEntryBoxDetails.cpp b/Engine/Source/Editor/UMGEditor/Private/Customizations/DynamicEntryBoxDetails.cpp
--- a/Engine/Source/Editor/UMGEditor/Private/Customizations/DynamicEntryBoxDetails.cpp
+++ b/Engine/Source/Editor/UMGEditor/Private/Customizations/DynamicEntryBoxDetails.cpp
@@ -33,14 +33,14 @@ void FDynamicEntryBoxBaseDetails::CustomizeDetails(IDetailLayoutBuilder& DetailL
 	IDetailCategoryBuilder& EntryLayoutCategory = DetailLayout.EditCategory(TEXT("EntryLayout"));
 
 	const TAttribute<bool> CanEditAignmentAttribute(this, &FDynamicEntryBoxDetails::CanEditAlignment);
-	EntryLayoutCategory.AddProperty(DetailLayout.GetProperty(GET_MEMBER_NAME_CHECKED(UDynamicEntryBoxBase, EntryHorizontalAlignment)))
+	EntryLayoutCategory.AddProperty(DetailLayout.GetProperty(TEXT("EntryHorizontalAlignment")))
 		.IsEnabled(CanEditAignmentAttribute);
-	EntryLayoutCategory.AddProperty(DetailLayout.GetProperty(GET_MEMBER_NAME_CHECKED(UDynamicEntryBoxBase, EntryVerticalAlignment)))
+	EntryLayoutCategory.AddProperty(DetailLayout.GetProperty(TEXT("EntryVerticalAlignment")))
 		.IsEnabled(CanEditAignmentAttribute);
 
-	EntryLayoutCategory.AddProperty(DetailLayout.GetProperty(GET_MEMBER_NAME_CHECKED(UDynamicEntryBoxBase, MaxElementSize)))
+	EntryLayoutCategory.AddProperty(DetailLayout.GetProperty(TEXT("MaxElementSize")))
 		.IsEnabled(TAttribute<bool>(this, &FDynamicEntryBoxDetails::CanEditMaxElementSize));
-	EntryLayoutCategory.AddProperty(DetailLayout.GetProperty(GET_MEMBER_NAME_CHECKED(UDynamicEntryBoxBase, EntrySpacing)))
+	EntryLayoutCategory.AddProperty(DetailLayout.GetProperty(TEXT("EntrySpacing")))
 		.IsEnabled(TAttribute<bool>(this, &FDynamicEntryBoxDetails::CanEditEntrySpacing));
 	EntryLayoutCategory.AddProperty(DetailLayout.GetProperty(GET_MEMBER_NAME_CHECKED(UDynamicEntryBoxBase, SpacingPattern)))
 		.IsEnabled(TAttribute<bool>(this, &FDynamicEntryBoxDetails::CanEditSpacingPattern));
diff --git a/Engine/Source/Runtime/Slate/Private/Widgets/Layout/SExpandableArea.cpp b/Engine/Source/Runtime/Slate/Private/Widgets/Layout/SExpandableArea.cpp
--- a/Engine/Source/Runtime/Slate/Private/Widgets/Layout/SExpandableArea.cpp
+++ b/Engine/Source/Runtime/Slate/Private/Widgets/Layout/SExpandableArea.cpp
@@ -15,6 +15,7 @@ void SExpandableArea::Construct( const FArguments& InArgs )
 {
 	check(InArgs._Style);
 
+	WidgetStyle = InArgs._Style;
 	bAreaCollapsed = InArgs._InitiallyCollapsed;
 	MinWidth = InArgs._MinWidth;
 	MaxHeight = InArgs._MaxHeight;
@@ -52,7 +53,7 @@ void SExpandableArea::Construct( const FArguments& InArgs )
 	bAllowAnimatedTransition = InArgs._AllowAnimatedTransition;
 	ChildSlot
 	[
-		SNew( SBorder )
+		SAssignNew(WidgetBorder, SBorder )
 		.BorderImage( FullBorderImage )
 		.BorderBackgroundColor( FullBorderBackgroundColor )
 		.Padding(0.0f)
@@ -66,7 +67,7 @@ void SExpandableArea::Construct( const FArguments& InArgs )
 				.BorderBackgroundColor(TitleBorderBackgroundColor)
 				.Padding(0.0f)
 				[
-					SNew( SButton )
+					SAssignNew(ExpandingButton, SButton )
 					.Cursor(InArgs._HeaderCursor.IsSet() ? InArgs._HeaderCursor : GetCursor())
 					.ButtonStyle(FCoreStyle::Get(), "NoBorder")
 					.ContentPadding(InArgs._HeaderPadding)
@@ -80,7 +81,7 @@ void SExpandableArea::Construct( const FArguments& InArgs )
 			+ SVerticalBox::Slot()
 			.FillHeight(1.0f)
 			[
-				SNew(SBorder)
+				SAssignNew(BodyBorder, SBorder)
 				.Visibility(this, &SExpandableArea::OnGetContentVisibility)
 				.BorderImage(InArgs._BodyBorderImage)
 				.BorderBackgroundColor(InArgs._BodyBorderBackgroundColor)
@@ -92,6 +93,11 @@ void SExpandableArea::Construct( const FArguments& InArgs )
 			]
 		]
 	];
+
+	if (bBodyDiffers)
+	{
+		WidgetBorder = TitleBorder;
+	}
 }
 
 void SExpandableArea::SetExpanded( bool bExpanded )
@@ -220,4 +226,53 @@ bool SExpandableArea::IsTitleHovered() const
 	return TitleBorder.IsValid() && TitleBorder->IsHovered();
 }
 
+
+void SExpandableArea::SetMaxHeight(float InMaxHeight)
+{
+	if (MaxHeight != InMaxHeight)
+	{
+		MaxHeight = InMaxHeight;
+		Invalidate(EInvalidateWidgetReason::Layout);
+	}
+}
+
+void SExpandableArea::SetStyle(const FExpandableAreaStyle* InStyle)
+{
+	WidgetStyle = InStyle;
+}
+
+void SExpandableArea::SetBorderBrush(const FSlateBrush* InBorderBrush)
+{
+	check(WidgetBorder)
+	WidgetBorder->SetBorderImage(InBorderBrush);
+}
+
+void SExpandableArea::InvalidateStyle()
+{
+	Invalidate(EInvalidateWidgetReason::Layout);
+}
+
+void SExpandableArea::InvalidateBorderBrush()
+{
+	Invalidate(EInvalidateWidgetReason::Layout);
+}
+
+void SExpandableArea::SetBorderBackgroundColor(const FSlateColor& InBorderColor)
+{
+	check(WidgetBorder)
+	WidgetBorder->SetBorderBackgroundColor(InBorderColor);
+}
+
+void SExpandableArea::SetHeaderPadding(FMargin InHeaderPadding)
+{
+	check(ExpandingButton)
+	ExpandingButton->SetContentPadding(InHeaderPadding);
+}
+
+void SExpandableArea::SetAreaPadding(FMargin InAreaPadding)
+{
+	check(BodyBorder)
+	BodyBorder->SetPadding(InAreaPadding);
+}
+
 #undef LOCTEXT_NAMESPACE
diff --git a/Engine/Source/Runtime/Slate/Public/Widgets/Layout/SExpandableArea.h b/Engine/Source/Runtime/Slate/Public/Widgets/Layout/SExpandableArea.h
--- a/Engine/Source/Runtime/Slate/Public/Widgets/Layout/SExpandableArea.h
+++ b/Engine/Source/Runtime/Slate/Public/Widgets/Layout/SExpandableArea.h
@@ -18,6 +18,7 @@
 #include "Framework/SlateDelegates.h"
 
 class SBorder;
+class SButton;
 
 /**
  * Represents an expandable area of content            
@@ -110,6 +111,46 @@ public:
 	 */
 	virtual bool IsTitleHovered() const;
 
+	/**
+	 * Sets the max height of the area
+	 */
+	void SetMaxHeight(float InMaxHeight);
+
+	/**
+	 * Sets the color of all borders
+	 */
+	void SetBorderBackgroundColor(const FSlateColor& InBorderColor);
+
+	/**
+	 * Sets the header padding
+	 */
+	void SetHeaderPadding(FMargin InHeaderPadding);
+
+	/**
+	 * Sets the body padding
+	 */
+	void SetAreaPadding(FMargin InAreaPadding);
+
+	/**
+	 * Sets the style of the widget
+	 */
+	void SetStyle(const FExpandableAreaStyle* InStyle);
+
+	/**
+	 * Invalidates the style
+	 */
+	void InvalidateStyle();
+
+	/**
+	 * Invalidates the border brush
+	 */
+	void InvalidateBorderBrush();
+
+	/**
+	 * Sets the brush of all borders
+	 */
+	void SetBorderBrush(const FSlateBrush* InBorderBrush);
+
 protected:
 
 	/**
@@ -147,9 +188,21 @@ protected:
 
 protected:
 
+	/** Style of the Expandable Area */
+	const FExpandableAreaStyle* WidgetStyle;
+
 	/** Border widget for the header area */
 	TSharedPtr<SBorder> TitleBorder;
 
+	/** Border widget if the body style is different from header */
+	TSharedPtr<SBorder> WidgetBorder;
+
+	/** Border widget for the body area */
+	TSharedPtr<SBorder> BodyBorder;
+
+	/** Button for expanding the are */
+	TSharedPtr<SButton> ExpandingButton;
+
 	/** Curved used to simulate a rollout of the section */
 	FCurveSequence RolloutCurve;
 
@@ -173,4 +226,5 @@ protected:
 
 	/** If true expanding and collapsing the area will animate, otherwise it will expand/collapse instantly */
 	bool bAllowAnimatedTransition;
+
 };
diff --git a/Engine/Source/Runtime/UMG/Private/Components/DynamicEntryBoxBase.cpp b/Engine/Source/Runtime/UMG/Private/Components/DynamicEntryBoxBase.cpp
--- a/Engine/Source/Runtime/UMG/Private/Components/DynamicEntryBoxBase.cpp
+++ b/Engine/Source/Runtime/UMG/Private/Components/DynamicEntryBoxBase.cpp
@@ -20,7 +20,9 @@ UDynamicEntryBoxBase::UDynamicEntryBoxBase(const FObjectInitializer& Initializer
 	, EntryWidgetPool(*this)
 {
 	SetVisibilityInternal(ESlateVisibility::SelfHitTestInvisible);
+PRAGMA_DISABLE_DEPRECATION_WARNINGS
 	EntrySizeRule.SizeRule = ESlateSizeRule::Automatic;
+PRAGMA_ENABLE_DEPRECATION_WARNINGS
 }
 
 void UDynamicEntryBoxBase::ReleaseSlateResources(bool bReleaseChildren)
@@ -35,6 +37,7 @@ void UDynamicEntryBoxBase::ResetInternal(bool bDeleteWidgets)
 {
 	EntryWidgetPool.ReleaseAll(bDeleteWidgets);
 
+PRAGMA_DISABLE_DEPRECATION_WARNINGS
 	if (MyPanelWidget.IsValid())
 	{
 		switch (EntryBoxType)
@@ -55,6 +58,7 @@ void UDynamicEntryBoxBase::ResetInternal(bool bDeleteWidgets)
 			break;
 		}
 	}
+PRAGMA_ENABLE_DEPRECATION_WARNINGS
 }
 
 const TArray<UUserWidget*>& UDynamicEntryBoxBase::GetAllEntries() const
@@ -67,6 +71,27 @@ int32 UDynamicEntryBoxBase::GetNumEntries() const
 	return EntryWidgetPool.GetActiveWidgets().Num();
 }
 
+PRAGMA_DISABLE_DEPRECATION_WARNINGS
+EDynamicBoxType UDynamicEntryBoxBase::GetBoxType() const
+{
+	return EntryBoxType;
+}
+
+const FVector2D& UDynamicEntryBoxBase::GetEntrySpacing() const
+{
+	return EntrySpacing;
+}
+
+const FSlateChildSize& UDynamicEntryBoxBase::GetEntrySizeRule() const
+{
+	return EntrySizeRule;
+}
+
+const FRadialBoxSettings& UDynamicEntryBoxBase::GetRadialBoxSettings() const
+{
+	return RadialBoxSettings;
+}
+
 void UDynamicEntryBoxBase::RemoveEntryInternal(UUserWidget* EntryWidget)
 {
 	if (EntryWidget)
@@ -126,7 +151,6 @@ void UDynamicEntryBoxBase::SetEntrySpacing(const FVector2D& InEntrySpacing)
 						int32 PatternIdx = CountIdx % SpacingPattern.Num();
 						Spacing += SpacingPattern[PatternIdx];
 					}
-					
 					// Negative padding is no good, so negative spacing is expressed as positive spacing on the opposite side
 					if (Spacing.X >= 0.f)
 					{
@@ -181,7 +205,6 @@ void UDynamicEntryBoxBase::SetEntrySpacing(const FVector2D& InEntrySpacing)
 				FMargin Padding;
 				Padding.Top = bIsHBox || bIsFirstChild ? 0.f : EntrySpacing.Y;
 				Padding.Left = bIsHBox && !bIsFirstChild ? EntrySpacing.X : 0.f;
-
 				SBoxPanel::FSlot& BoxSlot = (*BoxChildren)[ChildIdx];
 				BoxSlot.SetPadding(Padding);
 			}
@@ -202,6 +225,23 @@ void UDynamicEntryBoxBase::SetRadialSettings(const FRadialBoxSettings& InSetting
 	}
 }
 
+EVerticalAlignment UDynamicEntryBoxBase::GetEntryVerticalAlignment() const
+{
+	return EntryVerticalAlignment.GetValue();
+}
+
+EHorizontalAlignment UDynamicEntryBoxBase::GetEntryHorizontalAlignment() const
+{
+	return EntryHorizontalAlignment.GetValue();
+}
+
+int32 UDynamicEntryBoxBase::GetMaxElementSize() const
+{
+	return MaxElementSize;
+}
+
+PRAGMA_ENABLE_DEPRECATION_WARNINGS
+
 #if WITH_EDITOR
 
 const FText UDynamicEntryBoxBase::GetPaletteCategory()
@@ -210,6 +250,7 @@ const FText UDynamicEntryBoxBase::GetPaletteCategory()
 }
 #endif
 
+PRAGMA_DISABLE_DEPRECATION_WARNINGS
 TSharedRef<SWidget> UDynamicEntryBoxBase::RebuildWidget()
 {
 	TSharedPtr<SWidget> EntryBoxWidget;
@@ -260,6 +301,7 @@ TSharedRef<SWidget> UDynamicEntryBoxBase::RebuildWidget()
 
 	return EntryBoxWidget.ToSharedRef();
 }
+PRAGMA_ENABLE_DEPRECATION_WARNINGS
 
 #if WITH_EDITOR
 void UDynamicEntryBoxBase::PostEditChangeProperty(FPropertyChangedEvent& PropertyChangedEvent)
@@ -273,6 +315,7 @@ void UDynamicEntryBoxBase::PostEditChangeProperty(FPropertyChangedEvent& Propert
 }
 #endif
 
+PRAGMA_DISABLE_DEPRECATION_WARNINGS
 void UDynamicEntryBoxBase::SynchronizeProperties()
 {
 	Super::SynchronizeProperties();
@@ -285,6 +328,7 @@ void UDynamicEntryBoxBase::SynchronizeProperties()
 	}
 #endif
 }
+PRAGMA_ENABLE_DEPRECATION_WARNINGS
 
 bool UDynamicEntryBoxBase::IsEntryClassValid(TSubclassOf<UUserWidget> InEntryClass) const
 {
@@ -367,6 +411,7 @@ FMargin UDynamicEntryBoxBase::BuildEntryPadding(const FVector2D& DesiredSpacing)
 	return EntryPadding;
 }
 
+PRAGMA_DISABLE_DEPRECATION_WARNINGS
 void UDynamicEntryBoxBase::AddEntryChild(UUserWidget& ChildWidget)
 {
 	if (EntryBoxType == EDynamicBoxType::Wrap || EntryBoxType == EDynamicBoxType::VerticalWrap)
@@ -459,4 +504,30 @@ void UDynamicEntryBoxBase::AddEntryChild(UUserWidget& ChildWidget)
 	}
 }
 
+void UDynamicEntryBoxBase::InitEntryBoxType(EDynamicBoxType InEntryBoxType)
+{
+	EntryBoxType = InEntryBoxType;
+}
+
+void UDynamicEntryBoxBase::InitEntrySizeRule(FSlateChildSize InEntrySizeRule)
+{
+	EntrySizeRule = InEntrySizeRule;
+}
+
+void UDynamicEntryBoxBase::InitEntryHorizontalAlignment(EHorizontalAlignment InEntryHorizontalAlignment)
+{
+	EntryHorizontalAlignment = InEntryHorizontalAlignment;
+}
+
+void UDynamicEntryBoxBase::InitEntryVerticalAlignment(EVerticalAlignment InEntryVerticalAlignment)
+{
+	EntryVerticalAlignment = InEntryVerticalAlignment;
+}
+
+void UDynamicEntryBoxBase::InitMaxElementSize(int32 InMaxElementSize)
+{
+	MaxElementSize = InMaxElementSize;
+}
+PRAGMA_ENABLE_DEPRECATION_WARNINGS
+
 #undef LOCTEXT_NAMESPACE
diff --git a/Engine/Source/Runtime/UMG/Private/Components/ExpandableArea.cpp b/Engine/Source/Runtime/UMG/Private/Components/ExpandableArea.cpp
--- a/Engine/Source/Runtime/UMG/Private/Components/ExpandableArea.cpp
+++ b/Engine/Source/Runtime/UMG/Private/Components/ExpandableArea.cpp
@@ -47,8 +47,10 @@ UExpandableArea::UExpandableArea(const FObjectInitializer& ObjectInitializer)
 		DefaultExpandableAreaBorderBrush->UnlinkColors();
 	}
 
+PRAGMA_DISABLE_DEPRECATION_WARNINGS
 	Style = *DefaultExpandableAreaStyle;
 	BorderBrush = *DefaultExpandableAreaBorderBrush;
+PRAGMA_ENABLE_DEPRECATION_WARNINGS
 
 #if WITH_EDITOR 
 	if (EditorExpandableAreaStyle == nullptr)
@@ -66,35 +68,37 @@ UExpandableArea::UExpandableArea(const FObjectInitializer& ObjectInitializer)
 		// Unlink UMG Editor colors from the editor settings colors.
 		EditorExpandableAreaBorderBrush->UnlinkColors();
 	}
-
 	if (IsEditorWidget())
 	{
+PRAGMA_DISABLE_DEPRECATION_WARNINGS
 		Style = *EditorExpandableAreaStyle;
 		BorderBrush = *EditorExpandableAreaBorderBrush;
-
+PRAGMA_ENABLE_DEPRECATION_WARNINGS
 		// The CDO isn't an editor widget and thus won't use the editor style, call post edit change to mark difference from CDO
 		PostEditChange();
 	}
 #endif // WITH_EDITOR
 
+PRAGMA_DISABLE_DEPRECATION_WARNINGS
 	BorderColor = FLinearColor::White;
 	AreaPadding = FMargin(1);
 	HeaderPadding = FMargin(4.0f, 2.0f);
+PRAGMA_ENABLE_DEPRECATION_WARNINGS
 }
 
+PRAGMA_DISABLE_DEPRECATION_WARNINGS
 bool UExpandableArea::GetIsExpanded() const
 {
-	if ( MyExpandableArea.IsValid() )
-	{
-		return MyExpandableArea->IsExpanded();
-	}
-
 	return bIsExpanded;
 }
 
 void UExpandableArea::SetIsExpanded(bool IsExpanded)
 {
-	bIsExpanded = IsExpanded;
+	if (bIsExpanded != IsExpanded)
+	{
+		bIsExpanded = IsExpanded;
+		BroadcastFieldValueChanged(FFieldNotificationClassDescriptor::bIsExpanded);
+	}
 	if ( MyExpandableArea.IsValid() )
 	{
 		MyExpandableArea->SetExpanded(IsExpanded);
@@ -103,7 +107,11 @@ void UExpandableArea::SetIsExpanded(bool IsExpanded)
 
 void UExpandableArea::SetIsExpanded_Animated(bool IsExpanded)
 {
-	bIsExpanded = IsExpanded;
+	if (bIsExpanded != IsExpanded)
+	{
+		bIsExpanded = IsExpanded;
+		BroadcastFieldValueChanged(FFieldNotificationClassDescriptor::bIsExpanded);
+	}
 	if (MyExpandableArea.IsValid())
 	{
 		MyExpandableArea->SetExpanded_Animated(IsExpanded);
@@ -162,6 +170,94 @@ void UExpandableArea::SetContentForSlot(FName SlotName, UWidget* Content)
 	}
 }
 
+const FExpandableAreaStyle& UExpandableArea::GetStyle() const
+{
+	return Style;
+}
+
+void UExpandableArea::SetStyle(const FExpandableAreaStyle& InStyle)
+{
+	Style = InStyle;
+	if (MyExpandableArea.IsValid())
+	{
+		MyExpandableArea->InvalidateStyle();
+	}
+}
+
+const FSlateBrush& UExpandableArea::GetBorderBrush() const
+{
+	return BorderBrush;
+}
+
+void UExpandableArea::SetBorderBrush(const FSlateBrush& InBorderBrush)
+{
+	if (BorderBrush != InBorderBrush)
+	{
+		BorderBrush = InBorderBrush;
+		if (MyExpandableArea.IsValid())
+		{
+			MyExpandableArea->SetBorderBrush(&InBorderBrush);
+			MyExpandableArea->InvalidateBorderBrush();
+		}
+	}
+}
+
+const FSlateColor& UExpandableArea::GetBorderColor() const
+{
+	return BorderColor;
+}
+
+void UExpandableArea::SetBorderColor(const FSlateColor& InBorderColor)
+{
+	BorderColor = InBorderColor;
+	if (MyExpandableArea.IsValid())
+	{
+		MyExpandableArea->SetBorderBackgroundColor(InBorderColor);
+	}
+}
+
+float UExpandableArea::GetMaxHeight() const
+{
+	return MaxHeight;
+}
+
+void UExpandableArea::SetMaxHeight(float InMaxHeight)
+{
+	MaxHeight = InMaxHeight;
+	if (MyExpandableArea.IsValid())
+	{
+		MyExpandableArea->SetMaxHeight(MaxHeight);
+	}
+}
+
+FMargin UExpandableArea::GetHeaderPadding() const
+{
+	return HeaderPadding;
+}
+
+void UExpandableArea::SetHeaderPadding(FMargin InHeaderPadding)
+{
+	HeaderPadding = InHeaderPadding;
+	if (MyExpandableArea.IsValid())
+	{
+		MyExpandableArea->SetHeaderPadding(HeaderPadding);
+	}
+}
+
+FMargin UExpandableArea::GetAreaPadding() const
+{
+	return AreaPadding;
+}
+
+void UExpandableArea::SetAreaPadding(FMargin InAreaPadding)
+{
+	AreaPadding = InAreaPadding;
+	if (MyExpandableArea.IsValid())
+	{
+		MyExpandableArea->SetAreaPadding(AreaPadding);
+	}
+}
+
 TSharedRef<SWidget> UExpandableArea::RebuildWidget()
 {
 	TSharedRef<SWidget> HeaderWidget = HeaderContent ? HeaderContent->TakeWidget() : SNullWidget::NullWidget;
@@ -191,13 +287,28 @@ void UExpandableArea::SynchronizeProperties()
 {
 	Super::SynchronizeProperties();
 
+	MyExpandableArea->SetStyle(&Style);
+	MyExpandableArea->InvalidateStyle();
 	MyExpandableArea->SetExpanded(bIsExpanded);
+	MyExpandableArea->SetAreaPadding(AreaPadding);
+	MyExpandableArea->SetHeaderPadding(HeaderPadding);
+	MyExpandableArea->SetBorderBrush(&BorderBrush);
+	MyExpandableArea->InvalidateBorderBrush();
+	MyExpandableArea->SetBorderBackgroundColor(BorderColor);
+	MyExpandableArea->SetMaxHeight(MaxHeight);
+
 }
+PRAGMA_ENABLE_DEPRECATION_WARNINGS
 
 void UExpandableArea::SlateExpansionChanged(bool NewState)
 {
-	bIsExpanded = NewState;
-
+PRAGMA_DISABLE_DEPRECATION_WARNINGS
+	if (bIsExpanded != NewState)
+	{
+		bIsExpanded = NewState;
+		BroadcastFieldValueChanged(FFieldNotificationClassDescriptor::bIsExpanded);
+	}
+PRAGMA_ENABLE_DEPRECATION_WARNINGS
 	if ( OnExpansionChanged.IsBound() )
 	{
 		OnExpansionChanged.Broadcast(this, NewState);
@@ -226,7 +337,9 @@ void UExpandableArea::OnDescendantDeselectedByDesigner(UWidget* DescendantWidget
 {
 	if ( MyExpandableArea.IsValid() )
 	{
+PRAGMA_DISABLE_DEPRECATION_WARNINGS
 		MyExpandableArea->SetExpanded(bIsExpanded);
+PRAGMA_ENABLE_DEPRECATION_WARNINGS
 	}
 }
 
diff --git a/Engine/Source/Runtime/UMG/Private/Components/MenuAnchor.cpp b/Engine/Source/Runtime/UMG/Private/Components/MenuAnchor.cpp
--- a/Engine/Source/Runtime/UMG/Private/Components/MenuAnchor.cpp
+++ b/Engine/Source/Runtime/UMG/Private/Components/MenuAnchor.cpp
@@ -18,8 +18,10 @@ UMenuAnchor::UMenuAnchor(const FObjectInitializer& ObjectInitializer)
 	, ShouldDeferPaintingAfterWindowContent(true)
 	, UseApplicationMenuStack(true)
 {
+PRAGMA_DISABLE_DEPRECATION_WARNINGS
 	Placement = MenuPlacement_ComboBox;
 	bFitInWindow = true;
+PRAGMA_ENABLE_DEPRECATION_WARNINGS
 }
 
 void UMenuAnchor::ReleaseSlateResources(bool bReleaseChildren)
@@ -31,6 +33,7 @@ void UMenuAnchor::ReleaseSlateResources(bool bReleaseChildren)
 
 TSharedRef<SWidget> UMenuAnchor::RebuildWidget()
 {
+PRAGMA_DISABLE_DEPRECATION_WARNINGS
 	MyMenuAnchor = SNew(SMenuAnchor)
 		.Placement(Placement)
 		.FitInWindow(bFitInWindow)
@@ -38,7 +41,7 @@ TSharedRef<SWidget> UMenuAnchor::RebuildWidget()
 		.OnMenuOpenChanged(BIND_UOBJECT_DELEGATE(FOnIsOpenChanged, HandleMenuOpenChanged))
 		.ShouldDeferPaintingAfterWindowContent(ShouldDeferPaintingAfterWindowContent)
 		.UseApplicationMenuStack(UseApplicationMenuStack);
-
+PRAGMA_ENABLE_DEPRECATION_WARNINGS
 	if ( GetChildrenCount() > 0 )
 	{
 		MyMenuAnchor->SetContent(GetContentSlot()->Content ? GetContentSlot()->Content->TakeWidget() : SNullWidget::NullWidget);
@@ -144,7 +147,9 @@ bool UMenuAnchor::IsOpen() const
 	return false;
 }
 
-void UMenuAnchor::SetPlacement(TEnumAsByte<EMenuPlacement> InPlacement)
+PRAGMA_DISABLE_DEPRECATION_WARNINGS
+
+void UMenuAnchor::SetPlacement(EMenuPlacement InPlacement)
 {
 	Placement = InPlacement;
 	if (MyMenuAnchor.IsValid())
@@ -153,6 +158,11 @@ void UMenuAnchor::SetPlacement(TEnumAsByte<EMenuPlacement> InPlacement)
 	}
 }
 
+EMenuPlacement UMenuAnchor::GetPlacement() const
+{
+	return Placement;
+}
+
 void UMenuAnchor::FitInWindow(bool bFit)
 {
 	bFitInWindow = bFit;
@@ -162,6 +172,33 @@ void UMenuAnchor::FitInWindow(bool bFit)
 	}
 }
 
+bool UMenuAnchor::IsFitInWindow() const
+{
+	return bFitInWindow;
+}
+
+bool UMenuAnchor::IsDeferPaintingAfterWindowContent() const
+{
+	return ShouldDeferPaintingAfterWindowContent;
+}
+
+bool UMenuAnchor::IsUseApplicationMenuStack() const
+{
+	return UseApplicationMenuStack;
+}
+
+void UMenuAnchor::InitShouldDeferPaintingAfterWindowContent(bool InShouldDeferPaintingAfterWindowContent)
+{
+	ShouldDeferPaintingAfterWindowContent = InShouldDeferPaintingAfterWindowContent;
+}
+
+void UMenuAnchor::InitUseApplicationMenuStack(bool InUseApplicationMenuStack)
+{
+	UseApplicationMenuStack = InUseApplicationMenuStack;
+}
+
+PRAGMA_ENABLE_DEPRECATION_WARNINGS
+
 bool UMenuAnchor::ShouldOpenDueToClick() const
 {
 	if ( MyMenuAnchor.IsValid() )
diff --git a/Engine/Source/Runtime/UMG/Public/Components/DynamicEntryBoxBase.h b/Engine/Source/Runtime/UMG/Public/Components/DynamicEntryBoxBase.h
--- a/Engine/Source/Runtime/UMG/Public/Components/DynamicEntryBoxBase.h
+++ b/Engine/Source/Runtime/UMG/Public/Components/DynamicEntryBoxBase.h
@@ -33,16 +33,66 @@
 {
 	GENERATED_BODY()
 
+protected:
+
+	UE_DEPRECATED(5.2, "Direct access to EntrySpacing is deprecated. Please use the getter or setter.")
+	/**
+	 * The padding to apply between entries in the box.
+	 * Note horizontal boxes only use the X and vertical boxes only use Y. Value is also ignored for the first entry in the box.
+	 * Wrap and Overlay types use both X and Y for spacing.
+	 */
+	UPROPERTY(EditAnywhere, BlueprintReadWrite, Getter, Setter, BlueprintSetter = "SetEntrySpacing", Category = EntryLayout)
+	FVector2D EntrySpacing;
+
+	//@todo DanH EntryBox: Consider giving a callback option as well/instead. Then this thing could actually create circular or pinwheel layouts...
+	/** The looping sequence of entry paddings to apply as entries are created. Overlay boxes only. Ignores EntrySpacing if not empty. */
+	UPROPERTY(EditAnywhere, BlueprintReadOnly, Category = EntryLayout)
+	TArray<FVector2D> SpacingPattern;
+
+	UE_DEPRECATED(5.2, "Direct access to EntryBoxType is deprecated. Please use the getter.")
+	/** The type of box panel into which created entries are added. Some differences in functionality exist between each type. */
+	UPROPERTY(EditAnywhere, BlueprintReadOnly, Getter = "GetBoxType", Category = DynamicEntryBox, meta = (DesignerRebuild))
+	EDynamicBoxType EntryBoxType;
+
+	UE_DEPRECATED(5.2, "Direct access to EntrySizeRule is deprecated. Please use the getter.")
+	/** Sizing rule to apply to generated entries. Horizontal/Vertical boxes only. */
+	UPROPERTY(EditAnywhere, BlueprintReadOnly, Getter, Category = EntryLayout)
+	FSlateChildSize EntrySizeRule;
+
+	UE_DEPRECATED(5.2, "Direct access to EntryHorizontalAlignment is deprecated. Please use the getter.")
+	/** Horizontal alignment of generated entries. Horizontal/Vertical/Wrap boxes only. */
+	UPROPERTY(EditAnywhere, BlueprintReadOnly, Getter, Category = EntryLayout)
+	TEnumAsByte<EHorizontalAlignment> EntryHorizontalAlignment;
+
+	UE_DEPRECATED(5.2, "Direct access to EntryVerticalAlignment is deprecated. Please use the getter.")
+	/** Vertical alignment of generated entries. Horizontal/Vertical/Wrap boxes only. */
+	UPROPERTY(EditAnywhere, BlueprintReadOnly, Getter, Category = EntryLayout)
+	TEnumAsByte<EVerticalAlignment> EntryVerticalAlignment;
+
+	UE_DEPRECATED(5.2, "Direct access to MaxElementSize is deprecated. Please use the getter.")
+	/** The maximum size of each entry in the dominant axis of the box. Vertical/Horizontal boxes only. */
+	UPROPERTY(EditAnywhere, BlueprintReadOnly, Getter, Category = EntryLayout)
+	int32 MaxElementSize = 0;
+
+	UE_DEPRECATED(5.2, "Direct access to EntrySpacing is deprecated. Please use the getter or setter.")
+	/** Settings only relevant to RadialBox */
+	UPROPERTY(EditAnywhere, BlueprintReadWrite, Getter, Setter = "SetRadialSettings", BlueprintSetter = "SetRadialSettings", Category = EntryLayout)
+	FRadialBoxSettings RadialBoxSettings;
+
 public:
 	UDynamicEntryBoxBase(const FObjectInitializer& Initializer);
 	virtual void ReleaseSlateResources(bool bReleaseChildren) override;
 
-	EDynamicBoxType GetBoxType() const { return EntryBoxType; }
-	const FVector2D& GetEntrySpacing() const { return EntrySpacing; }
+	EDynamicBoxType GetBoxType() const;
+	const FVector2D& GetEntrySpacing() const;
 
 	UFUNCTION(BlueprintCallable, Category = DynamicEntryBox)
 	const TArray<UUserWidget*>& GetAllEntries() const;
 
+	const FSlateChildSize& GetEntrySizeRule() const;
+
+	const FRadialBoxSettings& GetRadialBoxSettings() const;
+
 	template <typename EntryWidgetT = UUserWidget>
 	TArray<EntryWidgetT*> GetTypedEntries() const
 	{
@@ -65,7 +115,13 @@
 
 	UFUNCTION(BlueprintCallable, Category = DynamicEntryBox)
 	void SetRadialSettings(const FRadialBoxSettings& InSettings);
-	
+
+	EVerticalAlignment GetEntryVerticalAlignment() const;
+
+	EHorizontalAlignment GetEntryHorizontalAlignment() const;
+
+	int32 GetMaxElementSize() const;
+
 protected:
 	virtual TSharedRef<SWidget> RebuildWidget() override;
 	virtual void SynchronizeProperties() override;
@@ -99,44 +155,22 @@
 	virtual void PostEditChangeProperty(struct FPropertyChangedEvent& PropertyChangedEvent) override;
 #endif
 
+	// Initialize EntryBoxType in the constructor before the SWidget is constructed.
+	void InitEntryBoxType(EDynamicBoxType InEntryBoxType);
+
+	// Initialize EntrySizeRule in the constructor before the SWidget is constructed.
+	void InitEntrySizeRule(FSlateChildSize InEntrySizeRule);
+
+	// Initialize EntryHorizontalAlignment in the constructor before the SWidget is constructed.
+	void InitEntryHorizontalAlignment(EHorizontalAlignment InEntryHorizontalAlignment);
+
+	// Initialize EntryVerticalAlignment in the constructor before the SWidget is constructed.
+	void InitEntryVerticalAlignment(EVerticalAlignment InEntryVerticalAlignment);
+
+	// Initialize MaxElementSize in the constructor before the SWidget is constructed.
+	void InitMaxElementSize(int32 InMaxElementSize);
+
 protected:
-	/** The type of box panel into which created entries are added. Some differences in functionality exist between each type. */
-	UPROPERTY(EditAnywhere, Category = DynamicEntryBox, meta = (DesignerRebuild))
-	EDynamicBoxType EntryBoxType;
-
-	/** 
-	 * The padding to apply between entries in the box.
-	 * Note horizontal boxes only use the X and vertical boxes only use Y. Value is also ignored for the first entry in the box.
-	 * Wrap and Overlay types use both X and Y for spacing.
-	 */
-	UPROPERTY(EditAnywhere, BlueprintReadOnly, Category = EntryLayout)
-	FVector2D EntrySpacing;
-
-	//@todo DanH EntryBox: Consider giving a callback option as well/instead. Then this thing could actually create circular or pinwheel layouts...
-	/** The looping sequence of entry paddings to apply as entries are created. Overlay boxes only. Ignores EntrySpacing if not empty. */
-	UPROPERTY(EditAnywhere, BlueprintReadOnly, Category = EntryLayout)
-	TArray<FVector2D> SpacingPattern;
-
-	/** Sizing rule to apply to generated entries. Horizontal/Vertical boxes only. */
-	UPROPERTY(EditAnywhere, BlueprintReadOnly, Category = EntryLayout)
-	FSlateChildSize EntrySizeRule;
-
-	/** Horizontal alignment of generated entries. Horizontal/Vertical/Wrap boxes only. */
-	UPROPERTY(EditAnywhere, BlueprintReadOnly, Category = EntryLayout)
-	TEnumAsByte<EHorizontalAlignment> EntryHorizontalAlignment;
-
-	/** Vertical alignment of generated entries. Horizontal/Vertical/Wrap boxes only. */
-	UPROPERTY(EditAnywhere, BlueprintReadOnly, Category = EntryLayout)
-	TEnumAsByte<EVerticalAlignment> EntryVerticalAlignment;
-
-	/** The maximum size of each entry in the dominant axis of the box. Vertical/Horizontal boxes only. */
-	UPROPERTY(EditAnywhere, BlueprintReadOnly, Category = EntryLayout)
-	int32 MaxElementSize = 0;
-
-	/** Settings only relevant to RadialBox */
-	UPROPERTY(EditAnywhere, BlueprintReadOnly, Category = EntryLayout)
-	FRadialBoxSettings RadialBoxSettings;
-
 	// Can be a horizontal, vertical, wrap box, or overlay
 	TSharedPtr<SPanel> MyPanelWidget;
 
diff --git a/Engine/Source/Runtime/UMG/Public/Components/ExpandableArea.h b/Engine/Source/Runtime/UMG/Public/Components/ExpandableArea.h
--- a/Engine/Source/Runtime/UMG/Public/Components/ExpandableArea.h
+++ b/Engine/Source/Runtime/UMG/Public/Components/ExpandableArea.h
@@ -28,28 +28,35 @@ class UMG_API UExpandableArea : public UWidget, public INamedSlotInterface
 
 public:
 
+	UE_DEPRECATED(5.2, "Direct access to Style is deprecated. Please use SetStyle or GetStyle.")
 	UPROPERTY(EditAnywhere, Category = "Style")
 	FExpandableAreaStyle Style;
 
-	UPROPERTY( EditAnywhere, BlueprintReadOnly, Category = "Style" )
+	UE_DEPRECATED(5.2, "Direct access to BorderBrush is deprecated. Please use the getter or setter.")
+	UPROPERTY( EditAnywhere, BlueprintReadWrite, Getter, Setter, Category = "Style" )
 	FSlateBrush BorderBrush;
 
-	UPROPERTY( EditAnywhere, BlueprintReadOnly, Category = "Style" )
+	UE_DEPRECATED(5.2, "Direct access to BorderColor is deprecated. Please use the getter or setter.")
+	UPROPERTY( EditAnywhere, BlueprintReadWrite, Getter, Setter, Category = "Style")
 	FSlateColor BorderColor;
 
 	/**  */
-	UPROPERTY(EditAnywhere, BlueprintReadOnly, Category="Expansion")
+	UE_DEPRECATED(5.2, "Direct access to bIsExpanded is deprecated. Please use the getter or setter.")
+	UPROPERTY(EditAnywhere, BlueprintReadWrite, Getter = "GetIsExpanded", Setter = "SetIsExpanded", BlueprintGetter = "GetIsExpanded", BlueprintSetter = "SetIsExpanded", FieldNotify, Category = "Expansion")
 	bool bIsExpanded;
 
+	UE_DEPRECATED(5.2, "Direct access to MaxHeight is deprecated. Please use the getter or setter.")
 	/** The maximum height of the area */
-	UPROPERTY(EditAnywhere, BlueprintReadOnly, Category="Expansion")
+	UPROPERTY(EditAnywhere, BlueprintReadWrite, Getter, Setter, Category="Expansion")
 	float MaxHeight;
 	
-	UPROPERTY( EditAnywhere, BlueprintReadOnly, Category = "Expansion" )
+	UE_DEPRECATED(5.2, "Direct access to HeaderPadding is deprecated. Please use the getter or setter.")
+	UPROPERTY( EditAnywhere, BlueprintReadWrite, Getter, Setter, Category = "Expansion" )
 	FMargin HeaderPadding;
 	
 	/**  */
-	UPROPERTY(EditAnywhere, BlueprintReadOnly, Category="Expansion")
+	UE_DEPRECATED(5.2, "Direct access to AreaPadding is deprecated. Please use the getter or setter.")
+	UPROPERTY(EditAnywhere, BlueprintReadWrite, Getter, Setter, Category="Expansion")
 	FMargin AreaPadding;
 
 	/** A bindable delegate for the IsChecked. */
@@ -67,6 +74,29 @@ public:
 	UFUNCTION(BlueprintCallable, Category = "Expansion")
 	void SetIsExpanded_Animated(bool IsExpanded);
 	
+	const FExpandableAreaStyle& GetStyle() const;
+
+	void SetStyle(const FExpandableAreaStyle& InStyle);
+
+	const FSlateBrush& GetBorderBrush() const;
+
+	void SetBorderBrush(const FSlateBrush& InBorderBrush);
+
+	const FSlateColor& GetBorderColor() const;
+
+	void SetBorderColor(const FSlateColor& InBorderColor);
+
+	float GetMaxHeight() const;
+
+	void SetMaxHeight(float InMaxHeight);
+
+	FMargin GetHeaderPadding() const;
+
+	void SetHeaderPadding(FMargin InHeaderPadding);
+
+	FMargin GetAreaPadding() const;
+	void SetAreaPadding(FMargin InAreaPadding);
+
 	// Begin INamedSlotInterface
 	virtual void GetSlotNames(TArray<FName>& SlotNames) const override;
 	virtual UWidget* GetContentForSlot(FName SlotName) const override;
diff --git a/Engine/Source/Runtime/UMG/Public/Components/MenuAnchor.h b/Engine/Source/Runtime/UMG/Public/Components/MenuAnchor.h
--- a/Engine/Source/Runtime/UMG/Public/Components/MenuAnchor.h
+++ b/Engine/Source/Runtime/UMG/Public/Components/MenuAnchor.h
@@ -45,19 +45,23 @@ public:
 	UPROPERTY(EditAnywhere, Category="Events")
 	FGetUserWidget OnGetUserMenuContentEvent;
 	
+	UE_DEPRECATED(5.2, "Direct access to Placement is deprecated. Please use the getter or setter.")
 	/** The placement location of the summoned widget. */
-	UPROPERTY(EditAnywhere, BlueprintReadOnly, Category="Menu Anchor")
+	UPROPERTY(EditAnywhere, BlueprintReadWrite, Setter, Getter, BlueprintSetter = "SetPlacement", Category = "Menu Anchor")
 	TEnumAsByte<EMenuPlacement> Placement;
 
+	UE_DEPRECATED(5.2, "Direct access to bFitInWindow is deprecated. Please use the getter or setter.")
 	/** Should the menu anchor attempt to fit the menu inside the window. */
-	UPROPERTY(EditAnywhere, BlueprintReadOnly, Category = "Menu Anchor", meta=(ScriptName="ShouldFitInWindow"))
+	UPROPERTY(EditAnywhere, BlueprintReadWrite, Setter = "FitInWindow", Getter = "IsFitInWindow", BlueprintSetter = "FitInWindow", Category = "Menu Anchor", meta = (ScriptName = "ShouldFitInWindow"))
 	bool bFitInWindow;
 
-	UPROPERTY(EditAnywhere, BlueprintReadOnly, AdvancedDisplay, Category = "Menu Anchor")
+	UE_DEPRECATED(5.2, "Direct access to ShouldDeferPaintingAfterWindowContent is deprecated. Please use the getter.")
+	UPROPERTY(EditAnywhere, BlueprintReadOnly, Getter = "IsDeferPaintingAfterWindowContent", AdvancedDisplay, Category = "Menu Anchor")
 	bool ShouldDeferPaintingAfterWindowContent;
 
+	UE_DEPRECATED(5.2, "Direct access to UseApplicationMenuStack is deprecated. Please use the getter.")
 	/** Does this menu behave like a normal stacked menu? Set it to false to control the menu's lifetime yourself. */
-	UPROPERTY(EditAnywhere, BlueprintReadOnly, AdvancedDisplay, Category = "Menu Anchor")
+	UPROPERTY(EditAnywhere, BlueprintReadOnly, Getter = "IsUseApplicationMenuStack", AdvancedDisplay, Category = "Menu Anchor")
 	bool UseApplicationMenuStack;
 
 public:
@@ -69,11 +73,19 @@ public:
 	//TODO UMG Add Set MenuClass
 
 	UFUNCTION(BlueprintCallable, Category = "Menu Anchor")
-	void SetPlacement(TEnumAsByte<EMenuPlacement> InPlacement);
+	void SetPlacement(EMenuPlacement InPlacement);
+
+	EMenuPlacement GetPlacement() const;
 
 	UFUNCTION(BlueprintCallable, Category = "Menu Anchor")
 	void FitInWindow(bool bFit);
 
+	bool IsFitInWindow() const;
+
+	bool IsDeferPaintingAfterWindowContent() const;
+
+	bool IsUseApplicationMenuStack() const;
+
 public:
 
 	/**
@@ -129,6 +141,13 @@ protected:
 	virtual TSharedRef<SWidget> RebuildWidget() override;
 	// End of UWidget interface
 
+	// Initialize ShouldDeferPaintingAfterWindowContent in the constructor before the SWidget is constructed.
+	void InitShouldDeferPaintingAfterWindowContent(bool InShouldDeferPaintingAfterWindowContent);
+
+	// Initialize UseApplicationMenuStack in the constructor before the SWidget is constructed.
+	void InitUseApplicationMenuStack(bool InUseApplicationMenuStack);
+
+
 protected:
 	TSharedRef<SWidget> HandleGetMenuContent();
 	void HandleMenuOpenChanged(bool bIsOpen);
