Adding an error message and prevent a crash when WidgetComponent being updated before PlayerController is set.

diff --git a/Engine/Source/Runtime/UMG/Private/Components/WidgetInteractionComponent.cpp b/Engine/Source/Runtime/UMG/Private/Components/WidgetInteractionComponent.cpp
--- a/Engine/Source/Runtime/UMG/Private/Components/WidgetInteractionComponent.cpp
+++ b/Engine/Source/Runtime/UMG/Private/Components/WidgetInteractionComponent.cpp
@@ -1,6 +1,7 @@
 // Copyright Epic Games, Inc. All Rights Reserved.
 
 #include "Components/WidgetInteractionComponent.h"
+#include "UMGPrivate.h"
 #include "CollisionQueryParams.h"
 #include "Components/PrimitiveComponent.h"
 #include "Engine/GameViewportClient.h"
@@ -162,7 +163,13 @@ UWidgetInteractionComponent::FWidgetTraceResult UWidgetInteractionComponent::Per
 			FCollisionQueryParams Params(SCENE_QUERY_STAT(WidgetInteractionComponentTrace));
 			Params.AddIgnoredComponents(PrimitiveChildren);
 
-			APlayerController* PlayerController = GetWorld()->GetFirstPlayerController();
+			const UWorld* World = GetWorld();
+			APlayerController* PlayerController = World ? World->GetFirstPlayerController():nullptr;
+			if (!PlayerController)
+			{
+				UE_LOG(LogUMG, Warning, TEXT("Widget Interaction Component cannot perform trace without a valid PlayerController."));
+				return FWidgetTraceResult();
+			}
 			ULocalPlayer* LocalPlayer = PlayerController->GetLocalPlayer();
 			
 			if ( LocalPlayer && LocalPlayer->ViewportClient )
