Check if SWidget exists before attempting to synchronize

diff --git a/Engine/Source/Runtime/AdvancedWidgets/Private/Components/RadialSlider.cpp b/Engine/Source/Runtime/AdvancedWidgets/Private/Components/RadialSlider.cpp
--- a/Engine/Source/Runtime/AdvancedWidgets/Private/Components/RadialSlider.cpp
+++ b/Engine/Source/Runtime/AdvancedWidgets/Private/Components/RadialSlider.cpp
@@ -74,6 +74,11 @@ void URadialSlider::SynchronizeProperties()
 {
 	Super::SynchronizeProperties();
 
+	if (!MyRadialSlider.IsValid())
+	{
+		return;
+	}
+
 	FRichCurve* OurCurve = SliderRange.GetRichCurve();
 	for (int32 i = 0; i < OurCurve->Keys.Num() - 1; i++)
 	{
diff --git a/Engine/Source/Runtime/UMG/Private/Components/Border.cpp b/Engine/Source/Runtime/UMG/Private/Components/Border.cpp
--- a/Engine/Source/Runtime/UMG/Private/Components/Border.cpp
+++ b/Engine/Source/Runtime/UMG/Private/Components/Border.cpp
@@ -62,6 +62,11 @@ TSharedRef<SWidget> UBorder::RebuildWidget()
 void UBorder::SynchronizeProperties()
 {
 	Super::SynchronizeProperties();
+
+	if (!MyBorder.IsValid())
+	{
+		return;
+	}
 	
 PRAGMA_DISABLE_DEPRECATION_WARNINGS
 	TAttribute<FLinearColor> ContentColorAndOpacityBinding = PROPERTY_BINDING(FLinearColor, ContentColorAndOpacity);
diff --git a/Engine/Source/Runtime/UMG/Private/Components/Button.cpp b/Engine/Source/Runtime/UMG/Private/Components/Button.cpp
--- a/Engine/Source/Runtime/UMG/Private/Components/Button.cpp
+++ b/Engine/Source/Runtime/UMG/Private/Components/Button.cpp
@@ -100,6 +100,11 @@ void UButton::SynchronizeProperties()
 {
 	Super::SynchronizeProperties();
 
+	if (!MyButton.IsValid())
+	{
+		return;
+	}
+
 	MyButton->SetColorAndOpacity( ColorAndOpacity );
 	MyButton->SetBorderBackgroundColor( BackgroundColor );
 }
diff --git a/Engine/Source/Runtime/UMG/Private/Components/CheckBox.cpp b/Engine/Source/Runtime/UMG/Private/Components/CheckBox.cpp
--- a/Engine/Source/Runtime/UMG/Private/Components/CheckBox.cpp
+++ b/Engine/Source/Runtime/UMG/Private/Components/CheckBox.cpp
@@ -101,6 +101,10 @@ void UCheckBox::SynchronizeProperties()
 {
 	Super::SynchronizeProperties();
 
+	if (!MyCheckbox.IsValid())
+	{
+		return;
+	}
 
 PRAGMA_DISABLE_DEPRECATION_WARNINGS
 	MyCheckbox->SetStyle(&WidgetStyle);
diff --git a/Engine/Source/Runtime/UMG/Private/Components/CircularThrobber.cpp b/Engine/Source/Runtime/UMG/Private/Components/CircularThrobber.cpp
--- a/Engine/Source/Runtime/UMG/Private/Components/CircularThrobber.cpp
+++ b/Engine/Source/Runtime/UMG/Private/Components/CircularThrobber.cpp
@@ -83,6 +83,11 @@ void UCircularThrobber::SynchronizeProperties()
 {
 	Super::SynchronizeProperties();
 
+	if (!MyCircularThrobber.IsValid())
+	{
+		return;
+	}
+
 PRAGMA_DISABLE_DEPRECATION_WARNINGS
 	MyCircularThrobber->SetNumPieces(FMath::Clamp(NumberOfPieces, 1, 25));
 	MyCircularThrobber->SetPeriod(FMath::Max(Period, SCircularThrobber::MinimumPeriodValue));
diff --git a/Engine/Source/Runtime/UMG/Private/Components/EditableText.cpp b/Engine/Source/Runtime/UMG/Private/Components/EditableText.cpp
--- a/Engine/Source/Runtime/UMG/Private/Components/EditableText.cpp
+++ b/Engine/Source/Runtime/UMG/Private/Components/EditableText.cpp
@@ -122,6 +122,10 @@ void UEditableText::SynchronizeProperties()
 {
 	Super::SynchronizeProperties();
 
+	if (!MyEditableText.IsValid())
+	{
+		return;
+	}
 
 	PRAGMA_DISABLE_DEPRECATION_WARNINGS
 
diff --git a/Engine/Source/Runtime/UMG/Private/Components/EditableTextBox.cpp b/Engine/Source/Runtime/UMG/Private/Components/EditableTextBox.cpp
--- a/Engine/Source/Runtime/UMG/Private/Components/EditableTextBox.cpp
+++ b/Engine/Source/Runtime/UMG/Private/Components/EditableTextBox.cpp
@@ -141,6 +141,11 @@ void UEditableTextBox::SynchronizeProperties()
 {
 	Super::SynchronizeProperties();
 
+	if (!MyEditableTextBlock.IsValid())
+	{
+		return;
+	}
+
 	TAttribute<FText> TextBinding = PROPERTY_BINDING(FText, Text);
 	TAttribute<FText> HintTextBinding = PROPERTY_BINDING(FText, HintText);
 
diff --git a/Engine/Source/Runtime/UMG/Private/Components/ExpandableArea.cpp b/Engine/Source/Runtime/UMG/Private/Components/ExpandableArea.cpp
--- a/Engine/Source/Runtime/UMG/Private/Components/ExpandableArea.cpp
+++ b/Engine/Source/Runtime/UMG/Private/Components/ExpandableArea.cpp
@@ -287,6 +287,11 @@ void UExpandableArea::SynchronizeProperties()
 {
 	Super::SynchronizeProperties();
 
+	if (!MyExpandableArea.IsValid())
+	{
+		return;
+	}
+	
 	MyExpandableArea->SetStyle(&Style);
 	MyExpandableArea->InvalidateStyle();
 	MyExpandableArea->SetExpanded(bIsExpanded);
diff --git a/Engine/Source/Runtime/UMG/Private/Components/GridPanel.cpp b/Engine/Source/Runtime/UMG/Private/Components/GridPanel.cpp
--- a/Engine/Source/Runtime/UMG/Private/Components/GridPanel.cpp
+++ b/Engine/Source/Runtime/UMG/Private/Components/GridPanel.cpp
@@ -87,6 +87,11 @@ void UGridPanel::SynchronizeProperties()
 {
 	Super::SynchronizeProperties();
 
+	if (!MyGridPanel.IsValid())
+	{
+		return;
+	}
+
 	MyGridPanel->ClearFill();
 
 	for ( int ColumnIndex = 0; ColumnIndex < ColumnFill.Num(); ColumnIndex++ )
diff --git a/Engine/Source/Runtime/UMG/Private/Components/InputKeySelector.cpp b/Engine/Source/Runtime/UMG/Private/Components/InputKeySelector.cpp
--- a/Engine/Source/Runtime/UMG/Private/Components/InputKeySelector.cpp
+++ b/Engine/Source/Runtime/UMG/Private/Components/InputKeySelector.cpp
@@ -246,6 +246,11 @@ void UInputKeySelector::SynchronizeProperties()
 {
 	Super::SynchronizeProperties();
 
+	if (!MyInputKeySelector.IsValid())
+	{
+		return;
+	}
+
 	MyInputKeySelector->SetSelectedKey( SelectedKey );
 	MyInputKeySelector->SetMargin( Margin );
 	MyInputKeySelector->SetButtonStyle( &WidgetStyle );
diff --git a/Engine/Source/Runtime/UMG/Private/Components/MultiLineEditableText.cpp b/Engine/Source/Runtime/UMG/Private/Components/MultiLineEditableText.cpp
--- a/Engine/Source/Runtime/UMG/Private/Components/MultiLineEditableText.cpp
+++ b/Engine/Source/Runtime/UMG/Private/Components/MultiLineEditableText.cpp
@@ -111,6 +111,11 @@ void UMultiLineEditableText::SynchronizeProperties()
 {
 	Super::SynchronizeProperties();
 
+	if (!MyMultiLineEditableText.IsValid())
+	{
+		return;
+	}
+
 	TAttribute<FText> HintTextBinding = PROPERTY_BINDING(FText, HintText);
 
 	MyMultiLineEditableText->SetTextStyle(&WidgetStyle);
diff --git a/Engine/Source/Runtime/UMG/Private/Components/MultiLineEditableTextBox.cpp b/Engine/Source/Runtime/UMG/Private/Components/MultiLineEditableTextBox.cpp
--- a/Engine/Source/Runtime/UMG/Private/Components/MultiLineEditableTextBox.cpp
+++ b/Engine/Source/Runtime/UMG/Private/Components/MultiLineEditableTextBox.cpp
@@ -124,6 +124,11 @@ void UMultiLineEditableTextBox::SynchronizeProperties()
 {
 	Super::SynchronizeProperties();
 
+	if (!MyEditableTextBlock.IsValid())
+	{
+		return;
+	}
+
 	TAttribute<FText> HintTextBinding = PROPERTY_BINDING(FText, HintText);
 
 	MyEditableTextBlock->SetStyle(&WidgetStyle);
diff --git a/Engine/Source/Runtime/UMG/Private/Components/ProgressBar.cpp b/Engine/Source/Runtime/UMG/Private/Components/ProgressBar.cpp
--- a/Engine/Source/Runtime/UMG/Private/Components/ProgressBar.cpp
+++ b/Engine/Source/Runtime/UMG/Private/Components/ProgressBar.cpp
@@ -79,6 +79,11 @@ void UProgressBar::SynchronizeProperties()
 {
 	Super::SynchronizeProperties();
 
+	if (!MyProgressBar.IsValid())
+	{
+		return;
+	}
+
 	TAttribute< TOptional<float> > PercentBinding = OPTIONAL_BINDING_CONVERT(float, Percent, TOptional<float>, ConvertFloatToOptionalFloat);
 	TAttribute<FSlateColor> FillColorAndOpacityBinding = PROPERTY_BINDING(FSlateColor, FillColorAndOpacity);
 
diff --git a/Engine/Source/Runtime/UMG/Private/Components/RetainerBox.cpp b/Engine/Source/Runtime/UMG/Private/Components/RetainerBox.cpp
--- a/Engine/Source/Runtime/UMG/Private/Components/RetainerBox.cpp
+++ b/Engine/Source/Runtime/UMG/Private/Components/RetainerBox.cpp
@@ -164,6 +164,11 @@ void URetainerBox::SynchronizeProperties()
 {
 	Super::SynchronizeProperties();
 
+	if (!MyRetainerWidget.IsValid())
+	{
+		return;
+	}
+	
 PRAGMA_DISABLE_DEPRECATION_WARNINGS
 	MyRetainerWidget->SetRetainedRendering(IsDesignTime() ? false : bRetainRender);
 	MyRetainerWidget->SetEffectMaterial(EffectMaterial);
diff --git a/Engine/Source/Runtime/UMG/Private/Components/ScaleBox.cpp b/Engine/Source/Runtime/UMG/Private/Components/ScaleBox.cpp
--- a/Engine/Source/Runtime/UMG/Private/Components/ScaleBox.cpp
+++ b/Engine/Source/Runtime/UMG/Private/Components/ScaleBox.cpp
@@ -110,6 +110,11 @@ void UScaleBox::SynchronizeProperties()
 {
 	Super::SynchronizeProperties();
 
+	if (!MyScaleBox.IsValid())
+	{
+		return;
+	}
+	
 PRAGMA_DISABLE_DEPRECATION_WARNINGS
 	MyScaleBox->SetStretchDirection(StretchDirection);
 	MyScaleBox->SetStretch(Stretch);
diff --git a/Engine/Source/Runtime/UMG/Private/Components/ScrollBox.cpp b/Engine/Source/Runtime/UMG/Private/Components/ScrollBox.cpp
--- a/Engine/Source/Runtime/UMG/Private/Components/ScrollBox.cpp
+++ b/Engine/Source/Runtime/UMG/Private/Components/ScrollBox.cpp
@@ -157,6 +157,11 @@ void UScrollBox::SynchronizeProperties()
 {
 	Super::SynchronizeProperties();
 
+	if (!MyScrollBox.IsValid())
+	{
+		return;
+	}
+
 	MyScrollBox->SetScrollOffset(DesiredScrollOffset);
 	MyScrollBox->SetOrientation(Orientation);
 	MyScrollBox->SetScrollBarVisibility(UWidget::ConvertSerializedVisibilityToRuntime(ScrollBarVisibility));
diff --git a/Engine/Source/Runtime/UMG/Private/Components/Slider.cpp b/Engine/Source/Runtime/UMG/Private/Components/Slider.cpp
--- a/Engine/Source/Runtime/UMG/Private/Components/Slider.cpp
+++ b/Engine/Source/Runtime/UMG/Private/Components/Slider.cpp
@@ -88,6 +88,11 @@ void USlider::SynchronizeProperties()
 {
 	Super::SynchronizeProperties();
 
+	if (!MySlider.IsValid())
+	{
+		return;
+	}
+
 PRAGMA_DISABLE_DEPRECATION_WARNINGS
 	TAttribute<float> ValueBinding = PROPERTY_BINDING(float, Value);
 	
diff --git a/Engine/Source/Runtime/UMG/Private/Components/Spacer.cpp b/Engine/Source/Runtime/UMG/Private/Components/Spacer.cpp
--- a/Engine/Source/Runtime/UMG/Private/Components/Spacer.cpp
+++ b/Engine/Source/Runtime/UMG/Private/Components/Spacer.cpp
@@ -60,6 +60,11 @@ TSharedRef<SWidget> USpacer::RebuildWidget()
 void USpacer::SynchronizeProperties()
 {
 	Super::SynchronizeProperties();
+
+	if (!MySpacer.IsValid())
+	{
+		return;
+	}
 	
 	MySpacer->SetSize(Size);
 }
diff --git a/Engine/Source/Runtime/UMG/Private/Components/SpinBox.cpp b/Engine/Source/Runtime/UMG/Private/Components/SpinBox.cpp
--- a/Engine/Source/Runtime/UMG/Private/Components/SpinBox.cpp
+++ b/Engine/Source/Runtime/UMG/Private/Components/SpinBox.cpp
@@ -109,6 +109,12 @@ PRAGMA_ENABLE_DEPRECATION_WARNINGS
 void USpinBox::SynchronizeProperties()
 {
 	Super::SynchronizeProperties();
+	
+	if (!MySpinBox.IsValid())
+	{
+		return;
+	}
+	
 PRAGMA_DISABLE_DEPRECATION_WARNINGS
 	MySpinBox->SetDelta(Delta);
 	MySpinBox->SetSliderExponent(SliderExponent);
diff --git a/Engine/Source/Runtime/UMG/Private/Components/StackBox.cpp b/Engine/Source/Runtime/UMG/Private/Components/StackBox.cpp
--- a/Engine/Source/Runtime/UMG/Private/Components/StackBox.cpp
+++ b/Engine/Source/Runtime/UMG/Private/Components/StackBox.cpp
@@ -107,6 +107,12 @@ TSharedRef<SWidget> UStackBox::RebuildWidget()
 void UStackBox::SynchronizeProperties()
 {
 	Super::SynchronizeProperties();
+
+	if (!MyBox.IsValid())
+	{
+		return;
+	}
+
 	MyBox->SetOrientation(Orientation);
 }
 
diff --git a/Engine/Source/Runtime/UMG/Private/Components/Throbber.cpp b/Engine/Source/Runtime/UMG/Private/Components/Throbber.cpp
--- a/Engine/Source/Runtime/UMG/Private/Components/Throbber.cpp
+++ b/Engine/Source/Runtime/UMG/Private/Components/Throbber.cpp
@@ -81,6 +81,11 @@ void UThrobber::SynchronizeProperties()
 {
 	Super::SynchronizeProperties();
 
+	if (!MyThrobber.IsValid())
+	{
+		return;
+	}
+
 PRAGMA_DISABLE_DEPRECATION_WARNINGS
 	MyThrobber->SetPieceImage(&Image);
 	MyThrobber->InvalidatePieceImage();
diff --git a/Engine/Source/Runtime/UMG/Private/Components/UniformGridPanel.cpp b/Engine/Source/Runtime/UMG/Private/Components/UniformGridPanel.cpp
--- a/Engine/Source/Runtime/UMG/Private/Components/UniformGridPanel.cpp
+++ b/Engine/Source/Runtime/UMG/Private/Components/UniformGridPanel.cpp
@@ -129,6 +129,11 @@ void UUniformGridPanel::SynchronizeProperties()
 {
 	Super::SynchronizeProperties();
 
+	if (!MyUniformGridPanel.IsValid())
+	{
+		return;
+	}
+
 	MyUniformGridPanel->SetSlotPadding(SlotPadding);
 	MyUniformGridPanel->SetMinDesiredSlotWidth(MinDesiredSlotWidth);
 	MyUniformGridPanel->SetMinDesiredSlotHeight(MinDesiredSlotHeight);
diff --git a/Engine/Source/Runtime/UMG/Private/Components/WrapBox.cpp b/Engine/Source/Runtime/UMG/Private/Components/WrapBox.cpp
--- a/Engine/Source/Runtime/UMG/Private/Components/WrapBox.cpp
+++ b/Engine/Source/Runtime/UMG/Private/Components/WrapBox.cpp
@@ -89,6 +89,11 @@ void UWrapBox::SynchronizeProperties()
 {
 	Super::SynchronizeProperties();
 
+	if (!MyWrapBox.IsValid())
+	{
+		return;
+	}
+
 	MyWrapBox->SetInnerSlotPadding(InnerSlotPadding);
 	MyWrapBox->SetUseAllottedSize(!bExplicitWrapSize);
 	MyWrapBox->SetWrapSize(WrapSize);
