UMG: Optimized animating widget materials by removing all string manipulation

  - Previously animating a widget material would resolve the brush handle every frame, which required a lot of string manipulation and property search/wrangling.
  - We now perform this resolution once when the track first stats, and cache the resulting handle as a component.

diff --git a/Engine/Source/Runtime/UMG/Private/Animation/MovieSceneUMGComponentTypes.cpp b/Engine/Source/Runtime/UMG/Private/Animation/MovieSceneUMGComponentTypes.cpp
--- a/Engine/Source/Runtime/UMG/Private/Animation/MovieSceneUMGComponentTypes.cpp
+++ b/Engine/Source/Runtime/UMG/Private/Animation/MovieSceneUMGComponentTypes.cpp
@@ -1,6 +1,7 @@
 // Copyright Epic Games, Inc. All Rights Reserved.
 
 #include "Animation/MovieSceneUMGComponentTypes.h"
+#include "Animation/WidgetMaterialTrackUtilities.h"
 #include "EntitySystem/BuiltInComponentTypes.h"
 #include "EntitySystem/MovieSceneComponentRegistry.h"
 #include "EntitySystem/MovieSceneEntityFactoryTemplates.h"
@@ -86,15 +87,48 @@ static void SetRenderTransform(UObject* Object, const FIntermediateWidgetTransfo
 
 FMovieSceneUMGComponentTypes::FMovieSceneUMGComponentTypes()
 {
-	FComponentRegistry* ComponentRegistry = UMovieSceneEntitySystemLinker::GetComponents();
+	FBuiltInComponentTypes* BuiltInComponents = FBuiltInComponentTypes::Get();
+	FComponentRegistry*     ComponentRegistry = UMovieSceneEntitySystemLinker::GetComponents();
 
 	ComponentRegistry->NewPropertyType(Margin, TEXT("FMargin Property"));
 
 	ComponentRegistry->NewPropertyType(WidgetTransform, TEXT("FWidgetTransform Property"));
 
 	ComponentRegistry->NewComponentType(&WidgetMaterialPath, TEXT("Widget Material Path"), EComponentTypeFlags::CopyToChildren | EComponentTypeFlags::CopyToOutput);
+	ComponentRegistry->NewComponentType(&WidgetMaterialHandle, TEXT("Widget Material Handle"), EComponentTypeFlags::CopyToOutput);
 
-	FBuiltInComponentTypes* BuiltInComponents = FBuiltInComponentTypes::Get();
+	/** Initializer that initializes the value of an FWidgetMaterialHandle derived from an FWidgetMaterialPath */
+	struct FWidgetMaterialHandleInitializer : TChildEntityInitializer<FWidgetMaterialPath, FWidgetMaterialHandle>
+	{
+		explicit FWidgetMaterialHandleInitializer(TComponentTypeID<FWidgetMaterialPath> Path, TComponentTypeID<FWidgetMaterialHandle> Handle)
+			: TChildEntityInitializer<FWidgetMaterialPath, FWidgetMaterialHandle>(Path, Handle)
+		{}
+
+		virtual void Run(const FEntityRange& ChildRange, const FEntityAllocation* ParentAllocation, TArrayView<const int32> ParentAllocationOffsets)
+		{
+			TComponentReader<FWidgetMaterialPath>   PathComponents        = ParentAllocation->ReadComponents(this->GetParentComponent());
+			TComponentWriter<FWidgetMaterialHandle> HandleComponents      = ChildRange.Allocation->WriteComponents(this->GetChildComponent(), FEntityAllocationWriteContext::NewAllocation());
+			TOptionalComponentReader<UObject*>      BoundObjectComponents = ChildRange.Allocation->TryReadComponents(FBuiltInComponentTypes::Get()->BoundObject);
+			if (!ensure(BoundObjectComponents))
+			{
+				return;
+			}
+
+			for (int32 Index = 0; Index < ChildRange.Num; ++Index)
+			{
+				const int32 ParentIndex = ParentAllocationOffsets[Index];
+				const int32 ChildIndex  = ChildRange.ComponentStartOffset + Index;
+
+				UWidget* Widget = Cast<UWidget>(BoundObjectComponents[ChildIndex]);
+				if (Widget)
+				{
+					HandleComponents[ChildIndex] = WidgetMaterialTrackUtilities::GetMaterialHandle(Widget, PathComponents[ParentIndex].Path);
+				}
+			}
+		}
+	};
+
+	ComponentRegistry->Factories.DefineChildComponent(FWidgetMaterialHandleInitializer(WidgetMaterialPath, WidgetMaterialHandle));
 
 	FMovieSceneTracksComponentTypes::Get()->Accessors.Float.Add(UWidget::StaticClass(), "RenderOpacity", &GetRenderOpacity, &SetRenderOpacity);
 
diff --git a/Engine/Source/Runtime/UMG/Private/Animation/MovieSceneWidgetMaterialSystem.cpp b/Engine/Source/Runtime/UMG/Private/Animation/MovieSceneWidgetMaterialSystem.cpp
--- a/Engine/Source/Runtime/UMG/Private/Animation/MovieSceneWidgetMaterialSystem.cpp
+++ b/Engine/Source/Runtime/UMG/Private/Animation/MovieSceneWidgetMaterialSystem.cpp
@@ -25,12 +25,12 @@ namespace UE::MovieScene
 
 FWidgetMaterialAccessor::FWidgetMaterialAccessor(const FWidgetMaterialKey& InKey)
 	: Widget(CastChecked<UWidget>(InKey.Object.ResolveObjectPtr()))
-	, WidgetMaterialPath(InKey.WidgetMaterialPath)
+	, WidgetMaterialHandle(InKey.WidgetMaterialHandle)
 {}
 
-FWidgetMaterialAccessor::FWidgetMaterialAccessor(UObject* InObject, FWidgetMaterialPath InWidgetMaterialPath)
+FWidgetMaterialAccessor::FWidgetMaterialAccessor(UObject* InObject, FWidgetMaterialHandle InWidgetMaterialHandle)
 	: Widget(Cast<UWidget>(InObject))
-	, WidgetMaterialPath(MoveTemp(InWidgetMaterialPath))
+	, WidgetMaterialHandle(MoveTemp(InWidgetMaterialHandle))
 {
 	// Object must be a widget
 	if (InObject && !Widget)
@@ -48,26 +48,24 @@ FWidgetMaterialAccessor::operator bool() const
 
 FString FWidgetMaterialAccessor::ToString() const
 {
-	return FString::Printf(TEXT("Brush %s.%s"), *Widget->GetPathName(), *FString::JoinBy(WidgetMaterialPath.Path, TEXT("."), UE_PROJECTION_MEMBER(FName, ToString)));
+	return FString::Printf(TEXT("Brush on widget %s"), *Widget->GetPathName());
 }
 
 UMaterialInterface* FWidgetMaterialAccessor::GetMaterial() const
 {
-	FWidgetMaterialHandle Handle = WidgetMaterialTrackUtilities::GetMaterialHandle(Widget, WidgetMaterialPath.Path);
-	if (Handle.IsValid())
+	if (WidgetMaterialHandle.IsValid())
 	{
-		return Handle.GetMaterial();
+		return WidgetMaterialHandle.GetMaterial();
 	}
 
 	return nullptr;
 }
 
-void FWidgetMaterialAccessor::SetMaterial(UMaterialInterface* InMaterial) const
+void FWidgetMaterialAccessor::SetMaterial(UMaterialInterface* InMaterial)
 {
-	FWidgetMaterialHandle Handle = WidgetMaterialTrackUtilities::GetMaterialHandle(Widget, WidgetMaterialPath.Path);
-	if (Handle.IsValid())
+	if (WidgetMaterialHandle.IsValid())
 	{
-		Handle.SetMaterial(InMaterial, Widget);
+		WidgetMaterialHandle.SetMaterial(InMaterial, Widget);
 	}
 }
 
@@ -98,7 +96,7 @@ UMovieSceneWidgetMaterialSystem::UMovieSceneWidgetMaterialSystem(const FObjectIn
 	FMovieSceneUMGComponentTypes*    WidgetComponents  = FMovieSceneUMGComponentTypes::Get();
 	FMovieSceneTracksComponentTypes* TracksComponents  = FMovieSceneTracksComponentTypes::Get();
 
-	RelevantComponent = WidgetComponents->WidgetMaterialPath;
+	RelevantComponent = WidgetComponents->WidgetMaterialHandle;
 	Phase = ESystemPhase::Instantiation | ESystemPhase::Evaluation;
 
 	if (HasAnyFlags(RF_ClassDefaultObject))
@@ -120,7 +118,7 @@ void UMovieSceneWidgetMaterialSystem::OnLink()
 	SystemImpl.MaterialSwitcherStorage = Linker->PreAnimatedState.GetOrCreateStorage<FPreAnimatedWidgetMaterialSwitcherStorage>();
 	SystemImpl.MaterialParameterStorage = Linker->PreAnimatedState.GetOrCreateStorage<FPreAnimatedWidgetMaterialParameterStorage>();
 
-	SystemImpl.OnLink(Linker, BuiltInComponents->BoundObject, WidgetComponents->WidgetMaterialPath);
+	SystemImpl.OnLink(Linker, BuiltInComponents->BoundObject, WidgetComponents->WidgetMaterialHandle);
 }
 
 void UMovieSceneWidgetMaterialSystem::OnUnlink()
@@ -135,7 +133,7 @@ void UMovieSceneWidgetMaterialSystem::OnRun(FSystemTaskPrerequisites& InPrerequi
 	FBuiltInComponentTypes*       BuiltInComponents = FBuiltInComponentTypes::Get();
 	FMovieSceneUMGComponentTypes* WidgetComponents  = FMovieSceneUMGComponentTypes::Get();
 
-	SystemImpl.OnRun(Linker, BuiltInComponents->BoundObject, WidgetComponents->WidgetMaterialPath, InPrerequisites, Subsequents);
+	SystemImpl.OnRun(Linker, BuiltInComponents->BoundObject, WidgetComponents->WidgetMaterialHandle, InPrerequisites, Subsequents);
 }
 
 void UMovieSceneWidgetMaterialSystem::SavePreAnimatedState(const FPreAnimationParameters& InParameters)
@@ -145,5 +143,5 @@ void UMovieSceneWidgetMaterialSystem::SavePreAnimatedState(const FPreAnimationPa
 	FBuiltInComponentTypes*       BuiltInComponents = FBuiltInComponentTypes::Get();
 	FMovieSceneUMGComponentTypes* WidgetComponents  = FMovieSceneUMGComponentTypes::Get();
 
-	SystemImpl.SavePreAnimatedState(Linker, BuiltInComponents->BoundObject, WidgetComponents->WidgetMaterialPath, InParameters);
+	SystemImpl.SavePreAnimatedState(Linker, BuiltInComponents->BoundObject, WidgetComponents->WidgetMaterialHandle, InParameters);
 }
diff --git a/Engine/Source/Runtime/UMG/Public/Animation/MovieSceneUMGComponentTypes.h b/Engine/Source/Runtime/UMG/Public/Animation/MovieSceneUMGComponentTypes.h
--- a/Engine/Source/Runtime/UMG/Public/Animation/MovieSceneUMGComponentTypes.h
+++ b/Engine/Source/Runtime/UMG/Public/Animation/MovieSceneUMGComponentTypes.h
@@ -4,6 +4,7 @@
 
 #include "Layout/Margin.h"
 #include "Slate/WidgetTransform.h"
+#include "Animation/WidgetMaterialTrackUtilities.h"
 #include "EntitySystem/MovieSceneEntityIDs.h"
 #include "EntitySystem/MovieScenePropertySystemTypes.h"
 #include "EntitySystem/MovieScenePropertyTraits.h"
@@ -61,6 +62,7 @@ struct UMG_API FMovieSceneUMGComponentTypes
 	TPropertyComponents<FWidgetTransformPropertyTraits> WidgetTransform;
 
 	TComponentTypeID<FWidgetMaterialPath> WidgetMaterialPath;
+	TComponentTypeID<FWidgetMaterialHandle> WidgetMaterialHandle;
 
 	TCustomPropertyRegistration<FWidgetTransformPropertyTraits, 1> CustomWidgetTransformAccessors;
 
diff --git a/Engine/Source/Runtime/UMG/Public/Animation/MovieSceneWidgetMaterialSystem.h b/Engine/Source/Runtime/UMG/Public/Animation/MovieSceneWidgetMaterialSystem.h
--- a/Engine/Source/Runtime/UMG/Public/Animation/MovieSceneWidgetMaterialSystem.h
+++ b/Engine/Source/Runtime/UMG/Public/Animation/MovieSceneWidgetMaterialSystem.h
@@ -19,20 +19,15 @@ namespace UE::MovieScene
 struct FWidgetMaterialKey
 {
 	FObjectKey Object;
-	FWidgetMaterialPath WidgetMaterialPath;
+	FWidgetMaterialHandle WidgetMaterialHandle;
 
 	friend uint32 GetTypeHash(const FWidgetMaterialKey& In)
 	{
-		uint32 Accumulator = GetTypeHash(In.Object);
-		for (const FName& Name : In.WidgetMaterialPath.Path)
-		{
-			Accumulator ^= GetTypeHash(Name);
-		}
-		return Accumulator;
+		return GetTypeHash(In.Object) ^ GetTypeHash(In.WidgetMaterialHandle);
 	}
 	friend bool operator==(const FWidgetMaterialKey& A, const FWidgetMaterialKey& B)
 	{
-		return A.Object == B.Object && A.WidgetMaterialPath.Path == B.WidgetMaterialPath.Path;
+		return A.Object == B.Object && A.WidgetMaterialHandle == B.WidgetMaterialHandle;
 	}
 };
 
@@ -41,30 +36,30 @@ struct FWidgetMaterialAccessor
 	using KeyType = FWidgetMaterialKey;
 
 	UWidget* Widget;
-	FWidgetMaterialPath WidgetMaterialPath;
+	FWidgetMaterialHandle WidgetMaterialHandle;
 
 	FWidgetMaterialAccessor(const FWidgetMaterialKey& InKey);
-	FWidgetMaterialAccessor(UObject* InObject, FWidgetMaterialPath InWidgetMaterialPath);
+	FWidgetMaterialAccessor(UObject* InObject, FWidgetMaterialHandle InWidgetMaterialHandle);
 
 	explicit operator bool() const;
 
 	UMaterialInterface* GetMaterial() const;
-	void SetMaterial(UMaterialInterface* InMaterial) const;
+	void SetMaterial(UMaterialInterface* InMaterial);
 	UMaterialInstanceDynamic* CreateDynamicMaterial(UMaterialInterface* InMaterial);
 	FString ToString() const;
 };
 
-using FPreAnimatedWidgetMaterialTraits          = TPreAnimatedMaterialTraits<FWidgetMaterialAccessor, UObject*, FWidgetMaterialPath>;
-using FPreAnimatedWidgetMaterialParameterTraits = TPreAnimatedMaterialParameterTraits<FWidgetMaterialAccessor, UObject*, FWidgetMaterialPath>;
+using FPreAnimatedWidgetMaterialTraits          = TPreAnimatedMaterialTraits<FWidgetMaterialAccessor, UObject*, FWidgetMaterialHandle>;
+using FPreAnimatedWidgetMaterialParameterTraits = TPreAnimatedMaterialParameterTraits<FWidgetMaterialAccessor, UObject*, FWidgetMaterialHandle>;
 
 struct FPreAnimatedWidgetMaterialSwitcherStorage
-	: public TPreAnimatedStateStorage<TPreAnimatedMaterialTraits<FWidgetMaterialAccessor, UObject*, FWidgetMaterialPath>>
+	: public TPreAnimatedStateStorage<TPreAnimatedMaterialTraits<FWidgetMaterialAccessor, UObject*, FWidgetMaterialHandle>>
 {
 	static TAutoRegisterPreAnimatedStorageID<FPreAnimatedWidgetMaterialSwitcherStorage> StorageID;
 };
 
 struct FPreAnimatedWidgetMaterialParameterStorage
-	: public TPreAnimatedStateStorage<TPreAnimatedMaterialParameterTraits<FWidgetMaterialAccessor, UObject*, FWidgetMaterialPath>>
+	: public TPreAnimatedStateStorage<TPreAnimatedMaterialParameterTraits<FWidgetMaterialAccessor, UObject*, FWidgetMaterialHandle>>
 {
 	static TAutoRegisterPreAnimatedStorageID<FPreAnimatedWidgetMaterialParameterStorage> StorageID;
 };
@@ -93,5 +88,5 @@ private:
 
 private:
 
-	UE::MovieScene::TMovieSceneMaterialSystem<UE::MovieScene::FWidgetMaterialAccessor, UObject*, UE::MovieScene::FWidgetMaterialPath> SystemImpl;
+	UE::MovieScene::TMovieSceneMaterialSystem<UE::MovieScene::FWidgetMaterialAccessor, UObject*, FWidgetMaterialHandle> SystemImpl;
 };
diff --git a/Engine/Source/Runtime/UMG/Public/Animation/WidgetMaterialTrackUtilities.h b/Engine/Source/Runtime/UMG/Public/Animation/WidgetMaterialTrackUtilities.h
--- a/Engine/Source/Runtime/UMG/Public/Animation/WidgetMaterialTrackUtilities.h
+++ b/Engine/Source/Runtime/UMG/Public/Animation/WidgetMaterialTrackUtilities.h
@@ -23,6 +23,19 @@ public:
 		, Data(InData)
 	{}
 
+	friend uint32 GetTypeHash(const FWidgetMaterialHandle& In)
+	{
+		return HashCombine(GetTypeHash(In.TypeName), GetTypeHash(In.Data));
+	}
+	friend bool operator==(const FWidgetMaterialHandle& A, const FWidgetMaterialHandle& B)
+	{
+		return A.TypeName == B.TypeName && A.Data == B.Data;
+	}
+	friend bool operator!=(const FWidgetMaterialHandle& A, const FWidgetMaterialHandle& B)
+	{
+		return !(A == B);
+	}
+
 	/** @return true if this handle points to valid data */
 	bool IsValid() const { return TypeName != NAME_None && Data != nullptr; }
 
