Add accessors for ComboBox Key and ComboBox String

diff --git a/Engine/Source/Runtime/Slate/Private/Widgets/Input/SComboButton.cpp b/Engine/Source/Runtime/Slate/Private/Widgets/Input/SComboButton.cpp
--- a/Engine/Source/Runtime/Slate/Private/Widgets/Input/SComboButton.cpp
+++ b/Engine/Source/Runtime/Slate/Private/Widgets/Input/SComboButton.cpp
@@ -10,7 +10,7 @@
 void SComboButton::Construct( const FArguments& InArgs )
 {
 	check(InArgs._ComboButtonStyle);
-
+	Style = InArgs._ComboButtonStyle;
 	// Work out which values we should use based on whether we were given an override, or should use the style's version
 	const FButtonStyle* const OurButtonStyle = InArgs._ButtonStyle ? InArgs._ButtonStyle : &InArgs._ComboButtonStyle->ButtonStyle;
 
@@ -23,8 +23,6 @@ void SComboButton::Construct( const FArguments& InArgs )
 
 	const bool bHasDownArrowShadow = !InArgs._ComboButtonStyle->ShadowOffset.IsZero();
 
-	TSharedPtr<SHorizontalBox> HBox;
-
 	SMenuAnchor::Construct( SMenuAnchor::FArguments()
 		.Placement(InArgs._MenuPlacement)
 		.Method(InArgs._Method)
@@ -32,7 +30,7 @@ void SComboButton::Construct( const FArguments& InArgs )
 		.OnGetMenuContent(InArgs._OnGetMenuContent)
 		.IsCollapsedByParent(InArgs._CollapseMenuOnParentFocus)
 		[
-			SNew( SButton )
+			SAssignNew(ButtonPtr, SButton )
 			.ButtonStyle( OurButtonStyle )
 			.ClickMethod( EButtonClickMethod::MouseDown )
 			.OnClicked( this, &SComboButton::OnButtonClicked )
@@ -70,7 +68,7 @@ void SComboButton::Construct( const FArguments& InArgs )
 					.VAlign(VAlign_Top)
 					.Padding(FMargin(InArgs._ComboButtonStyle->ShadowOffset.X, InArgs._ComboButtonStyle->ShadowOffset.Y, 0, 0))
 					[
-						SNew(SImage)
+						SAssignNew(ShadowImage, SImage)
 						.Visibility( InArgs._HasDownArrow && bHasDownArrowShadow ? EVisibility::Visible : EVisibility::Collapsed )
 						.Image( &InArgs._ComboButtonStyle->DownArrowImage )
 						.ColorAndOpacity( InArgs._ComboButtonStyle->ShadowColorAndOpacity )
@@ -78,7 +76,7 @@ void SComboButton::Construct( const FArguments& InArgs )
 					+ SOverlay::Slot()
 					.VAlign(VAlign_Top)
 					[
-						SNew(SImage)
+						SAssignNew(ForegroundArrowImage,SImage)
 						.Visibility( InArgs._HasDownArrow ? EVisibility::Visible : EVisibility::Collapsed )
 						.Image( &InArgs._ComboButtonStyle->DownArrowImage )
 						// Inherit tinting from parent
@@ -190,3 +188,23 @@ void SComboButton::SetOnGetMenuContent(FOnGetContent InOnGetMenuContent)
 {
 	OnGetMenuContent = InOnGetMenuContent;
 }
+
+void SComboButton::SetButtonContentPadding(FMargin InPadding)
+{
+	check(ButtonPtr);
+	ButtonPtr->SetContentPadding(InPadding);
+}
+
+void SComboButton::SetHasDownArrow(bool InHasArrowDown)
+{
+	const bool bHasDownArrowShadow = !Style->ShadowOffset.IsZero();
+
+	check(HBox && HBox->NumSlots() >= 2);
+	HBox->GetSlot(1).SetVerticalAlignment(InHasArrowDown ? (EVerticalAlignment)Style->DownArrowAlign : VAlign_Center);
+	HBox->GetSlot(1).SetPadding(InHasArrowDown ? Style->DownArrowPadding : FMargin(0));
+
+	check(ForegroundArrowImage);
+	ForegroundArrowImage->SetVisibility(InHasArrowDown ? EVisibility::Visible : EVisibility::Collapsed);
+	check(ShadowImage);
+	ShadowImage->SetVisibility(InHasArrowDown && bHasDownArrowShadow ? EVisibility::Visible : EVisibility::Collapsed);
+}
\ No newline at end of file
diff --git a/Engine/Source/Runtime/Slate/Public/Widgets/Input/SComboBox.h b/Engine/Source/Runtime/Slate/Public/Widgets/Input/SComboBox.h
--- a/Engine/Source/Runtime/Slate/Public/Widgets/Input/SComboBox.h
+++ b/Engine/Source/Runtime/Slate/Public/Widgets/Input/SComboBox.h
@@ -211,7 +211,7 @@ public:
 		OptionsSource = InArgs._OptionsSource;
 		CustomScrollbar = InArgs._CustomScrollbar;
 
-		TSharedRef<SWidget> ComboBoxMenuContent =
+		ComboBoxMenuContent =
 			SNew(SBox)
 			.MaxDesiredHeight(InArgs._MaxListHeight)
 			[
@@ -246,7 +246,7 @@ public:
 			]
 			.MenuContent()
 			[
-				ComboBoxMenuContent
+				ComboBoxMenuContent.ToSharedRef()
 			]
 			.HasDownArrow( InArgs._HasDownArrow )
 			.ContentPadding( InArgs._ContentPadding )
@@ -363,6 +363,16 @@ public:
 		}
 	}
 
+	void SetEnableGamepadNavigationMode(bool InEnableGamepadNavigationMode)
+	{
+		this->EnableGamepadNavigationMode = InEnableGamepadNavigationMode;
+	}
+
+	void SetMaxHeight(float InMaxHeight)
+	{
+		ComboBoxMenuContent->SetMaxDesiredHeight(InMaxHeight);
+	}
+
 	/** @return the item currently selected by the combo box. */
 	NullableOptionType GetSelectedItem()
 	{
@@ -642,7 +652,7 @@ private:
 	// When true, navigation away from the widget is prevented until a new value has been accepted or canceled. 
 	bool bControllerInputCaptured;
 
-		
+	TSharedPtr<SBox> ComboBoxMenuContent;
 
 	const TArray< OptionType >* OptionsSource;
 };
diff --git a/Engine/Source/Runtime/Slate/Public/Widgets/Input/SComboButton.h b/Engine/Source/Runtime/Slate/Public/Widgets/Input/SComboButton.h
--- a/Engine/Source/Runtime/Slate/Public/Widgets/Input/SComboButton.h
+++ b/Engine/Source/Runtime/Slate/Public/Widgets/Input/SComboButton.h
@@ -13,7 +13,10 @@
 #include "Widgets/SBoxPanel.h"
 #include "Widgets/Input/SMenuAnchor.h"
 
-DECLARE_DELEGATE( FOnComboBoxOpened )
+DECLARE_DELEGATE(FOnComboBoxOpened)
+
+class SButton;
+class SImage;
 
 /**
  * A button that, when clicked, brings up a popup.
@@ -90,6 +93,12 @@ public:
 		WidgetToFocusPtr = InWidgetToFocusPtr;
 	}
 
+	/** See the padding for button content. */
+	void SetButtonContentPadding(FMargin InPadding);
+
+	/** add/remove the expanding arrow. */
+	void SetHasDownArrow(bool InHasArrowDown);
+
 protected:
 	/**
 	 * Handle the button being clicked by summoning the ComboButton.
@@ -126,4 +135,20 @@ protected:
 
 	/** Can this button be focused? */
 	bool bIsFocusable;
+
+private:
+	/** The button widget within the ComboButton */
+	TSharedPtr<SButton> ButtonPtr;
+
+	/** The box containing the arrow */
+	TSharedPtr<SHorizontalBox> HBox;
+
+	/** The arrow image shadow */
+	TSharedPtr<SImage> ShadowImage;
+
+	/** The arrow image */
+	TSharedPtr<SImage> ForegroundArrowImage;
+
+	/** The ComboButton style */
+	const FComboButtonStyle* Style;
 };
diff --git a/Engine/Source/Runtime/UMG/Private/Components/CheckBox.cpp b/Engine/Source/Runtime/UMG/Private/Components/CheckBox.cpp
--- a/Engine/Source/Runtime/UMG/Private/Components/CheckBox.cpp
+++ b/Engine/Source/Runtime/UMG/Private/Components/CheckBox.cpp
@@ -59,9 +59,9 @@ PRAGMA_DISABLE_DEPRECATION_WARNINGS
 	ClickMethod = EButtonClickMethod::DownAndUp;
 	TouchMethod = EButtonTouchMethod::DownAndUp;
 	PressMethod = EButtonPressMethod::DownAndUp;
-PRAGMA_ENABLE_DEPRECATION_WARNINGS
 
 	IsFocusable = true;
+PRAGMA_ENABLE_DEPRECATION_WARNINGS
 #if WITH_EDITORONLY_DATA
 	AccessibleBehavior = ESlateAccessibleBehavior::Summary;
 	bCanChildrenBeAccessible = false;
@@ -250,6 +250,24 @@ void UCheckBox::SetWidgetStyle(const FCheckBoxStyle& InStyle)
 		MyCheckbox->SetStyle(&WidgetStyle);
 	}
 }
+
+bool UCheckBox::GetIsFocusable() const
+{
+	return IsFocusable;
+}
+
+void UCheckBox::InitIsFocusable(bool InIsFocusable)
+{
+	ensureMsgf(!MyCheckbox.IsValid(), TEXT("The widget is already created."));
+	IsFocusable = InIsFocusable;
+}
+
+void UCheckBox::InitCheckedStateDelegate(FGetCheckBoxState InCheckedStateDelegate)
+{
+	ensureMsgf(!MyCheckbox.IsValid(), TEXT("The widget is already created."));
+	CheckedStateDelegate = InCheckedStateDelegate;
+}
+
 PRAGMA_ENABLE_DEPRECATION_WARNINGS
 
 void UCheckBox::SlateOnCheckStateChangedCallback(ECheckBoxState NewState)
diff --git a/Engine/Source/Runtime/UMG/Private/Components/ComboBoxKey.cpp b/Engine/Source/Runtime/UMG/Private/Components/ComboBoxKey.cpp
--- a/Engine/Source/Runtime/UMG/Private/Components/ComboBoxKey.cpp
+++ b/Engine/Source/Runtime/UMG/Private/Components/ComboBoxKey.cpp
@@ -22,9 +22,11 @@ UComboBoxKey::UComboBoxKey()
 	if (!IsRunningDedicatedServer())
 	{
 #if WITH_EDITOR 
+PRAGMA_DISABLE_DEPRECATION_WARNINGS
 		WidgetStyle = IsEditorWidget() ? FDefaultStyleCache::Get().GetEditorComboBoxStyle() : FDefaultStyleCache::Get().GetComboBoxStyle();
 		ItemStyle = IsEditorWidget() ? FDefaultStyleCache::Get().GetEditorComboBoxRowStyle() : FDefaultStyleCache::Get().GetComboBoxRowStyle();
 		ScrollBarStyle = IsEditorWidget() ? FDefaultStyleCache::Get().GetEditorScrollBarStyle() : FDefaultStyleCache::Get().GetScrollBarStyle();
+PRAGMA_ENABLE_DEPRECATION_WARNINGS
 
 		if (IsEditorWidget())
 		{
@@ -32,12 +34,15 @@ UComboBoxKey::UComboBoxKey()
 			PostEditChange();
 		}
 #else
+PRAGMA_DISABLE_DEPRECATION_WARNINGS
 		WidgetStyle = FDefaultStyleCache::Get().GetComboBoxStyle();
 		ItemStyle = FDefaultStyleCache::Get().GetComboBoxRowStyle();
 		ScrollBarStyle = FDefaultStyleCache::Get().GetScrollBarStyle();
+PRAGMA_ENABLE_DEPRECATION_WARNINGS
 #endif // WITH_EDITOR
 	}
 
+PRAGMA_DISABLE_DEPRECATION_WARNINGS
 	ContentPadding = FMargin(4.0, 2.0);
 	ForegroundColor = ItemStyle.TextColor;
 
@@ -45,6 +50,7 @@ UComboBoxKey::UComboBoxKey()
 	bHasDownArrow = true;
 	bEnableGamepadNavigationMode = true;
 	bIsFocusable = true;
+PRAGMA_ENABLE_DEPRECATION_WARNINGS
 }
 
 
@@ -59,6 +65,7 @@ void UComboBoxKey::ReleaseSlateResources(bool bReleaseChildren)
 
 TSharedRef<SWidget> UComboBoxKey::RebuildWidget()
 {
+PRAGMA_DISABLE_DEPRECATION_WARNINGS
 	MyComboBox =
 		SNew(SComboBox<FName>)
 		.ComboBoxStyle(&WidgetStyle)
@@ -78,6 +85,7 @@ TSharedRef<SWidget> UComboBoxKey::RebuildWidget()
 		[
 			SAssignNew(ComboBoxContent, SBox)
 		];
+PRAGMA_ENABLE_DEPRECATION_WARNINGS
 
 	GenerateContent();
 
@@ -140,6 +148,7 @@ void UComboBoxKey::SetSelectedOption(FName Option)
 {
 	if (SelectedOption != Option)
 	{
+		BroadcastFieldValueChanged(FFieldNotificationClassDescriptor::SelectedOption);
 		if (MyComboBox)
 		{
 			MyComboBox->SetSelectedItem(Option);
@@ -151,7 +160,6 @@ void UComboBoxKey::SetSelectedOption(FName Option)
 	}
 }
 
-
 FName UComboBoxKey::GetSelectedOption() const
 {
 	return SelectedOption;
@@ -163,6 +171,118 @@ bool UComboBoxKey::IsOpen() const
 	return MyComboBox && MyComboBox->IsOpen();
 }
 
+PRAGMA_DISABLE_DEPRECATION_WARNINGS
+void  UComboBoxKey::SetContentPadding(FMargin InPadding)
+{
+	ContentPadding = InPadding;
+	if (MyComboBox.IsValid())
+	{
+		MyComboBox->SetButtonContentPadding(InPadding);
+	}
+}
+
+FMargin  UComboBoxKey::GetContentPadding() const
+{
+	return ContentPadding;
+}
+
+void UComboBoxKey::SetEnableGamepadNavigationMode(bool InEnableGamepadNavigationMode)
+{
+	bEnableGamepadNavigationMode = InEnableGamepadNavigationMode;
+	if (MyComboBox.IsValid())
+	{
+		MyComboBox->SetEnableGamepadNavigationMode(bEnableGamepadNavigationMode);
+	}
+}
+
+bool UComboBoxKey::IsEnableGamepadNavigationMode() const
+{
+	return bEnableGamepadNavigationMode;
+}
+
+bool UComboBoxKey::IsHasDownArrow() const
+{
+	return bHasDownArrow;
+}
+
+void UComboBoxKey::SetHasDownArrow(bool InHasDownArrow)
+{
+	bHasDownArrow = InHasDownArrow;
+	if (MyComboBox.IsValid())
+	{
+		MyComboBox->SetHasDownArrow(InHasDownArrow);
+	}
+}
+
+float UComboBoxKey::GetMaxListHeight() const
+{
+	return MaxListHeight;
+}
+
+void UComboBoxKey::SetMaxListHeight(float InMaxHeight)
+{
+	MaxListHeight = InMaxHeight;
+	if (MyComboBox.IsValid())
+	{
+		MyComboBox->SetMaxHeight(MaxListHeight);
+	}
+}
+
+const FComboBoxStyle& UComboBoxKey::GetWidgetStyle() const
+{
+	return WidgetStyle;
+}
+
+const FTableRowStyle& UComboBoxKey::GetItemStyle() const
+{
+	return ItemStyle;
+}
+
+const FScrollBarStyle& UComboBoxKey::GetScrollBarStyle() const
+{
+	return ScrollBarStyle;
+}
+
+bool UComboBoxKey::IsFocusable() const
+{
+	return bIsFocusable;
+}
+
+FSlateColor UComboBoxKey::GetForegroundColor() const
+{
+	return ForegroundColor;
+}
+
+void UComboBoxKey::InitWidgetStyle(const FComboBoxStyle& InWidgetStyle)
+{
+	ensureMsgf(!MyComboBox.IsValid(), TEXT("The widget is already created."));
+	WidgetStyle = InWidgetStyle;
+}
+
+void UComboBoxKey::InitItemStyle(const FTableRowStyle& InItemStyle)
+{
+	ensureMsgf(!MyComboBox.IsValid(), TEXT("The widget is already created."));
+	ItemStyle = InItemStyle;
+}
+
+void UComboBoxKey::InitScrollBarStyle(const FScrollBarStyle& InScrollBarStyle)
+{
+	ensureMsgf(!MyComboBox.IsValid(), TEXT("The widget is already created."));
+	ScrollBarStyle = InScrollBarStyle;
+}
+
+void UComboBoxKey::InitIsFocusable(bool InIsFocusable)
+{
+	ensureMsgf(!MyComboBox.IsValid(), TEXT("The widget is already created."));
+	bIsFocusable = InIsFocusable;
+}
+
+void UComboBoxKey::InitForegroundColor(FSlateColor InForegroundColor)
+{
+	ensureMsgf(!MyComboBox.IsValid(), TEXT("The widget is already created."));
+	ForegroundColor = InForegroundColor;
+}
+PRAGMA_ENABLE_DEPRECATION_WARNINGS
 
 void UComboBoxKey::GenerateContent()
 {
@@ -238,6 +358,7 @@ void UComboBoxKey::HandleSelectionChanged(FName Item, ESelectInfo::Type Selectio
 	if (SelectedOption != Item)
 	{
 		SelectedOption = Item;
+		BroadcastFieldValueChanged(FFieldNotificationClassDescriptor::SelectedOption);
 
 		if (!IsDesignTime())
 		{
diff --git a/Engine/Source/Runtime/UMG/Private/Components/ComboBoxString.cpp b/Engine/Source/Runtime/UMG/Private/Components/ComboBoxString.cpp
--- a/Engine/Source/Runtime/UMG/Private/Components/ComboBoxString.cpp
+++ b/Engine/Source/Runtime/UMG/Private/Components/ComboBoxString.cpp
@@ -19,6 +19,7 @@
 UComboBoxString::UComboBoxString(const FObjectInitializer& ObjectInitializer)
 	: Super(ObjectInitializer)
 {
+PRAGMA_DISABLE_DEPRECATION_WARNINGS
 	WidgetStyle = FDefaultStyleCache::Get().GetComboBoxStyle();
 	ItemStyle = FDefaultStyleCache::Get().GetComboBoxRowStyle();
 	ScrollBarStyle = FDefaultStyleCache::Get().GetScrollBarStyle();
@@ -48,6 +49,7 @@ UComboBoxString::UComboBoxString(const FObjectInitializer& ObjectInitializer)
 		static ConstructorHelpers::FObjectFinder<UFont> RobotoFontObj(*UWidget::GetDefaultFontName());
 		Font = FSlateFontInfo(RobotoFontObj.Object, 16, FName("Bold"));
 	}
+PRAGMA_ENABLE_DEPRECATION_WARNINGS
 }
 
 void UComboBoxString::PostInitProperties()
@@ -88,7 +90,9 @@ void UComboBoxString::PostLoad()
 
 	if (GetLinkerCustomVersion(FEditorObjectVersion::GUID) < FEditorObjectVersion::ComboBoxControllerSupportUpdate)
 	{
+PRAGMA_DISABLE_DEPRECATION_WARNINGS
 		EnableGamepadNavigationMode = false;
+PRAGMA_ENABLE_DEPRECATION_WARNINGS
 	}
 }
 
@@ -99,7 +103,7 @@ TSharedRef<SWidget> UComboBoxString::RebuildWidget()
 	{
 		CurrentOptionPtr = Options[InitialIndex];
 	}
-
+PRAGMA_DISABLE_DEPRECATION_WARNINGS
 	MyComboBox =
 		SNew(SComboBox< TSharedPtr<FString> >)
 		.ComboBoxStyle(&WidgetStyle)
@@ -119,7 +123,7 @@ TSharedRef<SWidget> UComboBoxString::RebuildWidget()
 		[
 			SAssignNew(ComboBoxContent, SBox)
 		];
-
+PRAGMA_ENABLE_DEPRECATION_WARNINGS
 	if ( InitialIndex != -1 )
 	{
 		// Generate the widget for the initially selected widget if needed
@@ -230,6 +234,7 @@ void UComboBoxString::SetSelectedIndex(const int32 Index)
 		// Don't select item if its already selected
 		if (SelectedOption != *CurrentOptionPtr)
 		{
+			BroadcastFieldValueChanged(FFieldNotificationClassDescriptor::SelectedOption);
 			SelectedOption = *CurrentOptionPtr;
 
 			if (ComboBoxContent.IsValid())
@@ -279,6 +284,131 @@ bool UComboBoxString::IsOpen() const
 	return MyComboBox.IsValid() && MyComboBox->IsOpen();
 }
 
+PRAGMA_DISABLE_DEPRECATION_WARNINGS
+void  UComboBoxString::SetContentPadding(FMargin InPadding)
+{
+	ContentPadding = InPadding;
+	if (MyComboBox.IsValid())
+	{
+		MyComboBox->SetButtonContentPadding(InPadding);
+	}
+}
+
+FMargin  UComboBoxString::GetContentPadding() const
+{
+	return ContentPadding;
+}
+
+void UComboBoxString::SetEnableGamepadNavigationMode(bool InEnableGamepadNavigationMode)
+{
+	EnableGamepadNavigationMode = InEnableGamepadNavigationMode;
+	if (MyComboBox.IsValid())
+	{
+		MyComboBox->SetEnableGamepadNavigationMode(EnableGamepadNavigationMode);
+	}
+}
+
+bool UComboBoxString::IsEnableGamepadNavigationMode() const
+{
+	return EnableGamepadNavigationMode;
+}
+
+bool UComboBoxString::IsHasDownArrow() const
+{
+	return HasDownArrow;
+}
+
+void UComboBoxString::SetHasDownArrow(bool InHasDownArrow)
+{
+	HasDownArrow = InHasDownArrow;
+	if (MyComboBox.IsValid())
+	{
+		MyComboBox->SetHasDownArrow(HasDownArrow);
+	}
+}
+
+float UComboBoxString::GetMaxListHeight() const
+{
+	return MaxListHeight;
+}
+
+void UComboBoxString::SetMaxListHeight(float InMaxHeight)
+{
+	MaxListHeight = InMaxHeight;
+	if (MyComboBox.IsValid())
+	{
+		MyComboBox->SetMaxHeight(MaxListHeight);
+	}
+}
+
+const FSlateFontInfo&  UComboBoxString::GetFont() const
+{
+	return Font;
+}
+
+const FComboBoxStyle& UComboBoxString::GetWidgetStyle() const
+{
+	return WidgetStyle;
+}
+
+const FTableRowStyle& UComboBoxString::GetItemStyle() const
+{
+	return ItemStyle;
+}
+
+const FScrollBarStyle& UComboBoxString::GetScrollBarStyle() const
+{
+	return ScrollBarStyle;
+}
+
+bool UComboBoxString::IsFocusable() const
+{
+	return bIsFocusable;
+}
+
+FSlateColor UComboBoxString::GetForegroundColor() const
+{
+	return ForegroundColor;
+}
+
+void UComboBoxString::InitWidgetStyle(const FComboBoxStyle& InWidgetStyle)
+{
+	ensureMsgf(!MyComboBox.IsValid(), TEXT("The widget is already created."));
+	WidgetStyle = InWidgetStyle;
+}
+
+void UComboBoxString::InitItemStyle(const FTableRowStyle& InItemStyle)
+{
+	ensureMsgf(!MyComboBox.IsValid(), TEXT("The widget is already created."));
+	ItemStyle = InItemStyle;
+}
+
+void UComboBoxString::InitScrollBarStyle(const FScrollBarStyle& InScrollBarStyle)
+{
+	ensureMsgf(!MyComboBox.IsValid(), TEXT("The widget is already created."));
+	ScrollBarStyle = InScrollBarStyle;
+}
+
+void UComboBoxString::InitFont(FSlateFontInfo InFont)
+{
+	ensureMsgf(!MyComboBox.IsValid(), TEXT("The widget is already created."));
+	Font = InFont;
+}
+
+void UComboBoxString::InitIsFocusable(bool InIsFocusable)
+{
+	ensureMsgf(!MyComboBox.IsValid(), TEXT("The widget is already created."));
+	bIsFocusable = InIsFocusable;
+}
+
+void UComboBoxString::InitForegroundColor(FSlateColor InForegroundColor)
+{
+	ensureMsgf(!MyComboBox.IsValid(), TEXT("The widget is already created."));
+	ForegroundColor = InForegroundColor;
+}
+
+PRAGMA_ENABLE_DEPRECATION_WARNINGS
+
 void UComboBoxString::UpdateOrGenerateWidget(TSharedPtr<FString> Item)
 {
 	// If no custom widget was supplied and the default STextBlock already exists,
@@ -310,14 +440,17 @@ TSharedRef<SWidget> UComboBoxString::HandleGenerateWidget(TSharedPtr<FString> It
 	}
 
 	// If a row wasn't generated just create the default one, a simple text block of the item's name.
+PRAGMA_DISABLE_DEPRECATION_WARNINGS
 	return SNew(STextBlock)
 		.Text(FText::FromString(StringItem))
 		.Font(Font);
+PRAGMA_ENABLE_DEPRECATION_WARNINGS
 }
 
 void UComboBoxString::HandleSelectionChanged(TSharedPtr<FString> Item, ESelectInfo::Type SelectionType)
 {
 	CurrentOptionPtr = Item;
+	BroadcastFieldValueChanged(FFieldNotificationClassDescriptor::SelectedOption);
 	SelectedOption = CurrentOptionPtr.IsValid() ? CurrentOptionPtr.ToSharedRef().Get() : FString();
 
 	// When the selection changes we always generate another widget to represent the content area of the combobox.
diff --git a/Engine/Source/Runtime/UMG/Private/Components/ScrollBar.cpp b/Engine/Source/Runtime/UMG/Private/Components/ScrollBar.cpp
--- a/Engine/Source/Runtime/UMG/Private/Components/ScrollBar.cpp
+++ b/Engine/Source/Runtime/UMG/Private/Components/ScrollBar.cpp
@@ -204,6 +204,11 @@ void UScrollBar::SetPadding(const FMargin& InPadding)
 	}
 }
 
+void UScrollBar::InitOrientation(EOrientation InOrientation)
+{
+	ensureMsgf(!MyScrollBar.IsValid(), TEXT("The widget is already created."));
+	Orientation = InOrientation;
+}
 PRAGMA_ENABLE_DEPRECATION_WARNINGS
 
 #if WITH_EDITOR
diff --git a/Engine/Source/Runtime/UMG/Public/Components/CheckBox.h b/Engine/Source/Runtime/UMG/Public/Components/CheckBox.h
--- a/Engine/Source/Runtime/UMG/Public/Components/CheckBox.h
+++ b/Engine/Source/Runtime/UMG/Public/Components/CheckBox.h
@@ -32,17 +32,18 @@ class UMG_API UCheckBox : public UContentWidget
 	GENERATED_UCLASS_BODY()
 
 public:
-	/** Whether the check box is currently in a checked state */
 	UE_DEPRECATED(5.1, "Direct access to CheckedState is deprecated. Please use the getter or setter.")
+	/** Whether the check box is currently in a checked state */
 	UPROPERTY(EditAnywhere, BlueprintReadWrite, Getter, Setter, BlueprintGetter="GetCheckedState", BlueprintSetter="SetCheckedState", FieldNotify, Category="Appearance")
 	ECheckBoxState CheckedState;
 
+	UE_DEPRECATED(5.2, "Direct access to CheckedStateDelegate is deprecated. Please use the InitCheckedStateDelegate() function.")
 	/** A bindable delegate for the IsChecked. */
 	UPROPERTY()
 	FGetCheckBoxState CheckedStateDelegate;
 
-	/** The checkbox bar style */
 	UE_DEPRECATED(5.1, "Direct access to WidgetStyle is deprecated. Please use the getter or setter.")
+	/** The checkbox bar style */
 	UPROPERTY(EditAnywhere, BlueprintReadWrite, Getter, Setter, Category="Style", meta = (DisplayName="Style"))
 	FCheckBoxStyle WidgetStyle;
 
@@ -50,23 +51,24 @@ public:
 	UPROPERTY(EditAnywhere, BlueprintReadOnly, Category="Appearance")
 	TEnumAsByte<EHorizontalAlignment> HorizontalAlignment;
 
-	/** The type of mouse action required by the user to trigger the buttons 'Click' */
 	UE_DEPRECATED(5.1, "Direct access to ClickMethod is deprecated. Please use the getter or setter.")
+	/** The type of mouse action required by the user to trigger the buttons 'Click' */
 	UPROPERTY(EditAnywhere, BlueprintReadWrite, Getter, Setter, BlueprintSetter="SetClickMethod", Category="Interaction", AdvancedDisplay)
 	TEnumAsByte<EButtonClickMethod::Type> ClickMethod;
 
-	/** The type of touch action required by the user to trigger the buttons 'Click' */
 	UE_DEPRECATED(5.1, "Direct access to TouchMethod is deprecated. Please use the getter or setter.")
+	/** The type of touch action required by the user to trigger the buttons 'Click' */
 	UPROPERTY(EditAnywhere, BlueprintReadWrite, Getter, Setter, BlueprintSetter="SetTouchMethod", Category="Interaction", AdvancedDisplay)
 	TEnumAsByte<EButtonTouchMethod::Type> TouchMethod;
 
-	/** The type of keyboard/gamepad button press action required by the user to trigger the buttons 'Click' */
 	UE_DEPRECATED(5.1, "Direct access to PressMethod is deprecated. Please use the getter or setter.")
+	/** The type of keyboard/gamepad button press action required by the user to trigger the buttons 'Click' */
 	UPROPERTY(EditAnywhere, BlueprintReadWrite, Getter, Setter, BlueprintSetter="SetPressMethod", Category="Interaction", AdvancedDisplay)
 	TEnumAsByte<EButtonPressMethod::Type> PressMethod;
 
+	UE_DEPRECATED(5.2, "Direct access to bIsFocusable is deprecated. Please use the getter. Note that this property is only set at construction and is not modifiable at runtime.")
 	/** Sometimes a button should only be mouse-clickable and never keyboard focusable. */
-	UPROPERTY(EditAnywhere, BlueprintReadOnly, Category="Interaction")
+	UPROPERTY(EditAnywhere, BlueprintReadOnly, Getter, Category="Interaction")
 	bool IsFocusable;
 
 public:
@@ -124,6 +126,9 @@ public:
 	UFUNCTION(BlueprintCallable, Category="Button")
 	void SetPressMethod(EButtonPressMethod::Type InPressMethod);
 
+	/** Is the checkbox focusable. */
+	bool GetIsFocusable() const;
+
 public:
 	
 	//~ Begin UWidget Interface
@@ -145,6 +150,10 @@ protected:
 	virtual void OnSlotRemoved(UPanelSlot* Slot) override;
 	// End UPanelWidget
 
+	/** Initialize IsFocusable in the constructor before the SWidget is constructed. */
+	void InitIsFocusable(bool InIsFocusable);
+
+	void InitCheckedStateDelegate(FGetCheckBoxState InCheckedStateDelegate);
 protected:
 	//~ Begin UWidget Interface
 	virtual TSharedRef<SWidget> RebuildWidget() override;
@@ -165,5 +174,7 @@ protected:
 protected:
 	TSharedPtr<SCheckBox> MyCheckbox;
 
+PRAGMA_DISABLE_DEPRECATION_WARNINGS
 	PROPERTY_BINDING_IMPLEMENTATION(ECheckBoxState, CheckedState)
+PRAGMA_ENABLE_DEPRECATION_WARNINGS
 };
diff --git a/Engine/Source/Runtime/UMG/Public/Components/ComboBoxKey.h b/Engine/Source/Runtime/UMG/Public/Components/ComboBoxKey.h
--- a/Engine/Source/Runtime/UMG/Public/Components/ComboBoxKey.h
+++ b/Engine/Source/Runtime/UMG/Public/Components/ComboBoxKey.h
@@ -32,49 +32,58 @@ private:
 	TArray<FName> Options;
 
 	/** */
-	UPROPERTY(EditAnywhere, Category = Content)
+	UPROPERTY(EditAnywhere, FieldNotify, Category = Content)
 	FName SelectedOption;
 
 public:
+	UE_DEPRECATED(5.2, "Direct access to WidgetStyle is deprecated. Please use the getter. Note that this property is only set at construction and is not modifiable at runtime.")
 	/** The combobox style. */
-	UPROPERTY(EditAnywhere, BlueprintReadOnly, Category = Style, meta = (DisplayName = "Style"))
+	UPROPERTY(EditAnywhere, BlueprintReadOnly, Getter, Category = Style, meta = (DisplayName = "Style"))
 	FComboBoxStyle WidgetStyle;
 
+	UE_DEPRECATED(5.2, "Direct access to ItemStyle is deprecated. Please use the getter. Note that this property is only set at construction and is not modifiable at runtime.")
 	/** The item row style. */
-	UPROPERTY(EditAnywhere, BlueprintReadOnly, Category = Style)
+	UPROPERTY(EditAnywhere, BlueprintReadOnly, Getter, Category = Style)
 	FTableRowStyle ItemStyle;
 
+	UE_DEPRECATED(5.2, "Direct access to ScrollBarStyle is deprecated. Please use the getter. Note that this property is only set at construction and is not modifiable at runtime.")
 	/** The scroll bar style. */
-    UPROPERTY(EditAnywhere, BlueprintReadWrite, Category="Style")
+    UPROPERTY(EditAnywhere, BlueprintReadOnly, Getter, Category="Style")
     FScrollBarStyle ScrollBarStyle;
 
+	UE_DEPRECATED(5.2, "Direct access to ForegroundColor is deprecated. Please use the getter. Note that this property is only set at construction and is not modifiable at runtime.")
 	/** The foreground color to pass through the hierarchy. */
-	UPROPERTY(EditAnywhere, BlueprintReadOnly, Category = Style, meta = (DesignerRebuild))
+	UPROPERTY(EditAnywhere, BlueprintReadOnly, Getter, Category = Style, meta = (DesignerRebuild))
 	FSlateColor ForegroundColor;
 
-	UPROPERTY(EditAnywhere, BlueprintReadOnly, Category = Style)
+	UE_DEPRECATED(5.2, "Direct access to ContentPadding is deprecated. Please use the getter or setter.")
+	UPROPERTY(EditAnywhere, BlueprintReadWrite, Getter, Setter, Category = Style)
 	FMargin ContentPadding;
 
+	UE_DEPRECATED(5.2, "Direct access to MaxListHeight is deprecated. Please use the getter or setter.")
 	/** The max height of the combobox list that opens */
-	UPROPERTY(EditAnywhere, BlueprintReadOnly, Category = Style, AdvancedDisplay)
+	UPROPERTY(EditAnywhere, BlueprintReadWrite, Getter, Setter, Category = Style, AdvancedDisplay)
 	float MaxListHeight;
 
+	UE_DEPRECATED(5.2, "Direct access to bHasDownArrow is deprecated. Please use the getter or setter.")
 	/**
 	 * When false, the down arrow is not generated and it is up to the API consumer
 	 * to make their own visual hint that this is a drop down.
 	 */
-	UPROPERTY(EditAnywhere, BlueprintReadOnly, Category = Style, AdvancedDisplay)
+	UPROPERTY(EditAnywhere, BlueprintReadWrite, Getter = "IsHasDownArrow", Setter = "SetHasDownArrow", Category = Style, AdvancedDisplay)
 	bool bHasDownArrow;
 
+	UE_DEPRECATED(5.2, "Direct access to bEnableGamepadNavigationMode is deprecated. Please use the getter or setter.")
 	/**
 	 * When false, directional keys will change the selection. When true, ComboBox
 	 * must be activated and will only capture arrow input while activated.
 	 */
-	UPROPERTY(EditAnywhere, BlueprintReadOnly, Category = Style, AdvancedDisplay)
+	UPROPERTY(EditAnywhere, BlueprintReadWrite, Getter = "IsEnableGamepadNavigationMode", Setter = "SetEnableGamepadNavigationMode", Category = Style, AdvancedDisplay)
 	bool bEnableGamepadNavigationMode;
 
+	UE_DEPRECATED(5.2, "Direct access to bIsFocusable is deprecated. Please use the getter. Note that this property is only set at construction and is not modifiable at runtime.")
 	/** When true, allows the combo box to receive keyboard focus */
-	UPROPERTY(EditAnywhere, BlueprintReadOnly, Category = Style)
+	UPROPERTY(EditAnywhere, BlueprintReadOnly, Getter = "IsFocusable", Category = Style)
 	bool bIsFocusable;
 
 public: // Events
@@ -123,7 +132,7 @@ public:
 	UFUNCTION(BlueprintCallable, Category = "ComboBox")
 	FName GetSelectedOption() const;
 
-	/** Is the combobox menu openned. */
+	/** Is the combobox menu opened. */
 	UFUNCTION(BlueprintCallable, Category = "ComboBox", Meta = (ReturnDisplayName = "bOpen"))
 	bool IsOpen() const;
 
@@ -131,10 +140,65 @@ public:
 	virtual void ReleaseSlateResources(bool bReleaseChildren) override;
 	//~ End UVisual Interface
 
+	/** Set the padding for content. */
+	void SetContentPadding(FMargin InPadding);
+
+	/** Get the padding for content. */
+	FMargin GetContentPadding() const;
+
+	/** Is the combobox navigated by gamepad. */
+	bool IsEnableGamepadNavigationMode() const;
+
+	/** Set whether the combobox is navigated by gamepad. */
+	void SetEnableGamepadNavigationMode(bool InEnableGamepadNavigationMode);
+
+	/** Is the combobox arrow showing. */
+	bool IsHasDownArrow() const;
+
+	/** Set whether the combobox arrow is showing. */
+	void SetHasDownArrow(bool InHasDownArrow);
+
+	/** Get the maximum height of the combobox list. */
+	float GetMaxListHeight() const;
+
+	/** Set the maximum height of the combobox list. */
+	void SetMaxListHeight(float InMaxHeight);
+
+	/** Get the style of the combobox. */
+	const FComboBoxStyle& GetWidgetStyle() const;
+
+	/** Get the style of the items. */
+	const FTableRowStyle& GetItemStyle() const;
+
+	/** Get the style of the scrollbar. */
+	const FScrollBarStyle& GetScrollBarStyle() const;
+
+	/** Is the combobox focusable. */
+	bool IsFocusable() const;
+
+	/** Get the foreground color of the button. */
+	FSlateColor GetForegroundColor() const;
+
 #if WITH_EDITOR
 	virtual const FText GetPaletteCategory() override;
 #endif
 
+protected:
+	/** Initialize the widget style in the constructor before the SWidget is constructed. */
+	void InitWidgetStyle(const FComboBoxStyle& InWidgetStyle);
+
+	/** Initialize the item style in the constructor before the SWidget is constructed. */
+	void InitItemStyle(const FTableRowStyle& InItemStyle);
+
+	/** Initialize the scrollbar style in the constructor before the SWidget is constructed. */
+	void InitScrollBarStyle(const FScrollBarStyle& InScrollBarStyle);
+
+	/** Initialize IsFocusable in the constructor before the SWidget is constructed. */
+	void InitIsFocusable(bool InIsFocusable);
+
+	/** Initialize ForegroundColor in the constructor before the SWidget is constructed. */
+	void InitForegroundColor(FSlateColor InForegroundColor);
+
 private:
 	/** Called by slate when it needs to generate the widget in the content box */
 	void GenerateContent();
diff --git a/Engine/Source/Runtime/UMG/Public/Components/ComboBoxString.h b/Engine/Source/Runtime/UMG/Public/Components/ComboBoxString.h
--- a/Engine/Source/Runtime/UMG/Public/Components/ComboBoxString.h
+++ b/Engine/Source/Runtime/UMG/Public/Components/ComboBoxString.h
@@ -32,56 +32,66 @@ private:
 	TArray<FString> DefaultOptions;
 
 	/** The item in the combobox to select by default */
-	UPROPERTY(EditAnywhere, Category=Content)
+	UPROPERTY(EditAnywhere, FieldNotify, Category=Content)
 	FString SelectedOption;
 
 public:
 
+	UE_DEPRECATED(5.2, "Direct access to WidgetStyle is deprecated. Please use the getter. Note that this property is only set at construction and is not modifiable at runtime.")
 	/** The style. */
-	UPROPERTY(EditAnywhere, BlueprintReadWrite, Category=Style, meta=( DisplayName="Style" ))
+	UPROPERTY(EditAnywhere, BlueprintReadOnly, Getter, Category=Style, meta=( DisplayName="Style" ))
 	FComboBoxStyle WidgetStyle;
 
+	UE_DEPRECATED(5.2, "Direct access to ItemStyle is deprecated. Please use the getter. Note that this property is only set at construction and is not modifiable at runtime.")
 	/** The item row style. */
-	UPROPERTY(EditAnywhere, BlueprintReadWrite, Category=Style)
+	UPROPERTY(EditAnywhere, BlueprintReadOnly, Getter, Category=Style)
 	FTableRowStyle ItemStyle;
 	
+	UE_DEPRECATED(5.2, "Direct access to ScrollBarStyle is deprecated. Please use the getter. Note that this property is only set at construction and is not modifiable at runtime.")
 	/** The scroll bar style. */
-	UPROPERTY(EditAnywhere, BlueprintReadWrite, Category="Style")
+	UPROPERTY(EditAnywhere, BlueprintReadOnly, Getter, Category="Style")
     FScrollBarStyle ScrollBarStyle;
 
-	UPROPERTY(EditAnywhere, BlueprintReadOnly, Category=Content)
+	UE_DEPRECATED(5.2, "Direct access to ContentPadding is deprecated. Please use the getter or setter.")
+	UPROPERTY(EditAnywhere, BlueprintReadWrite, Getter, Setter, Category=Content)
 	FMargin ContentPadding;
 
+	UE_DEPRECATED(5.2, "Direct access to MaxListHeight is deprecated. Please use the getter or setter.")
 	/** The max height of the combobox list that opens */
-	UPROPERTY(EditAnywhere, BlueprintReadOnly, Category=Content, AdvancedDisplay)
+	UPROPERTY(EditAnywhere, BlueprintReadWrite, Getter, Setter, Category=Content, AdvancedDisplay)
 	float MaxListHeight;
 
+	UE_DEPRECATED(5.2, "Direct access to HasDownArrow is deprecated. Please use the getter or setter.")
 	/**
 	 * When false, the down arrow is not generated and it is up to the API consumer
 	 * to make their own visual hint that this is a drop down.
 	 */
-	UPROPERTY(EditAnywhere, BlueprintReadOnly, Category=Content, AdvancedDisplay)
+	UPROPERTY(EditAnywhere, BlueprintReadWrite, Getter = "IsHasDownArrow", Setter = "SetHasDownArrow", Category = Content, AdvancedDisplay)
 	bool HasDownArrow;
 
+	UE_DEPRECATED(5.2, "Direct access to EnableGamepadNavigationMode is deprecated. Please use the getter or setter.")
 	/**
 	* When false, directional keys will change the selection. When true, ComboBox 
 	* must be activated and will only capture arrow input while activated.
 	*/
-	UPROPERTY(EditAnywhere, BlueprintReadOnly, Category=Content, AdvancedDisplay)
+	UPROPERTY(EditAnywhere, BlueprintReadWrite, Getter = "IsEnableGamepadNavigationMode", Setter = "SetEnableGamepadNavigationMode", Category = Content, AdvancedDisplay)
 	bool EnableGamepadNavigationMode;
 
+	UE_DEPRECATED(5.2, "Direct access to Font is deprecated. Please use the getter. Note that this property is only set at construction and is not modifiable at runtime.")
 	/**
 	 * The default font to use in the combobox, only applies if you're not implementing OnGenerateWidgetEvent
 	 * to factory each new entry.
 	 */
-	UPROPERTY(EditAnywhere, BlueprintReadOnly, Category=Style)
+	UPROPERTY(EditAnywhere, BlueprintReadOnly, Getter, Category=Style)
 	FSlateFontInfo Font;
 
+	UE_DEPRECATED(5.2, "Direct access to ForegroundColor is deprecated. Please use the getter. Note that this property is only set at construction and is not modifiable at runtime.")
 	/** The foreground color to pass through the hierarchy. */
-	UPROPERTY(EditAnywhere, BlueprintReadOnly, Category=Style, meta=(DesignerRebuild))
+	UPROPERTY(EditAnywhere, BlueprintReadOnly, Getter, Category=Style, meta=(DesignerRebuild))
 	FSlateColor ForegroundColor;
 
-	UPROPERTY(EditAnywhere, BlueprintReadOnly, Category=Interaction)
+	UE_DEPRECATED(5.2, "Direct access to bIsFocusable is deprecated. Please use the getter. Note that this property is only set at construction and is not modifiable at runtime.")
+	UPROPERTY(EditAnywhere, BlueprintReadOnly, Getter= "IsFocusable", Category=Interaction)
 	bool bIsFocusable;
 
 public: // Events
@@ -148,6 +158,48 @@ public:
 	virtual void ReleaseSlateResources(bool bReleaseChildren) override;
 	//~ End UVisual Interface
 
+	/** Set the padding for content. */
+	void SetContentPadding(FMargin InPadding);
+
+	/** Get the padding for content. */
+	FMargin GetContentPadding() const;
+
+	/** Is the combobox navigated by gamepad. */
+	bool IsEnableGamepadNavigationMode() const;
+
+	/** Set whether the combobox is navigated by gamepad. */
+	void SetEnableGamepadNavigationMode(bool InEnableGamepadNavigationMode);
+
+	/** Is the combobox arrow showing. */
+	bool IsHasDownArrow() const;
+
+	/** Set whether the combobox arrow is showing. */
+	void SetHasDownArrow(bool InHasDownArrow);
+
+	/** Get the maximum height of the combobox list. */
+	float GetMaxListHeight() const;
+
+	/** Set the maximum height of the combobox list. */
+	void SetMaxListHeight(float InMaxHeight);
+
+	/** Get the default font for Combobox if no OnGenerateWidgetEvent is specified. */
+	const FSlateFontInfo& GetFont() const;
+
+	/** Get the style of the combobox. */
+	const FComboBoxStyle& GetWidgetStyle() const;
+
+	/** Get the style of the items. */
+	const FTableRowStyle& GetItemStyle() const;
+
+	/** Get the style of the scrollbar. */
+	const FScrollBarStyle& GetScrollBarStyle() const;
+
+	/** Is the combobox focusable. */
+	bool IsFocusable() const;
+
+	/** Get the foreground color of the button. */
+	FSlateColor GetForegroundColor() const;
+
 	//~ Begin UObject Interface
 	virtual void PostInitProperties() override;
 	virtual void Serialize(FArchive& Ar) override;
@@ -175,6 +227,24 @@ protected:
 	virtual TSharedRef<SWidget> RebuildWidget() override;
 	//~ End UWidget Interface
 
+	/** Initialize the widget style in the constructor before the SWidget is constructed. */
+	void InitWidgetStyle(const FComboBoxStyle& InWidgetStyle);
+
+	/** Initialize the item style in the constructor before the SWidget is constructed. */
+	void InitItemStyle(const FTableRowStyle& InItemStyle);
+
+	/** Initialize the scrollbar style in the constructor before the SWidget is constructed. */
+	void InitScrollBarStyle(const FScrollBarStyle& InScrollBarStyle);
+
+	/** Initialize the default font in the constructor before the SWidget is constructed. */
+	void InitFont(FSlateFontInfo InFont);
+
+	/** Initialize IsFocusable in the constructor before the SWidget is constructed. */
+	void InitIsFocusable(bool InIsFocusable);
+
+	/** Initialize ForegroundColor in the constructor before the SWidget is constructed. */
+	void InitForegroundColor(FSlateColor InForegroundColor);
+
 protected:
 	/** The true objects bound to the Slate combobox. */
 	TArray< TSharedPtr<FString> > Options;
diff --git a/Engine/Source/Runtime/UMG/Public/Components/ScrollBar.h b/Engine/Source/Runtime/UMG/Public/Components/ScrollBar.h
--- a/Engine/Source/Runtime/UMG/Public/Components/ScrollBar.h
+++ b/Engine/Source/Runtime/UMG/Public/Components/ScrollBar.h
@@ -132,4 +132,7 @@ protected:
 	//~ Begin UWidget Interface
 	virtual TSharedRef<SWidget> RebuildWidget() override;
 	//~ End UWidget Interface
+
+	/** Initialize the orientation of the scrollbar in the constructor before the SWidget is constructed. */
+	void InitOrientation(EOrientation InOrientation);
 };
