Add ensure to init functions in UWidgets

diff --git a/Engine/Source/Runtime/UMG/Private/Components/DynamicEntryBoxBase.cpp b/Engine/Source/Runtime/UMG/Private/Components/DynamicEntryBoxBase.cpp
--- a/Engine/Source/Runtime/UMG/Private/Components/DynamicEntryBoxBase.cpp
+++ b/Engine/Source/Runtime/UMG/Private/Components/DynamicEntryBoxBase.cpp
@@ -506,26 +506,31 @@ void UDynamicEntryBoxBase::AddEntryChild(UUserWidget& ChildWidget)
 
 void UDynamicEntryBoxBase::InitEntryBoxType(EDynamicBoxType InEntryBoxType)
 {
+	ensureMsgf(!MyPanelWidget.IsValid(), TEXT("The widget is already created."));
 	EntryBoxType = InEntryBoxType;
 }
 
 void UDynamicEntryBoxBase::InitEntrySizeRule(FSlateChildSize InEntrySizeRule)
 {
+	ensureMsgf(!MyPanelWidget.IsValid(), TEXT("The widget is already created."));
 	EntrySizeRule = InEntrySizeRule;
 }
 
 void UDynamicEntryBoxBase::InitEntryHorizontalAlignment(EHorizontalAlignment InEntryHorizontalAlignment)
 {
+	ensureMsgf(!MyPanelWidget.IsValid(), TEXT("The widget is already created."));
 	EntryHorizontalAlignment = InEntryHorizontalAlignment;
 }
 
 void UDynamicEntryBoxBase::InitEntryVerticalAlignment(EVerticalAlignment InEntryVerticalAlignment)
 {
+	ensureMsgf(!MyPanelWidget.IsValid(), TEXT("The widget is already created."));
 	EntryVerticalAlignment = InEntryVerticalAlignment;
 }
 
 void UDynamicEntryBoxBase::InitMaxElementSize(int32 InMaxElementSize)
 {
+	ensureMsgf(!MyPanelWidget.IsValid(), TEXT("The widget is already created."));
 	MaxElementSize = InMaxElementSize;
 }
 PRAGMA_ENABLE_DEPRECATION_WARNINGS
diff --git a/Engine/Source/Runtime/UMG/Private/Components/MenuAnchor.cpp b/Engine/Source/Runtime/UMG/Private/Components/MenuAnchor.cpp
--- a/Engine/Source/Runtime/UMG/Private/Components/MenuAnchor.cpp
+++ b/Engine/Source/Runtime/UMG/Private/Components/MenuAnchor.cpp
@@ -189,11 +189,13 @@ bool UMenuAnchor::IsUseApplicationMenuStack() const
 
 void UMenuAnchor::InitShouldDeferPaintingAfterWindowContent(bool InShouldDeferPaintingAfterWindowContent)
 {
+	ensureMsgf(!MyMenuAnchor.IsValid(), TEXT("The widget is already created."));
 	ShouldDeferPaintingAfterWindowContent = InShouldDeferPaintingAfterWindowContent;
 }
 
 void UMenuAnchor::InitUseApplicationMenuStack(bool InUseApplicationMenuStack)
 {
+	ensureMsgf(!MyMenuAnchor.IsValid(), TEXT("The widget is already created."));
 	UseApplicationMenuStack = InUseApplicationMenuStack;
 }
 
diff --git a/Engine/Source/Runtime/UMG/Private/Components/RetainerBox.cpp b/Engine/Source/Runtime/UMG/Private/Components/RetainerBox.cpp
--- a/Engine/Source/Runtime/UMG/Private/Components/RetainerBox.cpp
+++ b/Engine/Source/Runtime/UMG/Private/Components/RetainerBox.cpp
@@ -215,6 +215,32 @@ FGeometry URetainerBox::GetCachedAllottedGeometry() const
 	return TempGeo;
 }
 
+PRAGMA_DISABLE_DEPRECATION_WARNINGS
+void URetainerBox::InitRenderOnInvalidation(bool InRenderOnInvalidation)
+{
+	ensureMsgf(!MyRetainerWidget.IsValid(), TEXT("The widget is already created."));
+	RenderOnInvalidation = InRenderOnInvalidation;
+}
+
+void URetainerBox::InitRenderOnPhase(bool InRenderOnPhase)
+{
+	ensureMsgf(!MyRetainerWidget.IsValid(), TEXT("The widget is already created."));
+	RenderOnPhase = InRenderOnPhase;
+}
+
+void URetainerBox::InitPhase(int32 InPhase)
+{
+	ensureMsgf(!MyRetainerWidget.IsValid(), TEXT("The widget is already created."));
+	Phase = InPhase;
+}
+
+void URetainerBox::InitPhaseCount(int32 InPhaseCount)
+{
+	ensureMsgf(!MyRetainerWidget.IsValid(), TEXT("The widget is already created."));
+	PhaseCount = InPhaseCount;
+}
+PRAGMA_ENABLE_DEPRECATION_WARNINGS
+
 #if WITH_EDITOR
 bool URetainerBox::CanEditChange(const FProperty* InProperty) const
 {
diff --git a/Engine/Source/Runtime/UMG/Public/Components/DynamicEntryBoxBase.h b/Engine/Source/Runtime/UMG/Public/Components/DynamicEntryBoxBase.h
--- a/Engine/Source/Runtime/UMG/Public/Components/DynamicEntryBoxBase.h
+++ b/Engine/Source/Runtime/UMG/Public/Components/DynamicEntryBoxBase.h
@@ -49,27 +49,27 @@ protected:
 	UPROPERTY(EditAnywhere, BlueprintReadOnly, Category = EntryLayout)
 	TArray<FVector2D> SpacingPattern;
 
-	UE_DEPRECATED(5.2, "Direct access to EntryBoxType is deprecated. Please use the getter.")
+	UE_DEPRECATED(5.2, "Direct access to EntryBoxType is deprecated. Please use the getter. Note that this property is only set at construction and is not modifiable at runtime.")
 	/** The type of box panel into which created entries are added. Some differences in functionality exist between each type. */
 	UPROPERTY(EditAnywhere, BlueprintReadOnly, Getter = "GetBoxType", Category = DynamicEntryBox, meta = (DesignerRebuild))
 	EDynamicBoxType EntryBoxType;
 
-	UE_DEPRECATED(5.2, "Direct access to EntrySizeRule is deprecated. Please use the getter.")
+	UE_DEPRECATED(5.2, "Direct access to EntrySizeRule is deprecated. Please use the getter. Note that this property is only set at construction and is not modifiable at runtime.")
 	/** Sizing rule to apply to generated entries. Horizontal/Vertical boxes only. */
 	UPROPERTY(EditAnywhere, BlueprintReadOnly, Getter, Category = EntryLayout)
 	FSlateChildSize EntrySizeRule;
 
-	UE_DEPRECATED(5.2, "Direct access to EntryHorizontalAlignment is deprecated. Please use the getter.")
+	UE_DEPRECATED(5.2, "Direct access to EntryHorizontalAlignment is deprecated. Please use the getter. Note that this property is only set at construction and is not modifiable at runtime.")
 	/** Horizontal alignment of generated entries. Horizontal/Vertical/Wrap boxes only. */
 	UPROPERTY(EditAnywhere, BlueprintReadOnly, Getter, Category = EntryLayout)
 	TEnumAsByte<EHorizontalAlignment> EntryHorizontalAlignment;
 
-	UE_DEPRECATED(5.2, "Direct access to EntryVerticalAlignment is deprecated. Please use the getter.")
+	UE_DEPRECATED(5.2, "Direct access to EntryVerticalAlignment is deprecated. Please use the getter. Note that this property is only set at construction and is not modifiable at runtime.")
 	/** Vertical alignment of generated entries. Horizontal/Vertical/Wrap boxes only. */
 	UPROPERTY(EditAnywhere, BlueprintReadOnly, Getter, Category = EntryLayout)
 	TEnumAsByte<EVerticalAlignment> EntryVerticalAlignment;
 
-	UE_DEPRECATED(5.2, "Direct access to MaxElementSize is deprecated. Please use the getter.")
+	UE_DEPRECATED(5.2, "Direct access to MaxElementSize is deprecated. Please use the getter. Note that this property is only set at construction and is not modifiable at runtime.")
 	/** The maximum size of each entry in the dominant axis of the box. Vertical/Horizontal boxes only. */
 	UPROPERTY(EditAnywhere, BlueprintReadOnly, Getter, Category = EntryLayout)
 	int32 MaxElementSize = 0;
diff --git a/Engine/Source/Runtime/UMG/Public/Components/RetainerBox.h b/Engine/Source/Runtime/UMG/Public/Components/RetainerBox.h
--- a/Engine/Source/Runtime/UMG/Public/Components/RetainerBox.h
+++ b/Engine/Source/Runtime/UMG/Public/Components/RetainerBox.h
@@ -33,7 +33,7 @@ protected:
 	bool bRetainRender = true;
 
 public:
-	UE_DEPRECATED(5.2, "Direct access to RenderOnInvalidation is deprecated. Please use the getter.")
+	UE_DEPRECATED(5.2, "Direct access to RenderOnInvalidation is deprecated. Please use the getter. Note that this property is only set at construction and is not modifiable at runtime.")
 	/**
 	 * Should this widget redraw the contents it has every time it receives an invalidation request
 	 * from it's children, similar to the invalidation panel.
@@ -41,14 +41,14 @@ public:
 	UPROPERTY(EditAnywhere, BlueprintReadOnly, Getter = "IsRenderOnInvalidation", Category = "Render Rules", meta = (EditCondition = bRetainRender))
 	bool RenderOnInvalidation;
 
-	UE_DEPRECATED(5.2, "Direct access to RenderOnPhase is deprecated. Please use the getter.")
+	UE_DEPRECATED(5.2, "Direct access to RenderOnPhase is deprecated. Please use the getter. Note that this property is only set at construction and is not modifiable at runtime.")
 	/**
 	 * Should this widget redraw the contents it has every time the phase occurs.
 	 */
 	UPROPERTY(EditAnywhere, BlueprintReadOnly, Getter = "IsRenderOnPhase", Category = "Render Rules", meta = (EditCondition = bRetainRender))
 	bool RenderOnPhase;
 
-	UE_DEPRECATED(5.2, "Direct access to Phase is deprecated. Please use the getter.")
+	UE_DEPRECATED(5.2, "Direct access to Phase is deprecated. Please use the getter. Note that this property is only set at construction and is not modifiable at runtime.")
 	/**
 	 * The Phase this widget will draw on.
 	 *
@@ -59,7 +59,7 @@ public:
 	UPROPERTY(EditAnywhere, BlueprintReadOnly, Getter, Category="Render Rules", meta=(UIMin=0, ClampMin=0))
 	int32 Phase;
 
-	UE_DEPRECATED(5.2, "Direct access to PhaseCount is deprecated. Please use the getter.")
+	UE_DEPRECATED(5.2, "Direct access to PhaseCount is deprecated. Please use the getter. Note that this property is only set at construction and is not modifiable at runtime.")
 	/**
 	 * The PhaseCount controls how many phases are possible know what to modulus the current frame 
 	 * count by to determine if this is the current frame to draw the widget on.
@@ -179,6 +179,17 @@ protected:
 #endif
 	//~ End UObject interface
 
+	// Initialize RenderOnInvalidation in the constructor before the SWidget is constructed.
+	void InitRenderOnInvalidation(bool InRenderOnInvalidation);
+
+	// Initialize RenderOnPhase in the constructor before the SWidget is constructed.
+	void InitRenderOnPhase(bool InRenderOnPhase);
+
+	// Initialize Phase in the constructor before the SWidget is constructed.
+	void InitPhase(int32 InPhase);
+
+	// Initialize PhaseCount in the constructor before the SWidget is constructed.
+	void InitPhaseCount(int32 InPhaseCount);
 public:
 	/**
 	* Gets the name of the texture parameter to set the render target to on the material.
diff --git a/Engine/Source/Runtime/UMG/Public/Components/ScrollBar.h b/Engine/Source/Runtime/UMG/Public/Components/ScrollBar.h
--- a/Engine/Source/Runtime/UMG/Public/Components/ScrollBar.h
+++ b/Engine/Source/Runtime/UMG/Public/Components/ScrollBar.h
@@ -61,21 +61,6 @@ public:
 	UFUNCTION(BlueprintCallable, Category="Scrolling")
 	void SetState(float InOffsetFraction, float InThumbSizeFraction);
 
-	///** @return true if scrolling is possible; false if the view is big enough to fit all the content */
-	//bool IsNeeded() const;
-
-	///** @return normalized distance from top */
-	//float DistanceFromTop() const;
-
-	///** @return normalized distance from bottom */
-	//float DistanceFromBottom() const;
-
-	///** @return the scrollbar's visibility as a product of internal rules and user-specified visibility */
-	//EVisibility ShouldBeVisible() const;
-
-	///** @return True if the user is scrolling by dragging the scroll bar thumb. */
-	//bool IsScrolling() const;
-
 	//~ Begin UWidget Interface
 	virtual void SynchronizeProperties() override;
 	//~ End UWidget Interface
