Add accessors for UserWidget, ScrollBox, Button

diff --git a/Engine/Plugins/Experimental/CommonUI/Source/CommonUI/Private/CommonButtonBase.cpp b/Engine/Plugins/Experimental/CommonUI/Source/CommonUI/Private/CommonButtonBase.cpp
--- a/Engine/Plugins/Experimental/CommonUI/Source/CommonUI/Private/CommonButtonBase.cpp
+++ b/Engine/Plugins/Experimental/CommonUI/Source/CommonUI/Private/CommonButtonBase.cpp
@@ -161,7 +161,7 @@ UCommonButtonInternalBase::UCommonButtonInternalBase(const FObjectInitializer& I
 	{
 		// SButton will have a transparent background and have no padding if Button Style is set to None
 		static const FButtonStyle* TransparentButtonStyle = new FButtonStyle(FUMGCoreStyle::Get().GetWidgetStyle<FButtonStyle>("NoBorder"));
-		WidgetStyle = *TransparentButtonStyle;
+		SetStyle(*TransparentButtonStyle);
 	}
 }
 
@@ -176,7 +176,7 @@ void UCommonButtonInternalBase::SetButtonEnabled(bool bInIsButtonEnabled)
 
 void UCommonButtonInternalBase::SetButtonFocusable(bool bInIsButtonFocusable)
 {
-	IsFocusable = bInIsButtonFocusable;
+	InitIsFocusable(bInIsButtonFocusable);
 	if (MyCommonButton.IsValid())
 	{
 		MyCommonButton->SetIsButtonFocusable(bInIsButtonFocusable);
@@ -240,10 +240,10 @@ TSharedRef<SWidget> UCommonButtonInternalBase::RebuildWidget()
 		.OnDoubleClicked(BIND_UOBJECT_DELEGATE(FOnClicked, SlateHandleDoubleClicked))
 		.OnPressed(BIND_UOBJECT_DELEGATE(FSimpleDelegate, SlateHandlePressedOverride))
 		.OnReleased(BIND_UOBJECT_DELEGATE(FSimpleDelegate, SlateHandleReleasedOverride))
-		.ButtonStyle(&WidgetStyle)
-		.ClickMethod(ClickMethod)
-		.TouchMethod(TouchMethod)
-		.IsFocusable(IsFocusable)
+		.ButtonStyle(&GetStyle())
+		.ClickMethod(GetClickMethod())
+		.TouchMethod(GetTouchMethod())
+		.IsFocusable(GetIsFocusable())
 		.IsButtonEnabled(bButtonEnabled)
 		.IsInteractionEnabled(bInteractionEnabled)
 		.OnReceivedFocus(BIND_UOBJECT_DELEGATE(FSimpleDelegate, SlateHandleOnReceivedFocus))
@@ -335,7 +335,7 @@ UCommonButtonBase::UCommonButtonBase(const FObjectInitializer& ObjectInitializer
 	, bButtonEnabled(true)
 	, bInteractionEnabled(true)
 {
-	bIsFocusable = true;
+	InitIsFocusable(true);
 }
 
 void UCommonButtonBase::OnWidgetRebuilt()
@@ -394,10 +394,10 @@ bool UCommonButtonBase::Initialize()
 	{
 		UCommonButtonInternalBase* RootButtonRaw = ConstructInternalButton();
 
-		RootButtonRaw->ClickMethod = ClickMethod;
-		RootButtonRaw->TouchMethod = TouchMethod;
-		RootButtonRaw->PressMethod = PressMethod;
-		RootButtonRaw->IsFocusable = bIsFocusable;
+		RootButtonRaw->SetClickMethod(ClickMethod);
+		RootButtonRaw->SetTouchMethod(TouchMethod);
+		RootButtonRaw->SetPressMethod(PressMethod);
+		RootButtonRaw->SetButtonFocusable(IsFocusable());
 		RootButtonRaw->SetButtonEnabled(bButtonEnabled);
 		RootButtonRaw->SetInteractionEnabled(bInteractionEnabled);
 		RootButton = RootButtonRaw;
@@ -561,7 +561,7 @@ void UCommonButtonBase::SetIsInteractionEnabled(bool bInIsInteractionEnabled)
 
 		if (bApplyAlphaOnDisable)
 		{
-			FLinearColor ButtonColor = RootButton->ColorAndOpacity;
+			FLinearColor ButtonColor = RootButton->GetColorAndOpacity();
 			ButtonColor.A = 1.f;
 			RootButton->SetColorAndOpacity(ButtonColor);
 		}
@@ -572,7 +572,7 @@ void UCommonButtonBase::SetIsInteractionEnabled(bool bInIsInteractionEnabled)
 
 		if (bApplyAlphaOnDisable)
 		{
-			FLinearColor ButtonColor = RootButton->ColorAndOpacity;
+			FLinearColor ButtonColor = RootButton->GetColorAndOpacity();
 			ButtonColor.A = 0.5f;
 			RootButton->SetColorAndOpacity(ButtonColor);
 		}
@@ -1540,7 +1540,7 @@ void UCommonButtonBase::DisableButton()
 
 void UCommonButtonBase::SetIsFocusable(bool bInIsFocusable)
 {
-	bIsFocusable = bInIsFocusable;
+	InitIsFocusable(bInIsFocusable);
 
 	if (RootButton.IsValid())
 	{
@@ -1550,6 +1550,6 @@ void UCommonButtonBase::SetIsFocusable(bool bInIsFocusable)
 
 bool UCommonButtonBase::GetIsFocusable() const
 {
-	return bIsFocusable;
+	return IsFocusable();
 }
 
diff --git a/Engine/Plugins/Experimental/CommonUI/Source/CommonUI/Private/CommonHierarchicalScrollBox.cpp b/Engine/Plugins/Experimental/CommonUI/Source/CommonUI/Private/CommonHierarchicalScrollBox.cpp
--- a/Engine/Plugins/Experimental/CommonUI/Source/CommonUI/Private/CommonHierarchicalScrollBox.cpp
+++ b/Engine/Plugins/Experimental/CommonUI/Source/CommonUI/Private/CommonHierarchicalScrollBox.cpp
@@ -12,20 +12,20 @@
 UCommonHierarchicalScrollBox::UCommonHierarchicalScrollBox(const FObjectInitializer& ObjectInitializer)
 	: Super(ObjectInitializer)
 {
-	bAnimateWheelScrolling = true;
+	SetAnimateWheelScrolling(true);
 }
 
 TSharedRef<SWidget> UCommonHierarchicalScrollBox::RebuildWidget()
 {
 	MyScrollBox = SNew(SCommonHierarchicalScrollBox)
-		.Style(&WidgetStyle)
-		.ScrollBarStyle(&WidgetBarStyle)
-		.Orientation(Orientation)
-		.ConsumeMouseWheel(ConsumeMouseWheel)
-		.NavigationDestination(NavigationDestination)
-		.NavigationScrollPadding(NavigationScrollPadding)
+		.Style(&GetWidgetStyle())
+		.ScrollBarStyle(&GetWidgetBarStyle())
+		.Orientation(GetOrientation())
+		.ConsumeMouseWheel(GetConsumeMouseWheel())
+		.NavigationDestination(GetNavigationDestination())
+		.NavigationScrollPadding(GetNavigationScrollPadding())
 		.OnUserScrolled(BIND_UOBJECT_DELEGATE(FOnUserScrolled, SlateHandleUserScrolled))
-		.AnimateWheelScrolling(bAnimateWheelScrolling);
+		.AnimateWheelScrolling(IsAnimateWheelScrolling());
 
 	for ( UPanelSlot* PanelSlot : Slots )
 	{
diff --git a/Engine/Plugins/Experimental/CommonUI/Source/CommonUI/Private/CommonRotator.cpp b/Engine/Plugins/Experimental/CommonUI/Source/CommonUI/Private/CommonRotator.cpp
--- a/Engine/Plugins/Experimental/CommonUI/Source/CommonUI/Private/CommonRotator.cpp
+++ b/Engine/Plugins/Experimental/CommonUI/Source/CommonUI/Private/CommonRotator.cpp
@@ -12,7 +12,7 @@
 // UCommonRotator
 UCommonRotator::UCommonRotator(const FObjectInitializer& ObjectInitializer)
 {
-	bIsFocusable = true;
+	SetIsFocusable(true);
 }
 
 bool UCommonRotator::Initialize()
diff --git a/Engine/Plugins/Experimental/CommonUI/Source/CommonUI/Private/Input/CommonAnalogCursor.cpp b/Engine/Plugins/Experimental/CommonUI/Source/CommonUI/Private/Input/CommonAnalogCursor.cpp
--- a/Engine/Plugins/Experimental/CommonUI/Source/CommonUI/Private/Input/CommonAnalogCursor.cpp
+++ b/Engine/Plugins/Experimental/CommonUI/Source/CommonUI/Private/Input/CommonAnalogCursor.cpp
@@ -433,7 +433,7 @@ EOrientation FCommonAnalogCursor::DetermineScrollOrientation(const UWidget& Widg
 	}
 	else if (const UScrollBox* AsScrollBox = Cast<const UScrollBox>(&Widget))
 	{
-		return AsScrollBox->Orientation;
+		return AsScrollBox->GetOrientation();
 	}
 	return EOrientation::Orient_Vertical;
 }
diff --git a/Engine/Plugins/Experimental/CommonUI/Source/CommonUI/Private/Input/UIActionRouterTypes.cpp b/Engine/Plugins/Experimental/CommonUI/Source/CommonUI/Private/Input/UIActionRouterTypes.cpp
--- a/Engine/Plugins/Experimental/CommonUI/Source/CommonUI/Private/Input/UIActionRouterTypes.cpp
+++ b/Engine/Plugins/Experimental/CommonUI/Source/CommonUI/Private/Input/UIActionRouterTypes.cpp
@@ -1307,7 +1307,7 @@ void FActivatableTreeRoot::FocusLeafmostNode()
 		}
 		else
 		{
-			if (LeafWidget->bIsFocusable)
+			if (LeafWidget->IsFocusable())
 			{
 				UE_LOG(LogUIActionRouter, Display, TEXT("[User %d] No focus target for leaf-most node [%s] - setting focus directly to the widget as a last resort."), OwnerSlateId, *LeafWidget->GetName());
 				LeafWidget->SetFocus();
diff --git a/Engine/Source/Runtime/Slate/Private/Widgets/Layout/SScrollBar.cpp b/Engine/Source/Runtime/Slate/Private/Widgets/Layout/SScrollBar.cpp
--- a/Engine/Source/Runtime/Slate/Private/Widgets/Layout/SScrollBar.cpp
+++ b/Engine/Source/Runtime/Slate/Private/Widgets/Layout/SScrollBar.cpp
@@ -369,7 +369,12 @@ void SScrollBar::SetStyle(const FScrollBarStyle* InStyle)
 		BottomBrush = &Style->HorizontalBottomSlotImage;
 	}
 
-	Invalidate(EInvalidateWidget::Layout);
+	InvalidateStyle();
+}
+
+void SScrollBar::InvalidateStyle()
+{
+	Invalidate(EInvalidateWidgetReason::Layout);
 }
 
 void SScrollBar::SetDragFocusCause(EFocusCause InDragFocusCause)
diff --git a/Engine/Source/Runtime/Slate/Private/Widgets/Layout/SScrollBox.cpp b/Engine/Source/Runtime/Slate/Private/Widgets/Layout/SScrollBox.cpp
--- a/Engine/Source/Runtime/Slate/Private/Widgets/Layout/SScrollBox.cpp
+++ b/Engine/Source/Runtime/Slate/Private/Widgets/Layout/SScrollBox.cpp
@@ -465,6 +465,40 @@ bool SScrollBox::InternalScrollDescendantIntoView(const FGeometry& MyGeometry, c
 	return false;
 }
 
+void SScrollBox::SetStyle(const FScrollBoxStyle* InStyle)
+{
+	if (Style != InStyle)
+	{
+		Style = InStyle;
+		InvalidateStyle();
+	}
+}
+
+void SScrollBox::SetScrollBarStyle(const FScrollBarStyle* InBarStyle)
+{
+	if (InBarStyle != ScrollBarStyle)
+	{
+		ScrollBarStyle = InBarStyle;
+		if (!bScrollBarIsExternal && ScrollBar.IsValid())
+		{
+			ScrollBar->SetStyle(ScrollBarStyle);
+		}
+	}
+}
+
+void SScrollBox::InvalidateStyle()
+{
+	Invalidate(EInvalidateWidgetReason::Layout);
+}
+
+void SScrollBox::InvalidateScrollBarStyle()
+{
+	if (ScrollBar.IsValid())
+	{
+		ScrollBar->InvalidateStyle();
+	}
+}
+
 EOrientation SScrollBox::GetOrientation()
 {
 	return Orientation;
diff --git a/Engine/Source/Runtime/Slate/Public/Widgets/Layout/SScrollBar.h b/Engine/Source/Runtime/Slate/Public/Widgets/Layout/SScrollBar.h
--- a/Engine/Source/Runtime/Slate/Public/Widgets/Layout/SScrollBar.h
+++ b/Engine/Source/Runtime/Slate/Public/Widgets/Layout/SScrollBar.h
@@ -118,6 +118,9 @@ public:
 	/** Set argument Style */
 	void SetStyle(const FScrollBarStyle* InStyle);
 
+	/** Invalidate the style */
+	void InvalidateStyle();
+
 	/** Set UserVisibility attribute */
 	void SetUserVisibility(TAttribute<EVisibility> InUserVisibility) { UserVisibility = InUserVisibility; }
 
diff --git a/Engine/Source/Runtime/Slate/Public/Widgets/Layout/SScrollBox.h b/Engine/Source/Runtime/Slate/Public/Widgets/Layout/SScrollBox.h
--- a/Engine/Source/Runtime/Slate/Public/Widgets/Layout/SScrollBox.h
+++ b/Engine/Source/Runtime/Slate/Public/Widgets/Layout/SScrollBox.h
@@ -327,6 +327,15 @@ public:
 	void SetScrollBarPadding(const FMargin& InPadding);
 
 	void SetScrollBarRightClickDragAllowed(bool bIsAllowed);
+
+	void SetStyle(const FScrollBoxStyle* InStyle);
+	
+	void SetScrollBarStyle(const FScrollBarStyle* InBarStyle);
+
+	void InvalidateStyle();
+
+	void InvalidateScrollBarStyle();
+
 public:
 
 	// SWidget interface
diff --git a/Engine/Source/Runtime/UMG/Private/Components/Button.cpp b/Engine/Source/Runtime/UMG/Private/Components/Button.cpp
--- a/Engine/Source/Runtime/UMG/Private/Components/Button.cpp
+++ b/Engine/Source/Runtime/UMG/Private/Components/Button.cpp
@@ -31,9 +31,9 @@ UButton::UButton(const FObjectInitializer& ObjectInitializer)
 		// Unlink UMG default colors.
 		DefaultButtonStyle->UnlinkColors();
 	}
-
+PRAGMA_DISABLE_DEPRECATION_WARNINGS
 	WidgetStyle = *DefaultButtonStyle;
-
+PRAGMA_ENABLE_DEPRECATION_WARNINGS
 #if WITH_EDITOR 
 	if (EditorButtonStyle == nullptr)
 	{
@@ -45,13 +45,16 @@ UButton::UButton(const FObjectInitializer& ObjectInitializer)
 
 	if (IsEditorWidget())
 	{
+PRAGMA_DISABLE_DEPRECATION_WARNINGS
 		WidgetStyle = *EditorButtonStyle;
+PRAGMA_ENABLE_DEPRECATION_WARNINGS
 
 		// The CDO isn't an editor widget and thus won't use the editor style, call post edit change to mark difference from CDO
 		PostEditChange();
 	}
 #endif // WITH_EDITOR
 
+PRAGMA_DISABLE_DEPRECATION_WARNINGS
 	ColorAndOpacity = FLinearColor::White;
 	BackgroundColor = FLinearColor::White;
 
@@ -59,6 +62,7 @@ UButton::UButton(const FObjectInitializer& ObjectInitializer)
 	TouchMethod = EButtonTouchMethod::DownAndUp;
 
 	IsFocusable = true;
+PRAGMA_ENABLE_DEPRECATION_WARNINGS
 
 #if WITH_EDITORONLY_DATA
 	AccessibleBehavior = ESlateAccessibleBehavior::Summary;
@@ -75,6 +79,7 @@ void UButton::ReleaseSlateResources(bool bReleaseChildren)
 
 TSharedRef<SWidget> UButton::RebuildWidget()
 {
+PRAGMA_DISABLE_DEPRECATION_WARNINGS
 	MyButton = SNew(SButton)
 		.OnClicked(BIND_UOBJECT_DELEGATE(FOnClicked, SlateHandleClicked))
 		.OnPressed(BIND_UOBJECT_DELEGATE(FSimpleDelegate, SlateHandlePressed))
@@ -87,7 +92,7 @@ TSharedRef<SWidget> UButton::RebuildWidget()
 		.PressMethod(PressMethod)
 		.IsFocusable(IsFocusable)
 		;
-
+PRAGMA_ENABLE_DEPRECATION_WARNINGS
 	if ( GetChildrenCount() > 0 )
 	{
 		Cast<UButtonSlot>(GetContentSlot())->BuildSlot(MyButton.ToSharedRef());
@@ -105,8 +110,14 @@ void UButton::SynchronizeProperties()
 		return;
 	}
 
+PRAGMA_DISABLE_DEPRECATION_WARNINGS
+	MyButton->SetButtonStyle(&WidgetStyle);
 	MyButton->SetColorAndOpacity( ColorAndOpacity );
 	MyButton->SetBorderBackgroundColor( BackgroundColor );
+	MyButton->SetClickMethod(ClickMethod);
+	MyButton->SetTouchMethod(TouchMethod);
+	MyButton->SetPressMethod(PressMethod);
+PRAGMA_ENABLE_DEPRECATION_WARNINGS
 }
 
 UClass* UButton::GetSlotClass() const
@@ -132,6 +143,7 @@ void UButton::OnSlotRemoved(UPanelSlot* InSlot)
 	}
 }
 
+PRAGMA_DISABLE_DEPRECATION_WARNINGS
 void UButton::SetStyle(const FButtonStyle& InStyle)
 {
 	WidgetStyle = InStyle;
@@ -141,6 +153,11 @@ void UButton::SetStyle(const FButtonStyle& InStyle)
 	}
 }
 
+const FButtonStyle& UButton::GetStyle() const
+{
+	return WidgetStyle;
+}
+
 void UButton::SetColorAndOpacity(FLinearColor InColorAndOpacity)
 {
 	ColorAndOpacity = InColorAndOpacity;
@@ -150,6 +167,11 @@ void UButton::SetColorAndOpacity(FLinearColor InColorAndOpacity)
 	}
 }
 
+FLinearColor UButton::GetColorAndOpacity() const
+{
+	return ColorAndOpacity;
+}
+
 void UButton::SetBackgroundColor(FLinearColor InBackgroundColor)
 {
 	BackgroundColor = InBackgroundColor;
@@ -159,6 +181,12 @@ void UButton::SetBackgroundColor(FLinearColor InBackgroundColor)
 	}
 }
 
+FLinearColor UButton::GetBackgroundColor() const
+{
+	return BackgroundColor;
+}
+PRAGMA_ENABLE_DEPRECATION_WARNINGS
+
 bool UButton::IsPressed() const
 {
 	if ( MyButton.IsValid() )
@@ -169,6 +197,7 @@ bool UButton::IsPressed() const
 	return false;
 }
 
+PRAGMA_DISABLE_DEPRECATION_WARNINGS
 void UButton::SetClickMethod(EButtonClickMethod::Type InClickMethod)
 {
 	ClickMethod = InClickMethod;
@@ -178,6 +207,11 @@ void UButton::SetClickMethod(EButtonClickMethod::Type InClickMethod)
 	}
 }
 
+EButtonClickMethod::Type UButton::GetClickMethod() const
+{
+	return ClickMethod;
+}
+
 void UButton::SetTouchMethod(EButtonTouchMethod::Type InTouchMethod)
 {
 	TouchMethod = InTouchMethod;
@@ -187,6 +221,11 @@ void UButton::SetTouchMethod(EButtonTouchMethod::Type InTouchMethod)
 	}
 }
 
+EButtonTouchMethod::Type UButton::GetTouchMethod() const
+{
+	return TouchMethod;
+}
+
 void UButton::SetPressMethod(EButtonPressMethod::Type InPressMethod)
 {
 	PressMethod = InPressMethod;
@@ -196,6 +235,23 @@ void UButton::SetPressMethod(EButtonPressMethod::Type InPressMethod)
 	}
 }
 
+EButtonPressMethod::Type UButton::GetPressMethod() const
+{
+	return PressMethod;
+}
+
+bool UButton::GetIsFocusable() const
+{
+	return IsFocusable;
+}
+
+void UButton::InitIsFocusable(bool InIsFocusable)
+{
+	IsFocusable = InIsFocusable;
+}
+
+PRAGMA_ENABLE_DEPRECATION_WARNINGS
+
 void UButton::PostLoad()
 {
 	Super::PostLoad();
diff --git a/Engine/Source/Runtime/UMG/Private/Components/ScrollBox.cpp b/Engine/Source/Runtime/UMG/Private/Components/ScrollBox.cpp
--- a/Engine/Source/Runtime/UMG/Private/Components/ScrollBox.cpp
+++ b/Engine/Source/Runtime/UMG/Private/Components/ScrollBox.cpp
@@ -57,9 +57,10 @@ UScrollBox::UScrollBox(const FObjectInitializer& ObjectInitializer)
 		// Unlink UMG default colors.
 		DefaultScrollBoxBarStyle->UnlinkColors();
 	}
-	
+PRAGMA_DISABLE_DEPRECATION_WARNINGS
 	WidgetStyle = *DefaultScrollBoxStyle;
 	WidgetBarStyle = *DefaultScrollBoxBarStyle;
+PRAGMA_ENABLE_DEPRECATION_WARNINGS
 
 #if WITH_EDITOR 
 	if (EditorScrollBoxStyle == nullptr)
@@ -80,15 +81,18 @@ UScrollBox::UScrollBox(const FObjectInitializer& ObjectInitializer)
 	
 	if (IsEditorWidget())
 	{
+PRAGMA_DISABLE_DEPRECATION_WARNINGS
 		WidgetStyle = *EditorScrollBoxStyle;
 		WidgetBarStyle = *EditorScrollBoxBarStyle;
+PRAGMA_ENABLE_DEPRECATION_WARNINGS
 
 		// The CDO isn't an editor widget and thus won't use the editor style, call post edit change to mark difference from CDO
 		PostEditChange();
 	}
 #endif // WITH_EDITOR
-
+PRAGMA_DISABLE_DEPRECATION_WARNINGS
 	bAllowRightClickDragScrolling = true;
+PRAGMA_ENABLE_DEPRECATION_WARNINGS
 }
 
 void UScrollBox::ReleaseSlateResources(bool bReleaseChildren)
@@ -127,6 +131,7 @@ void UScrollBox::OnSlotRemoved(UPanelSlot* InSlot)
 
 TSharedRef<SWidget> UScrollBox::RebuildWidget()
 {
+PRAGMA_DISABLE_DEPRECATION_WARNINGS
 	MyScrollBox = SNew(SScrollBox)
 		.Style(&WidgetStyle)
 		.ScrollBarStyle(&WidgetBarStyle)
@@ -140,7 +145,7 @@ TSharedRef<SWidget> UScrollBox::RebuildWidget()
 		.AnimateWheelScrolling(bAnimateWheelScrolling)
 		.WheelScrollMultiplier(WheelScrollMultiplier)
 		.OnUserScrolled(BIND_UOBJECT_DELEGATE(FOnUserScrolled, SlateHandleUserScrolled));
-
+PRAGMA_ENABLE_DEPRECATION_WARNINGS
 	for ( UPanelSlot* PanelSlot : Slots )
 	{
 		if ( UScrollBoxSlot* TypedSlot = Cast<UScrollBoxSlot>(PanelSlot) )
@@ -162,6 +167,7 @@ void UScrollBox::SynchronizeProperties()
 		return;
 	}
 
+PRAGMA_DISABLE_DEPRECATION_WARNINGS
 	MyScrollBox->SetScrollOffset(DesiredScrollOffset);
 	MyScrollBox->SetOrientation(Orientation);
 	MyScrollBox->SetScrollBarVisibility(UWidget::ConvertSerializedVisibilityToRuntime(ScrollBarVisibility));
@@ -174,6 +180,11 @@ void UScrollBox::SynchronizeProperties()
 	MyScrollBox->SetConsumeMouseWheel(ConsumeMouseWheel);
 	MyScrollBox->SetAnimateWheelScrolling(bAnimateWheelScrolling);
 	MyScrollBox->SetWheelScrollMultiplier(WheelScrollMultiplier);
+	MyScrollBox->SetStyle(&WidgetStyle);
+	MyScrollBox->InvalidateStyle();
+	MyScrollBox->SetScrollBarStyle(&WidgetBarStyle);
+	MyScrollBox->InvalidateScrollBarStyle();
+PRAGMA_ENABLE_DEPRECATION_WARNINGS
 }
 
 float UScrollBox::GetScrollOffset() const
@@ -267,21 +278,56 @@ void UScrollBox::Serialize(FArchive& Ar)
 	const bool bDeprecateThickness = Ar.IsLoading() && Ar.CustomVer(FEditorObjectVersion::GUID) < FEditorObjectVersion::ScrollBarThicknessChange;
 	if (bDeprecateThickness)
 	{
+PRAGMA_DISABLE_DEPRECATION_WARNINGS
 		// Set ScrollbarThickness property to previous default value.
 		ScrollbarThickness.Set(5.0f, 5.0f);
+PRAGMA_ENABLE_DEPRECATION_WARNINGS
 	}
 
 	Super::Serialize(Ar);
 
 	if (bDeprecateThickness)
 	{
+PRAGMA_DISABLE_DEPRECATION_WARNINGS
 		// Implicit padding of 2 was removed, so ScrollbarThickness value must be incremented by 4.
 		ScrollbarThickness += FVector2D(4.0f, 4.0f);
+PRAGMA_ENABLE_DEPRECATION_WARNINGS
 	}
 }
 
 #endif // if WITH_EDITORONLY_DATA
 
+PRAGMA_DISABLE_DEPRECATION_WARNINGS
+void UScrollBox::SetWidgetStyle(const FScrollBoxStyle& NewWidgetStyle)
+{
+	WidgetStyle = NewWidgetStyle;
+
+	if (MyScrollBox.IsValid())
+	{
+		MyScrollBox->InvalidateStyle();
+	}
+}
+
+const FScrollBoxStyle& UScrollBox::GetWidgetStyle() const
+{
+	return WidgetStyle;
+}
+
+void UScrollBox::SetWidgetBarStyle(const FScrollBarStyle& NewWidgetBarStyle)
+{
+	WidgetBarStyle = NewWidgetBarStyle;
+
+	if (MyScrollBox.IsValid())
+	{
+		MyScrollBox->InvalidateScrollBarStyle();
+	}
+}
+
+const FScrollBarStyle& UScrollBox::GetWidgetBarStyle() const
+{
+	return WidgetBarStyle;
+}
+
 void UScrollBox::SetNavigationDestination(const EDescendantScrollDestination NewNavigationDestination)
 {
 	NavigationDestination = NewNavigationDestination;
@@ -292,6 +338,11 @@ void UScrollBox::SetNavigationDestination(const EDescendantScrollDestination New
 	}
 }
 
+EDescendantScrollDestination UScrollBox::GetNavigationDestination() const
+{
+	return NavigationDestination;
+}
+
 void UScrollBox::SetConsumeMouseWheel(EConsumeMouseWheel NewConsumeMouseWheel)
 {
 	ConsumeMouseWheel = NewConsumeMouseWheel;
@@ -302,6 +353,11 @@ void UScrollBox::SetConsumeMouseWheel(EConsumeMouseWheel NewConsumeMouseWheel)
 	}
 }
 
+EConsumeMouseWheel UScrollBox::GetConsumeMouseWheel() const
+{
+	return ConsumeMouseWheel;
+}
+
 void UScrollBox::SetOrientation(EOrientation NewOrientation)
 {
 	Orientation = NewOrientation;
@@ -312,6 +368,11 @@ void UScrollBox::SetOrientation(EOrientation NewOrientation)
 	}
 }
 
+EOrientation UScrollBox::GetOrientation() const
+{
+	return Orientation;
+}
+
 void UScrollBox::SetScrollBarVisibility(ESlateVisibility NewScrollBarVisibility)
 {
 	ScrollBarVisibility = NewScrollBarVisibility;
@@ -329,6 +390,11 @@ void UScrollBox::SetScrollBarVisibility(ESlateVisibility NewScrollBarVisibility)
 	}
 }
 
+ESlateVisibility UScrollBox::GetScrollBarVisibility() const
+{
+	return ScrollBarVisibility;
+}
+
 void UScrollBox::SetScrollbarThickness(const FVector2D& NewScrollbarThickness)
 {
 	ScrollbarThickness = NewScrollbarThickness;
@@ -339,6 +405,11 @@ void UScrollBox::SetScrollbarThickness(const FVector2D& NewScrollbarThickness)
 	}
 }
 
+FVector2D UScrollBox::GetScrollbarThickness() const
+{
+	return ScrollbarThickness;
+}
+
 void UScrollBox::SetScrollbarPadding(const FMargin& NewScrollbarPadding)
 {
 	ScrollbarPadding = NewScrollbarPadding;
@@ -349,6 +420,11 @@ void UScrollBox::SetScrollbarPadding(const FMargin& NewScrollbarPadding)
 	}
 }
 
+FMargin UScrollBox::GetScrollbarPadding() const
+{
+	return ScrollbarPadding;
+}
+
 void UScrollBox::SetAlwaysShowScrollbar(bool NewAlwaysShowScrollbar)
 {
 	AlwaysShowScrollbar = NewAlwaysShowScrollbar;
@@ -359,6 +435,11 @@ void UScrollBox::SetAlwaysShowScrollbar(bool NewAlwaysShowScrollbar)
 	}
 }
 
+bool UScrollBox::IsAlwaysShowScrollbar() const
+{
+	return AlwaysShowScrollbar;
+}
+
 void UScrollBox::SetAllowOverscroll(bool NewAllowOverscroll)
 {
 	AllowOverscroll = NewAllowOverscroll;
@@ -369,6 +450,11 @@ void UScrollBox::SetAllowOverscroll(bool NewAllowOverscroll)
 	}
 }
 
+bool UScrollBox::IsAllowOverscroll() const
+{
+	return AllowOverscroll;
+}
+
 void UScrollBox::SetAnimateWheelScrolling(bool bShouldAnimateWheelScrolling)
 {
 	bAnimateWheelScrolling = bShouldAnimateWheelScrolling;
@@ -378,6 +464,11 @@ void UScrollBox::SetAnimateWheelScrolling(bool bShouldAnimateWheelScrolling)
 	}
 }
 
+bool UScrollBox::IsAnimateWheelScrolling() const
+{
+	return bAnimateWheelScrolling;
+}
+
 void UScrollBox::SetWheelScrollMultiplier(float NewWheelScrollMultiplier)
 {
 	WheelScrollMultiplier = NewWheelScrollMultiplier;
@@ -387,6 +478,11 @@ void UScrollBox::SetWheelScrollMultiplier(float NewWheelScrollMultiplier)
 	}
 }
 
+float UScrollBox::GetWheelScrollMultiplier() const
+{
+	return WheelScrollMultiplier;
+}
+
 void UScrollBox::SetScrollWhenFocusChanges(EScrollWhenFocusChanges NewScrollWhenFocusChanges)
 {
 	ScrollWhenFocusChanges = NewScrollWhenFocusChanges;
@@ -396,6 +492,11 @@ void UScrollBox::SetScrollWhenFocusChanges(EScrollWhenFocusChanges NewScrollWhen
 	}
 }
 
+EScrollWhenFocusChanges UScrollBox::GetScrollWhenFocusChanges() const
+{
+	return ScrollWhenFocusChanges;
+}
+
 void UScrollBox::EndInertialScrolling()
 {
 	if (MyScrollBox.IsValid())
@@ -404,6 +505,68 @@ void UScrollBox::EndInertialScrolling()
 	}
 }
 
+void UScrollBox::SetAlwaysShowScrollbarTrack(bool NewAlwaysShowScrollbarTrack)
+{
+	AlwaysShowScrollbarTrack = NewAlwaysShowScrollbarTrack;
+	if (MyScrollBox)
+	{
+		MyScrollBox->SetScrollBarTrackAlwaysVisible(AlwaysShowScrollbarTrack);
+	}
+}
+
+bool UScrollBox::IsAlwaysShowScrollbarTrack() const
+{
+	return AlwaysShowScrollbarTrack;
+}
+
+float UScrollBox::GetNavigationScrollPadding() const
+{
+	return NavigationScrollPadding;
+}
+
+void UScrollBox::SetAllowRightClickDragScrolling(bool bShouldAllowRightClickDragScrolling)
+{
+	AlwaysShowScrollbarTrack = bShouldAllowRightClickDragScrolling;
+	if (MyScrollBox)
+	{
+		MyScrollBox->SetScrollBarRightClickDragAllowed(AlwaysShowScrollbarTrack);
+	}
+}
+
+bool UScrollBox::IsAllowRightClickDragScrolling() const
+{
+	return AlwaysShowScrollbarTrack;
+}
+
+bool UScrollBox::IsFrontPadScrolling() const
+{
+	return FrontPadScrolling;
+}
+
+bool UScrollBox::IsBackPadScrolling() const
+{
+	return BackPadScrolling;
+}
+
+void UScrollBox::InitBackPadScrolling(bool InBackPadScrolling)
+{
+	ensureMsgf(!MyScrollBox.IsValid(), TEXT("The widget is already created."));
+	BackPadScrolling = InBackPadScrolling;
+}
+
+void UScrollBox::InitFrontPadScrolling(bool InFrontPadScrolling)
+{
+	ensureMsgf(!MyScrollBox.IsValid(), TEXT("The widget is already created."));
+	FrontPadScrolling = InFrontPadScrolling;
+}
+
+void UScrollBox::InitNavigationScrollPadding(float InNavigationScrollPadding)
+{
+	ensureMsgf(!MyScrollBox.IsValid(), TEXT("The widget is already created."));
+	NavigationScrollPadding = InNavigationScrollPadding;
+}
+PRAGMA_ENABLE_DEPRECATION_WARNINGS
+
 void UScrollBox::SlateHandleUserScrolled(float CurrentOffset)
 {
 	OnUserScrolled.Broadcast(CurrentOffset);
diff --git a/Engine/Source/Runtime/UMG/Private/UserWidget.cpp b/Engine/Source/Runtime/UMG/Private/UserWidget.cpp
--- a/Engine/Source/Runtime/UMG/Private/UserWidget.cpp
+++ b/Engine/Source/Runtime/UMG/Private/UserWidget.cpp
@@ -83,9 +83,11 @@ UUserWidget::UUserWidget(const FObjectInitializer& ObjectInitializer)
 {
 	SetVisibilityInternal(ESlateVisibility::SelfHitTestInvisible);
 
+PRAGMA_DISABLE_DEPRECATION_WARNINGS
 	bIsFocusable = false;
 	ColorAndOpacity = FLinearColor::White;
 	ForegroundColor = FSlateColor::UseForeground();
+PRAGMA_ENABLE_DEPRECATION_WARNINGS
 
 	MinimumDesiredSize = FVector2D(0, 0);
 
@@ -312,15 +314,18 @@ void UUserWidget::SynchronizeProperties()
 	TSharedPtr<SObjectWidget> SafeGCWidget = MyGCWidget.Pin();
 	if ( SafeGCWidget.IsValid() )
 	{
+PRAGMA_DISABLE_DEPRECATION_WARNINGS
 		TAttribute<FLinearColor> ColorBinding = PROPERTY_BINDING(FLinearColor, ColorAndOpacity);
 		TAttribute<FSlateColor> ForegroundColorBinding = PROPERTY_BINDING(FSlateColor, ForegroundColor);
 
 		SafeGCWidget->SetColorAndOpacity(ColorBinding);
 		SafeGCWidget->SetForegroundColor(ForegroundColorBinding);
 		SafeGCWidget->SetPadding(Padding);
+PRAGMA_ENABLE_DEPRECATION_WARNINGS
 	}
 }
 
+PRAGMA_DISABLE_DEPRECATION_WARNINGS
 void UUserWidget::SetColorAndOpacity(FLinearColor InColorAndOpacity)
 {
 	ColorAndOpacity = InColorAndOpacity;
@@ -332,6 +337,11 @@ void UUserWidget::SetColorAndOpacity(FLinearColor InColorAndOpacity)
 	}
 }
 
+const FLinearColor& UUserWidget::GetColorAndOpacity() const
+{
+	return ColorAndOpacity;
+}
+
 void UUserWidget::SetForegroundColor(FSlateColor InForegroundColor)
 {
 	ForegroundColor = InForegroundColor;
@@ -343,6 +353,11 @@ void UUserWidget::SetForegroundColor(FSlateColor InForegroundColor)
 	}
 }
 
+const FSlateColor& UUserWidget::GetForegroundColor() const
+{
+	return ForegroundColor;
+}
+
 void UUserWidget::SetPadding(FMargin InPadding)
 {
 	Padding = InPadding;
@@ -354,6 +369,12 @@ void UUserWidget::SetPadding(FMargin InPadding)
 	}
 }
 
+FMargin UUserWidget::GetPadding() const
+{
+	return Padding;
+}
+PRAGMA_ENABLE_DEPRECATION_WARNINGS
+
 UWorld* UUserWidget::GetWorld() const
 {
 	if ( UWorld* LastWorld = CachedWorld.Get() )
@@ -1629,6 +1650,7 @@ void UUserWidget::UnregisterInputComponent()
 	}
 }
 
+PRAGMA_DISABLE_DEPRECATION_WARNINGS
 void UUserWidget::SetInputActionPriority( int32 NewPriority )
 {
 	if ( InputComponent )
@@ -1638,6 +1660,11 @@ void UUserWidget::SetInputActionPriority( int32 NewPriority )
 	}
 }
 
+int32 UUserWidget::GetInputActionPriority() const
+{
+	return Priority;
+}
+
 void UUserWidget::SetInputActionBlocking( bool bShouldBlock )
 {
 	if ( InputComponent )
@@ -1647,6 +1674,13 @@ void UUserWidget::SetInputActionBlocking( bool bShouldBlock )
 	}
 }
 
+bool UUserWidget::IsInputActionBlocking() const
+{
+	return bStopAction;
+}
+
+PRAGMA_ENABLE_DEPRECATION_WARNINGS
+
 void UUserWidget::OnInputAction( FOnInputAction Callback )
 {
 	if ( GetIsEnabled() )
@@ -1664,8 +1698,10 @@ void UUserWidget::InitializeInputComponent()
 		// class then this would fail
 		UClass* InputClass = Controller->InputComponent ? Controller->InputComponent->GetClass() : UInputSettings::GetDefaultInputComponentClass();
 		InputComponent = NewObject< UInputComponent >( this, InputClass, NAME_None, RF_Transient );
+PRAGMA_DISABLE_DEPRECATION_WARNINGS
 		InputComponent->bBlockInput = bStopAction;
 		InputComponent->Priority = Priority;
+PRAGMA_ENABLE_DEPRECATION_WARNINGS
 		Controller->PushInputComponent( InputComponent );
 	}
 	else
@@ -1736,11 +1772,25 @@ bool UUserWidget::NativeIsInteractable() const
 	return IsInteractable();
 }
 
+PRAGMA_DISABLE_DEPRECATION_WARNINGS
 bool UUserWidget::NativeSupportsKeyboardFocus() const
 {
 	return bIsFocusable;
 }
 
+bool UUserWidget::IsFocusable() const
+{
+	return bIsFocusable;
+}
+
+void UUserWidget::InitIsFocusable(bool InIsFocusable)
+{
+	ensureMsgf(!MyWidget.IsValid(), TEXT("The widget is already created."));
+	bIsFocusable = InIsFocusable;
+}
+
+PRAGMA_ENABLE_DEPRECATION_WARNINGS
+
 FReply UUserWidget::NativeOnFocusReceived( const FGeometry& InGeometry, const FFocusEvent& InFocusEvent )
 {
 	return OnFocusReceived( InGeometry, InFocusEvent ).NativeReply;
diff --git a/Engine/Source/Runtime/UMG/Public/Blueprint/UserWidget.h b/Engine/Source/Runtime/UMG/Public/Blueprint/UserWidget.h
--- a/Engine/Source/Runtime/UMG/Public/Blueprint/UserWidget.h
+++ b/Engine/Source/Runtime/UMG/Public/Blueprint/UserWidget.h
@@ -894,5 +894,49 @@
 	virtual void OnAnimationFinishedPlaying(UUMGSequencePlayer& Player);
 
 public:
+	UE_DEPRECATED(5.2, "Direct access to ColorAndOpacity is deprecated. Please use the getter or setter.")
+	/** The color and opacity of this widget.  Tints all child widgets. */
+	UPROPERTY(EditAnywhere, BlueprintReadWrite, Getter, Setter, BlueprintSetter = "SetColorAndOpacity", Category = "Appearance")
+	FLinearColor ColorAndOpacity;
+
+	UPROPERTY()
+	FGetLinearColor ColorAndOpacityDelegate;
+
+	UE_DEPRECATED(5.2, "Direct access to ForegroundColor is deprecated. Please use the getter or setter.")
+	/**
+	 * The foreground color of the widget, this is inherited by sub widgets.  Any color property
+	 * that is marked as inherit will use this color.
+	 */
+	UPROPERTY(EditAnywhere, BlueprintReadWrite, Getter, Setter, BlueprintSetter = "SetForegroundColor", Category = "Appearance")
+	FSlateColor ForegroundColor;
+
+	UPROPERTY()
+	FGetSlateColor ForegroundColorDelegate;
+
+	/** Called when the visibility has changed */
+	UPROPERTY(BlueprintAssignable, Category = "Appearance|Event")
+	FOnVisibilityChangedEvent OnVisibilityChanged;
+	DECLARE_EVENT_OneParam(UUserWidget, FNativeOnVisibilityChangedEvent, ESlateVisibility);
+	FNativeOnVisibilityChangedEvent OnNativeVisibilityChanged;
+
+	UE_DEPRECATED(5.2, "Direct access to Padding is deprecated. Please use the getter or setter.")
+	/** The padding area around the content. */
+	UPROPERTY(EditAnywhere, BlueprintReadWrite, Getter, Setter, BlueprintSetter = "SetPadding", Category = "Appearance")
+	FMargin Padding;
+
+	UE_DEPRECATED(5.2, "Direct access to Priority is deprecated. Please use the getter or setter.")
+	UPROPERTY(EditAnywhere, BlueprintReadWrite, Getter = "GetInputActionPriority", Setter = "SetInputActionPriority", BlueprintSetter = "SetInputActionPriority", Category = "Input")
+	int32 Priority;
+
+	UE_DEPRECATED(5.2, "Direct access to bIsFocusable is deprecated. Please use the getter. Note that this property is only set at construction and is not modifiable at runtime.")
+	/** Setting this flag to true, allows this widget to accept focus when clicked, or when navigated to. */
+	UPROPERTY(EditAnywhere, BlueprintReadOnly, Getter = "IsFocusable", Category = "Interaction")
+	uint8 bIsFocusable : 1;
+	 
+	UE_DEPRECATED(5.2, "Direct access to bStopAction is deprecated. Please use the getter or setter.")
+	UPROPERTY(EditAnywhere, BlueprintReadWrite, Getter = "IsInputActionBlocking", Setter = "SetInputActionBlocking", BlueprintSetter = "SetInputActionBlocking", Category = "Input")
+	uint8 bStopAction : 1;
+
+public:
 
 	/**
@@ -904,20 +948,51 @@
 	void SetColorAndOpacity(FLinearColor InColorAndOpacity);
 
 	/**
+	 * Gets the tint of the widget.
+	 */
+	const FLinearColor& GetColorAndOpacity() const;
+
+	/**
 	 * Sets the foreground color of the widget, this is inherited by sub widgets.  Any color property 
 	 * that is marked as inherit will use this color.
-	 * 
+	 *
 	 * @param InForegroundColor	The foreground color.
 	 */
 	UFUNCTION(BlueprintCallable, BlueprintCosmetic, Category="Appearance")
 	void SetForegroundColor(FSlateColor InForegroundColor);
 
 	/**
+	 * Gets the foreground color of the widget, this is inherited by sub widgets.  Any color property
+	 * that is marked as inherit uses this color.
+	 */
+	const FSlateColor& GetForegroundColor() const;
+
+	/**
 	 * Sets the padding for the user widget, putting a larger gap between the widget border and it's root widget.
 	 */
 	UFUNCTION(BlueprintCallable, BlueprintCosmetic, Category="Appearance")
 	void SetPadding(FMargin InPadding);
 
 	/**
+	 * Gets the padding for the user widget.
+	 */
+	FMargin GetPadding() const;
+
+	/**
+	 * Gets the priority of the input action.
+	 */
+	int32 GetInputActionPriority() const;
+
+	/**
+	 * Returns whether the input action is blocking.
+	 */
+	bool IsInputActionBlocking() const;
+
+	/**
+	 * Returns whether this widget to accept focus when clicked, or when navigated to.
+	 */
+	bool IsFocusable() const;
+
+	/**
 	 * Plays an animation in this widget a specified number of times
 	 * 
@@ -1184,32 +1259,6 @@
 	void ClearStoppedSequencePlayers();
 
 public:
-	/** The color and opacity of this widget.  Tints all child widgets. */
-	UPROPERTY(EditAnywhere, BlueprintReadOnly, Category="Appearance")
-	FLinearColor ColorAndOpacity;
 
-	UPROPERTY()
-	FGetLinearColor ColorAndOpacityDelegate;
-
-	/**
-	 * The foreground color of the widget, this is inherited by sub widgets.  Any color property
-	 * that is marked as inherit will use this color.
-	 */
-	UPROPERTY(EditAnywhere, BlueprintReadOnly, Category="Appearance")
-	FSlateColor ForegroundColor;
-
-	UPROPERTY()
-	FGetSlateColor ForegroundColorDelegate;
-
-	/** Called when the visibility has changed */
-	UPROPERTY(BlueprintAssignable, Category = "Appearance|Event")
-	FOnVisibilityChangedEvent OnVisibilityChanged;
-	DECLARE_EVENT_OneParam(UUserWidget, FNativeOnVisibilityChangedEvent, ESlateVisibility);
-	FNativeOnVisibilityChangedEvent OnNativeVisibilityChanged;
-
-	/** The padding area around the content. */
-	UPROPERTY(EditAnywhere, BlueprintReadOnly, Category="Appearance")
-	FMargin Padding;
-
 	/** All the sequence players currently playing */
 	UPROPERTY(Transient)
@@ -1260,16 +1309,6 @@
 
 #endif
 
-	UPROPERTY(EditAnywhere, BlueprintReadOnly, Category = "Input")
-	int32 Priority;
-
-	/** Setting this flag to true, allows this widget to accept focus when clicked, or when navigated to. */
-	UPROPERTY(EditAnywhere, BlueprintReadWrite, Category="Interaction")
-	uint8 bIsFocusable : 1;
-
-	UPROPERTY(EditAnywhere, BlueprintReadWrite, Category = "Input")
-	uint8 bStopAction : 1;
-
 	/** If a widget has an implemented tick blueprint function */
 	UPROPERTY()
 	uint8 bHasScriptImplementedTick : 1;
@@ -1420,6 +1459,9 @@
 
 	virtual void InitializeInputComponent();
 
+	/** Initialize IsFocusable in the constructor before the SWidget is constructed. */
+	void InitIsFocusable(bool InIsFocusable);
+
 private:
 	FVector2D MinimumDesiredSize;
 
diff --git a/Engine/Source/Runtime/UMG/Public/Components/Button.h b/Engine/Source/Runtime/UMG/Public/Components/Button.h
--- a/Engine/Source/Runtime/UMG/Public/Components/Button.h
+++ b/Engine/Source/Runtime/UMG/Public/Components/Button.h
@@ -33,32 +33,39 @@ class UMG_API UButton : public UContentWidget
 	GENERATED_UCLASS_BODY()
 
 public:
+	UE_DEPRECATED(5.2, "Direct access to WidgetStyle is deprecated. Please use the getter and setter.")
 	/** The button style used at runtime */
-	UPROPERTY(EditAnywhere, BlueprintReadWrite, Category="Appearance", meta=( DisplayName="Style" ))
+	UPROPERTY(EditAnywhere, BlueprintReadWrite, Getter = "GetStyle", Setter = "SetStyle", BlueprintSetter = "SetStyle", Category = "Appearance", meta = (DisplayName = "Style"))
 	FButtonStyle WidgetStyle;
 	
+	UE_DEPRECATED(5.2, "Direct access to ColorAndOpacity is deprecated. Please use the getter and setter.")
 	/** The color multiplier for the button content */
-	UPROPERTY(EditAnywhere, BlueprintReadOnly, Category="Appearance", meta=( sRGB="true" ))
+	UPROPERTY(EditAnywhere, BlueprintReadWrite, Getter, Setter, BlueprintSetter = "SetColorAndOpacity" , Category="Appearance", meta=( sRGB="true" ))
 	FLinearColor ColorAndOpacity;
 	
+	UE_DEPRECATED(5.2, "Direct access to BackgroundColor is deprecated. Please use the getter and setter.")
 	/** The color multiplier for the button background */
-	UPROPERTY(EditAnywhere, BlueprintReadOnly, Category="Appearance", meta=( sRGB="true" ))
+	UPROPERTY(EditAnywhere, BlueprintReadWrite, Getter, Setter, BlueprintSetter = "SetBackgroundColor", Category="Appearance", meta=( sRGB="true" ))
 	FLinearColor BackgroundColor;
 
+	UE_DEPRECATED(5.2, "Direct access to ClickMethod is deprecated. Please use the getter and setter.")
 	/** The type of mouse action required by the user to trigger the buttons 'Click' */
-	UPROPERTY(EditAnywhere, BlueprintReadOnly, Category="Interaction", AdvancedDisplay)
+	UPROPERTY(EditAnywhere, BlueprintReadWrite, Getter, Setter, BlueprintSetter = "SetClickMethod", Category="Interaction", AdvancedDisplay)
 	TEnumAsByte<EButtonClickMethod::Type> ClickMethod;
 
+	UE_DEPRECATED(5.2, "Direct access to TouchMethod is deprecated. Please use the getter and setter.")
 	/** The type of touch action required by the user to trigger the buttons 'Click' */
-	UPROPERTY(EditAnywhere, BlueprintReadOnly, Category="Interaction", AdvancedDisplay)
+	UPROPERTY(EditAnywhere, BlueprintReadWrite, Getter, Setter, BlueprintSetter = "SetTouchMethod", Category="Interaction", AdvancedDisplay)
 	TEnumAsByte<EButtonTouchMethod::Type> TouchMethod;
 
+	UE_DEPRECATED(5.2, "Direct access to PressMethod is deprecated. Please use the getter and setter.")
 	/** The type of keyboard/gamepad button press action required by the user to trigger the buttons 'Click' */
-	UPROPERTY(EditAnywhere, BlueprintReadOnly, Category="Interaction", AdvancedDisplay)
+	UPROPERTY(EditAnywhere, BlueprintReadWrite, Getter, Setter, BlueprintSetter = "SetPressMethod", Category="Interaction", AdvancedDisplay)
 	TEnumAsByte<EButtonPressMethod::Type> PressMethod;
 
+	UE_DEPRECATED(5.2, "Direct access to IsFocusable is deprecated. Please use the getter.")
 	/** Sometimes a button should only be mouse-clickable and never keyboard focusable. */
-	UPROPERTY(EditAnywhere, BlueprintReadOnly, Category="Interaction")
+	UPROPERTY(EditAnywhere, BlueprintReadOnly, Getter, Category="Interaction")
 	bool IsFocusable;
 
 public:
@@ -87,14 +94,20 @@ public:
 	UFUNCTION(BlueprintCallable, Category="Button|Appearance")
 	void SetStyle(const FButtonStyle& InStyle);
 
+	const FButtonStyle& GetStyle() const;
+
 	/** Sets the color multiplier for the button content */
 	UFUNCTION(BlueprintCallable, Category="Button|Appearance")
 	void SetColorAndOpacity(FLinearColor InColorAndOpacity);
 
+	FLinearColor GetColorAndOpacity() const;
+
 	/** Sets the color multiplier for the button background */
 	UFUNCTION(BlueprintCallable, Category="Button|Appearance")
 	void SetBackgroundColor(FLinearColor InBackgroundColor);
 
+	FLinearColor GetBackgroundColor() const;
+
 	/**
 	 * Returns true if the user is actively pressing the button.  Do not use this for detecting 'Clicks', use the OnClicked event instead.
 	 *
@@ -106,12 +119,20 @@ public:
 	UFUNCTION(BlueprintCallable, Category="Button")
 	void SetClickMethod(EButtonClickMethod::Type InClickMethod);
 
+	EButtonClickMethod::Type GetClickMethod() const;
+
 	UFUNCTION(BlueprintCallable, Category="Button")
 	void SetTouchMethod(EButtonTouchMethod::Type InTouchMethod);
 
+	EButtonTouchMethod::Type GetTouchMethod() const;
+
 	UFUNCTION(BlueprintCallable, Category="Button")
 	void SetPressMethod(EButtonPressMethod::Type InPressMethod);
 
+	EButtonPressMethod::Type GetPressMethod() const;
+
+	bool GetIsFocusable() const;
+
 public:
 
 	//~ Begin UWidget Interface
@@ -146,6 +167,8 @@ protected:
 	void SlateHandleHovered();
 	void SlateHandleUnhovered();
 
+	// Initialize IsFocusable in the constructor before the SWidget is constructed.
+	void InitIsFocusable(bool InIsFocusable);
 protected:
 	//~ Begin UWidget Interface
 	virtual TSharedRef<SWidget> RebuildWidget() override;
diff --git a/Engine/Source/Runtime/UMG/Public/Components/ListViewBase.h b/Engine/Source/Runtime/UMG/Public/Components/ListViewBase.h
--- a/Engine/Source/Runtime/UMG/Public/Components/ListViewBase.h
+++ b/Engine/Source/Runtime/UMG/Public/Components/ListViewBase.h
@@ -387,9 +387,9 @@ private:
 		UUserWidget& EntryWidget = OnGenerateEntryWidgetInternal(Item, DesiredEntryClass, OwnerTable);
 		
 		// Combine the desired entry padding with the padding the widget wants natively on the CDO.
-		const FMargin DefaultPadding = EntryWidget.GetClass()->GetDefaultObject<UUserWidget>()->Padding;
+		const FMargin DefaultPadding = EntryWidget.GetClass()->GetDefaultObject<UUserWidget>()->GetPadding();
 		EntryWidget.SetPadding(DefaultPadding + GetDesiredEntryPadding(Item));
-		
+
 		TSharedPtr<SWidget> CachedWidget = EntryWidget.GetCachedWidget();
 		CachedWidget->SetCanTick(true); // this is a hack to force ticking to true so selection works (which should NOT require ticking! but currently does)
 		return StaticCastSharedPtr<SObjectTableRow<ItemType>>(CachedWidget).ToSharedRef();
diff --git a/Engine/Source/Runtime/UMG/Public/Components/ScrollBox.h b/Engine/Source/Runtime/UMG/Public/Components/ScrollBox.h
--- a/Engine/Source/Runtime/UMG/Public/Components/ScrollBox.h
+++ b/Engine/Source/Runtime/UMG/Public/Components/ScrollBox.h
@@ -23,115 +23,176 @@ class UMG_API UScrollBox : public UPanelWidget
 
 public:
 
+	UE_DEPRECATED(5.2, "Direct access to WidgetStyle is deprecated. Please use the getter or setter.")
 	/** The style */
-	UPROPERTY(EditAnywhere, BlueprintReadWrite, Category="Style", meta=( DisplayName="Style" ))
+	UPROPERTY(EditAnywhere, BlueprintReadWrite, Getter, Setter, Category="Style", meta=( DisplayName="Style" ))
 	FScrollBoxStyle WidgetStyle;
 
+	UE_DEPRECATED(5.2, "Direct access to WidgetBarStyle is deprecated. Please use the getter or setter.")
 	/** The bar style */
-	UPROPERTY(EditAnywhere, BlueprintReadWrite, Category="Style", meta=( DisplayName="Bar Style" ))
+	UPROPERTY(EditAnywhere, BlueprintReadWrite, Getter, Setter, Category="Style", meta=( DisplayName="Bar Style" ))
 	FScrollBarStyle WidgetBarStyle;
 
+	UE_DEPRECATED(5.2, "Direct access to Orientation is deprecated. Please use the getter or setter.")
 	/** The orientation of the scrolling and stacking in the box. */
-	UPROPERTY(EditAnywhere, BlueprintReadOnly, Category="Scroll")
+	UPROPERTY(EditAnywhere, BlueprintReadWrite, Getter, Setter, BlueprintSetter = "SetOrientation", Category="Scroll")
 	TEnumAsByte<EOrientation> Orientation;
 
+	UE_DEPRECATED(5.2, "Direct access to ScrollBarVisibility is deprecated. Please use the getter or setter.")
 	/** Visibility */
-	UPROPERTY(EditAnywhere, BlueprintReadOnly, Category="Scroll")
+	UPROPERTY(EditAnywhere, BlueprintReadWrite, Getter, Setter, BlueprintSetter = "SetScrollBarVisibility", Category="Scroll")
 	ESlateVisibility ScrollBarVisibility;
 
+	UE_DEPRECATED(5.2, "Direct access to ConsumeMouseWheel is deprecated. Please use the getter or setter.")
 	/** When mouse wheel events should be consumed. */
-	UPROPERTY(EditAnywhere, BlueprintReadOnly, Category = "Scroll")
+	UPROPERTY(EditAnywhere, BlueprintReadWrite, Getter, Setter, BlueprintSetter = "SetConsumeMouseWheel", Category = "Scroll")
 	EConsumeMouseWheel ConsumeMouseWheel;
 
+	UE_DEPRECATED(5.2, "Direct access to ScrollbarThickness is deprecated. Please use the getter or setter.")
 	/** The thickness of the scrollbar thumb */
-	UPROPERTY(EditAnywhere, BlueprintReadOnly, Category="Scroll")
+	UPROPERTY(EditAnywhere, BlueprintReadWrite, Getter, Setter, BlueprintSetter = "SetScrollbarThickness", Category="Scroll")
 	FVector2D ScrollbarThickness;
 
+	UE_DEPRECATED(5.2, "Direct access to ScrollbarPadding is deprecated. Please use the getter or setter.")
 	/** The margin around the scrollbar */
-	UPROPERTY(EditAnywhere, BlueprintReadOnly, Category = "Scroll")
+	UPROPERTY(EditAnywhere, BlueprintReadWrite, Getter, Setter, BlueprintSetter = "SetScrollbarPadding", Category = "Scroll")
 	FMargin ScrollbarPadding;
 
 	/**  */
-	UPROPERTY(EditAnywhere, BlueprintReadOnly, Category="Scroll")
+	UE_DEPRECATED(5.2, "Direct access to AlwaysShowScrollbar is deprecated. Please use the getter or setter.")
+	UPROPERTY(EditAnywhere, BlueprintReadWrite, Getter = "IsAlwaysShowScrollbar", Setter, BlueprintSetter = "SetAlwaysShowScrollbar", Category = "Scroll")
 	bool AlwaysShowScrollbar;
 
 	/**  */
-	UPROPERTY(EditAnywhere, BlueprintReadOnly, Category = "Scroll")
+	UE_DEPRECATED(5.2, "Direct access to AlwaysShowScrollbarTrack is deprecated. Please use the getter or setter.")
+	UPROPERTY(EditAnywhere, BlueprintReadWrite, Getter = "IsAlwaysShowScrollbarTrack", Setter, Category = "Scroll")
 	bool AlwaysShowScrollbarTrack;
 
+	UE_DEPRECATED(5.2, "Direct access to AllowOverscroll is deprecated. Please use the getter or setter.")
 	/**  Disable to stop scrollbars from activating inertial overscrolling */
-	UPROPERTY(EditAnywhere, BlueprintReadOnly, Category = "Scroll")
+	UPROPERTY(EditAnywhere, BlueprintReadWrite, Getter = "IsAllowOverscroll", Setter, BlueprintSetter = "SetAllowOverscroll", Category = "Scroll")
 	bool AllowOverscroll;
 
+	UE_DEPRECATED(5.2, "Direct access to BackPadScrolling is deprecated. Please use the getter. Note that BackPadScrolling is only set at construction and is not modifiable at runtime.")
 	/** Whether to back pad this scroll box, allowing user to scroll backward until child contents are no longer visible */
-	UPROPERTY(EditAnywhere, BlueprintReadOnly, Category = "Scroll")
+	UPROPERTY(EditAnywhere, BlueprintReadOnly, Getter = "IsBackPadScrolling", Category = "Scroll")
 	bool BackPadScrolling;
 
+	UE_DEPRECATED(5.2, "Direct access to FrontPadScrolling is deprecated. Please use the getter. Note that FrontPadScrolling is only set at construction and is not modifiable at runtime.")
 	/** Whether to front pad this scroll box, allowing user to scroll forward until child contents are no longer visible */
-	UPROPERTY(EditAnywhere, BlueprintReadOnly, Category = "Scroll")
+	UPROPERTY(EditAnywhere, BlueprintReadOnly, Getter = "IsFrontPadScrolling", Category = "Scroll")
 	bool FrontPadScrolling;
 	
+	UE_DEPRECATED(5.2, "Direct access to bAnimateWheelScrolling is deprecated. Please use the getter or setter.")
 	/** True to lerp smoothly when wheel scrolling along the scroll box */
-	UPROPERTY(EditAnywhere, BlueprintReadOnly, Category = "Scroll")
+	UPROPERTY(EditAnywhere, BlueprintReadWrite, Getter = "IsAnimateWheelScrolling", Setter = "SetAnimateWheelScrolling", BlueprintSetter = "SetAnimateWheelScrolling", Category = "Scroll")
 	bool bAnimateWheelScrolling = false;
 
+	UE_DEPRECATED(5.2, "Direct access to NavigationDestination is deprecated. Please use the getter or setter.")
 	/** Sets where to scroll a widget to when using explicit navigation or if ScrollWhenFocusChanges is enabled */
-	UPROPERTY(EditAnywhere, BlueprintReadOnly, Category="Scroll")
+	UPROPERTY(EditAnywhere, BlueprintReadWrite, Getter, Setter, BlueprintSetter = "SetNavigationDestination", Getter, Setter, Category = "Scroll")
 	EDescendantScrollDestination NavigationDestination;
 
+	UE_DEPRECATED(5.2, "Direct access to NavigationScrollPadding is deprecated. Please use the getter. Note that NavigationScrollPadding is only set at construction and is not modifiable at runtime.")
 	/**
 	 * The amount of padding to ensure exists between the item being navigated to, at the edge of the
 	 * scrollbox.  Use this if you want to ensure there's a preview of the next item the user could scroll to.
 	 */
-	UPROPERTY(EditAnywhere, BlueprintReadOnly, Category="Scroll")
+	UPROPERTY(EditAnywhere, BlueprintReadOnly, Getter, Category="Scroll")
 	float NavigationScrollPadding;
 
+	UE_DEPRECATED(5.2, "Direct access to ScrollWhenFocusChanges is deprecated. Please use the getter or setter.")
 	/** Scroll behavior when user focus is given to a child widget */
-	UPROPERTY(EditAnywhere, BlueprintReadOnly, Category="Scroll", meta=(DisplayName="Scroll When Focus Changes"))
+	UPROPERTY(EditAnywhere, BlueprintReadWrite, Getter, Setter, BlueprintSetter = "SetScrollWhenFocusChanges", Category="Scroll", meta=(DisplayName="Scroll When Focus Changes"))
 	EScrollWhenFocusChanges ScrollWhenFocusChanges;
 	
+	UE_DEPRECATED(5.2, "Direct access to bAllowRightClickDragScrolling is deprecated. Please use the getter or setter.")
 	/** Option to disable right-click-drag scrolling */
-	UPROPERTY(EditAnywhere, BlueprintReadOnly, Category = "Scroll")
+	UPROPERTY(EditAnywhere, BlueprintReadWrite, Getter = "IsAllowRightClickDragScrolling", Setter = "SetAllowRightClickDragScrolling", Category = "Scroll")
 	bool bAllowRightClickDragScrolling;
 
+	UE_DEPRECATED(5.2, "Direct access to WheelScrollMultiplier is deprecated. Please use the getter or setter.")
 	/** The multiplier to apply when wheel scrolling */
-	UPROPERTY(EditAnywhere, BlueprintReadOnly, Category = "Scroll")
+	UPROPERTY(EditAnywhere, BlueprintReadWrite, Getter, Setter, BlueprintSetter = "SetWheelScrollMultiplier", Category = "Scroll")
 	float WheelScrollMultiplier = 1.f;
 
+	void SetWidgetStyle(const FScrollBoxStyle& NewWidgetStyle);
+
+	const FScrollBoxStyle& GetWidgetStyle() const;
+
+	void SetWidgetBarStyle(const FScrollBarStyle& NewWidgetBarStyle);
+
+	const FScrollBarStyle& GetWidgetBarStyle() const;
 
 	UFUNCTION(BlueprintCallable, Category = "Scroll")
 	void SetConsumeMouseWheel(EConsumeMouseWheel NewConsumeMouseWheel);
 
+	EConsumeMouseWheel GetConsumeMouseWheel() const;
+
 	UFUNCTION(BlueprintCallable, Category = "Scroll")
 	void SetOrientation(EOrientation NewOrientation);
 
+	EOrientation GetOrientation() const;
+
 	UFUNCTION(BlueprintCallable, Category = "Scroll")
 	void SetScrollBarVisibility(ESlateVisibility NewScrollBarVisibility);
 
+	ESlateVisibility GetScrollBarVisibility() const;
+
 	UFUNCTION(BlueprintCallable, Category = "Scroll")
 	void SetScrollbarThickness(const FVector2D& NewScrollbarThickness);
 
+	FVector2D GetScrollbarThickness() const;
+
 	UFUNCTION(BlueprintCallable, Category = "Scroll")
 	void SetScrollbarPadding(const FMargin& NewScrollbarPadding);
 
+	FMargin GetScrollbarPadding() const;
+
 	UFUNCTION(BlueprintCallable, Category = "Scroll")
 	void SetAlwaysShowScrollbar(bool NewAlwaysShowScrollbar);
 	
+	bool IsAlwaysShowScrollbar() const;
+
 	UFUNCTION(BlueprintCallable, Category = "Scroll")
 	void SetAllowOverscroll(bool NewAllowOverscroll);
 
+	bool IsAllowOverscroll() const;
+
 	UFUNCTION(BlueprintCallable, Category = "Scroll")
 	void SetAnimateWheelScrolling(bool bShouldAnimateWheelScrolling);
 
+	bool IsAnimateWheelScrolling() const;
+
 	UFUNCTION(BlueprintCallable, Category = "Scroll")
 	void SetWheelScrollMultiplier(float NewWheelScrollMultiplier);
 
+	float GetWheelScrollMultiplier() const;
+
 	UFUNCTION(BlueprintCallable, Category = "Scroll")
 	void SetScrollWhenFocusChanges(EScrollWhenFocusChanges NewScrollWhenFocusChanges);
 
+	EScrollWhenFocusChanges GetScrollWhenFocusChanges() const;
+
 	UFUNCTION(BlueprintCallable, Category = "Scroll")
 	void SetNavigationDestination(const EDescendantScrollDestination NewNavigationDestination);
 
+	EDescendantScrollDestination GetNavigationDestination() const;
+
+	void SetAlwaysShowScrollbarTrack(bool NewAlwaysShowScrollbarTrack);
+
+	bool IsAlwaysShowScrollbarTrack() const;
+
+	float GetNavigationScrollPadding() const;
+
+	void SetAllowRightClickDragScrolling(bool bShouldAllowRightClickDragScrolling);
+
+	bool IsAllowRightClickDragScrolling() const;
+
+	bool IsFrontPadScrolling() const;
+
+	bool IsBackPadScrolling() const;
+
 	/** Instantly stops any inertial scrolling that is currently in progress */
 	UFUNCTION(BlueprintCallable, Category = "Scroll")
 	void EndInertialScrolling();
@@ -210,6 +271,13 @@ protected:
 
 	void SlateHandleUserScrolled(float CurrentOffset);
 
+	// Initialize IsFocusable in the constructor before the SWidget is constructed.
+	void InitBackPadScrolling(bool InBackPadScrolling);
+	// Initialize IsFocusable in the constructor before the SWidget is constructed.
+	void InitFrontPadScrolling(bool InFrontPadScrolling);
+	// Initialize IsFocusable in the constructor before the SWidget is constructed.
+	void InitNavigationScrollPadding(float InNavigationScrollPadding);
+
 protected:
 	/** The desired scroll offset for the underlying scrollbox.  This is a cache so that it can be set before the widget is constructed. */
 	float DesiredScrollOffset;
