Fix to allow BlueprintWrite for bIsFocusable property of userwidget

diff --git a/Engine/Source/Runtime/UMG/Public/Blueprint/UserWidget.h b/Engine/Source/Runtime/UMG/Public/Blueprint/UserWidget.h
--- a/Engine/Source/Runtime/UMG/Public/Blueprint/UserWidget.h
+++ b/Engine/Source/Runtime/UMG/Public/Blueprint/UserWidget.h
@@ -930,7 +930,7 @@ public:
 
 	UE_DEPRECATED(5.2, "Direct access to bIsFocusable is deprecated. Please use the getter. Note that this property is only set at construction and is not modifiable at runtime.")
 	/** Setting this flag to true, allows this widget to accept focus when clicked, or when navigated to. */
-	UPROPERTY(EditAnywhere, BlueprintReadOnly, Getter = "IsFocusable", Category = "Interaction")
+	UPROPERTY(EditAnywhere, BlueprintReadWrite, Getter = "IsFocusable", Category = "Interaction")
 	uint8 bIsFocusable : 1;
 	 
 	UE_DEPRECATED(5.2, "Direct access to bStopAction is deprecated. Please use the getter or setter.")
