Remove ensure message in InitIsFocusable

diff --git a/Engine/Source/Runtime/UMG/Private/UserWidget.cpp b/Engine/Source/Runtime/UMG/Private/UserWidget.cpp
--- a/Engine/Source/Runtime/UMG/Private/UserWidget.cpp
+++ b/Engine/Source/Runtime/UMG/Private/UserWidget.cpp
@@ -1785,7 +1785,6 @@ bool UUserWidget::IsFocusable() const
 
 void UUserWidget::InitIsFocusable(bool InIsFocusable)
 {
-	ensureMsgf(!MyWidget.IsValid(), TEXT("The widget is already created."));
 	bIsFocusable = InIsFocusable;
 }
 
