Support blueprintwrite for style properties of Combobox

diff --git a/Engine/Source/Runtime/Slate/Public/Widgets/Input/SComboBox.h b/Engine/Source/Runtime/Slate/Public/Widgets/Input/SComboBox.h
--- a/Engine/Source/Runtime/Slate/Public/Widgets/Input/SComboBox.h
+++ b/Engine/Source/Runtime/Slate/Public/Widgets/Input/SComboBox.h
@@ -194,13 +194,14 @@ public:
 		check(InArgs._ComboBoxStyle);
 
 		ItemStyle = InArgs._ItemStyle;
-		MenuRowPadding = InArgs._ComboBoxStyle->MenuRowPadding;
+		ComboBoxStyle = InArgs._ComboBoxStyle;
+		MenuRowPadding = ComboBoxStyle->MenuRowPadding;
 
 		// Work out which values we should use based on whether we were given an override, or should use the style's version
-		const FComboButtonStyle& OurComboButtonStyle = InArgs._ComboBoxStyle->ComboButtonStyle;
+		const FComboButtonStyle& OurComboButtonStyle = ComboBoxStyle->ComboButtonStyle;
 		const FButtonStyle* const OurButtonStyle = InArgs._ButtonStyle ? InArgs._ButtonStyle : &OurComboButtonStyle.ButtonStyle;
-		PressedSound = InArgs._PressedSoundOverride.Get(InArgs._ComboBoxStyle->PressedSlateSound);
-		SelectionChangeSound = InArgs._SelectionChangeSoundOverride.Get(InArgs._ComboBoxStyle->SelectionChangeSlateSound);
+		PressedSound = InArgs._PressedSoundOverride.Get(ComboBoxStyle->PressedSlateSound);
+		SelectionChangeSound = InArgs._SelectionChangeSoundOverride.Get(ComboBoxStyle->SelectionChangeSlateSound);
 
 		this->OnComboBoxOpening = InArgs._OnComboBoxOpening;
 		this->OnSelectionChanged = InArgs._OnSelectionChanged;
@@ -373,6 +374,34 @@ public:
 		ComboBoxMenuContent->SetMaxDesiredHeight(InMaxHeight);
 	}
 
+	void SetStyle(const FComboBoxStyle* InStyle) 
+	{ 
+		if (ComboBoxStyle != InStyle)
+		{
+			ComboBoxStyle = InStyle;
+			InvalidateStyle();
+		}
+	}
+
+	void InvalidateStyle() 
+	{ 
+		Invalidate(EInvalidateWidgetReason::Layout); 
+	}
+
+	void SetItemStyle(const FTableRowStyle* InItemStyle) 
+	{ 
+		if (ItemStyle != InItemStyle)
+		{
+			ItemStyle = InItemStyle;
+			InvalidateItemStyle();
+		}
+	}
+
+	void InvalidateItemStyle() 
+	{ 
+		Invalidate(EInvalidateWidgetReason::Layout); 
+	}
+
 	/** @return the item currently selected by the combo box. */
 	NullableOptionType GetSelectedItem()
 	{
@@ -630,6 +659,9 @@ private:
 	/** The item style to use. */
 	const FTableRowStyle* ItemStyle;
 
+	/** The combo box style to use. */
+	const FComboBoxStyle* ComboBoxStyle;
+	
 	/** The padding around each menu row */
 	FMargin MenuRowPadding;
 
diff --git a/Engine/Source/Runtime/UMG/Private/Components/ComboBoxKey.cpp b/Engine/Source/Runtime/UMG/Private/Components/ComboBoxKey.cpp
--- a/Engine/Source/Runtime/UMG/Private/Components/ComboBoxKey.cpp
+++ b/Engine/Source/Runtime/UMG/Private/Components/ComboBoxKey.cpp
@@ -253,16 +253,22 @@ FSlateColor UComboBoxKey::GetForegroundColor() const
 	return ForegroundColor;
 }
 
-void UComboBoxKey::InitWidgetStyle(const FComboBoxStyle& InWidgetStyle)
+void UComboBoxKey::SetWidgetStyle(const FComboBoxStyle& InWidgetStyle)
 {
-	ensureMsgf(!MyComboBox.IsValid(), TEXT("The widget is already created."));
 	WidgetStyle = InWidgetStyle;
+	if (MyComboBox.IsValid())
+	{
+		MyComboBox->InvalidateStyle();
+	}
 }
 
-void UComboBoxKey::InitItemStyle(const FTableRowStyle& InItemStyle)
+void UComboBoxKey::SetItemStyle(const FTableRowStyle& InItemStyle)
 {
-	ensureMsgf(!MyComboBox.IsValid(), TEXT("The widget is already created."));
 	ItemStyle = InItemStyle;
+	if (MyComboBox.IsValid())
+	{
+		MyComboBox->InvalidateItemStyle();
+	}
 }
 
 void UComboBoxKey::InitScrollBarStyle(const FScrollBarStyle& InScrollBarStyle)
diff --git a/Engine/Source/Runtime/UMG/Private/Components/ComboBoxString.cpp b/Engine/Source/Runtime/UMG/Private/Components/ComboBoxString.cpp
--- a/Engine/Source/Runtime/UMG/Private/Components/ComboBoxString.cpp
+++ b/Engine/Source/Runtime/UMG/Private/Components/ComboBoxString.cpp
@@ -371,16 +371,22 @@ FSlateColor UComboBoxString::GetForegroundColor() const
 	return ForegroundColor;
 }
 
-void UComboBoxString::InitWidgetStyle(const FComboBoxStyle& InWidgetStyle)
+void UComboBoxString::SetWidgetStyle(const FComboBoxStyle& InWidgetStyle)
 {
-	ensureMsgf(!MyComboBox.IsValid(), TEXT("The widget is already created."));
 	WidgetStyle = InWidgetStyle;
+	if (MyComboBox.IsValid())
+	{
+		MyComboBox->InvalidateStyle();
+	}
 }
 
-void UComboBoxString::InitItemStyle(const FTableRowStyle& InItemStyle)
+void UComboBoxString::SetItemStyle(const FTableRowStyle& InItemStyle)
 {
-	ensureMsgf(!MyComboBox.IsValid(), TEXT("The widget is already created."));
 	ItemStyle = InItemStyle;
+	if (MyComboBox.IsValid())
+	{
+		MyComboBox->InvalidateItemStyle();
+	}
 }
 
 void UComboBoxString::InitScrollBarStyle(const FScrollBarStyle& InScrollBarStyle)
diff --git a/Engine/Source/Runtime/UMG/Public/Components/ComboBoxKey.h b/Engine/Source/Runtime/UMG/Public/Components/ComboBoxKey.h
--- a/Engine/Source/Runtime/UMG/Public/Components/ComboBoxKey.h
+++ b/Engine/Source/Runtime/UMG/Public/Components/ComboBoxKey.h
@@ -36,14 +36,14 @@ private:
 	FName SelectedOption;
 
 public:
-	UE_DEPRECATED(5.2, "Direct access to WidgetStyle is deprecated. Please use the getter. Note that this property is only set at construction and is not modifiable at runtime.")
+	UE_DEPRECATED(5.2, "Direct access to WidgetStyle is deprecated. Please use the getter or setter.")
 	/** The combobox style. */
-	UPROPERTY(EditAnywhere, BlueprintReadOnly, Getter, Category = Style, meta = (DisplayName = "Style"))
+	UPROPERTY(EditAnywhere, BlueprintReadWrite, Getter, Setter, Category = Style, meta = (DisplayName = "Style"))
 	FComboBoxStyle WidgetStyle;
 
-	UE_DEPRECATED(5.2, "Direct access to ItemStyle is deprecated. Please use the getter. Note that this property is only set at construction and is not modifiable at runtime.")
+	UE_DEPRECATED(5.2, "Direct access to ItemStyle is deprecated. Please use the getter or setter.")
 	/** The item row style. */
-	UPROPERTY(EditAnywhere, BlueprintReadOnly, Getter, Category = Style)
+	UPROPERTY(EditAnywhere, BlueprintReadWrite, Getter, Setter, Category = Style)
 	FTableRowStyle ItemStyle;
 
 	UE_DEPRECATED(5.2, "Direct access to ScrollBarStyle is deprecated. Please use the getter. Note that this property is only set at construction and is not modifiable at runtime.")
@@ -167,9 +167,15 @@ public:
 	/** Get the style of the combobox. */
 	const FComboBoxStyle& GetWidgetStyle() const;
 
+	/** Set the style of the combobox. */
+	void SetWidgetStyle(const FComboBoxStyle& InWidgetStyle);
+
 	/** Get the style of the items. */
 	const FTableRowStyle& GetItemStyle() const;
 
+	/** Set the style of the items. */
+	void SetItemStyle(const FTableRowStyle& InItemStyle);
+
 	/** Get the style of the scrollbar. */
 	const FScrollBarStyle& GetScrollBarStyle() const;
 
@@ -184,11 +190,6 @@ public:
 #endif
 
 protected:
-	/** Initialize the widget style in the constructor before the SWidget is constructed. */
-	void InitWidgetStyle(const FComboBoxStyle& InWidgetStyle);
-
-	/** Initialize the item style in the constructor before the SWidget is constructed. */
-	void InitItemStyle(const FTableRowStyle& InItemStyle);
 
 	/** Initialize the scrollbar style in the constructor before the SWidget is constructed. */
 	void InitScrollBarStyle(const FScrollBarStyle& InScrollBarStyle);
diff --git a/Engine/Source/Runtime/UMG/Public/Components/ComboBoxString.h b/Engine/Source/Runtime/UMG/Public/Components/ComboBoxString.h
--- a/Engine/Source/Runtime/UMG/Public/Components/ComboBoxString.h
+++ b/Engine/Source/Runtime/UMG/Public/Components/ComboBoxString.h
@@ -37,14 +37,14 @@ private:
 
 public:
 
-	UE_DEPRECATED(5.2, "Direct access to WidgetStyle is deprecated. Please use the getter. Note that this property is only set at construction and is not modifiable at runtime.")
+	UE_DEPRECATED(5.2, "Direct access to WidgetStyle is deprecated. Please use the getter or setter.")
 	/** The style. */
-	UPROPERTY(EditAnywhere, BlueprintReadOnly, Getter, Category=Style, meta=( DisplayName="Style" ))
+	UPROPERTY(EditAnywhere, BlueprintReadWrite, Getter, Setter, Category=Style, meta=( DisplayName="Style" ))
 	FComboBoxStyle WidgetStyle;
 
-	UE_DEPRECATED(5.2, "Direct access to ItemStyle is deprecated. Please use the getter. Note that this property is only set at construction and is not modifiable at runtime.")
+	UE_DEPRECATED(5.2, "Direct access to ItemStyle is deprecated. Please use the getter or setter.")
 	/** The item row style. */
-	UPROPERTY(EditAnywhere, BlueprintReadOnly, Getter, Category=Style)
+	UPROPERTY(EditAnywhere, BlueprintReadWrite, Getter, Setter, Category=Style)
 	FTableRowStyle ItemStyle;
 	
 	UE_DEPRECATED(5.2, "Direct access to ScrollBarStyle is deprecated. Please use the getter. Note that this property is only set at construction and is not modifiable at runtime.")
@@ -188,9 +188,15 @@ public:
 	/** Get the style of the combobox. */
 	const FComboBoxStyle& GetWidgetStyle() const;
 
+	/** Set the style of the combobox. */
+	void SetWidgetStyle(const FComboBoxStyle& InWidgetStyle);
+
 	/** Get the style of the items. */
 	const FTableRowStyle& GetItemStyle() const;
 
+	/** Set the style of the items. */
+	void SetItemStyle(const FTableRowStyle& InItemStyle);
+
 	/** Get the style of the scrollbar. */
 	const FScrollBarStyle& GetScrollBarStyle() const;
 
@@ -227,12 +233,6 @@ protected:
 	virtual TSharedRef<SWidget> RebuildWidget() override;
 	//~ End UWidget Interface
 
-	/** Initialize the widget style in the constructor before the SWidget is constructed. */
-	void InitWidgetStyle(const FComboBoxStyle& InWidgetStyle);
-
-	/** Initialize the item style in the constructor before the SWidget is constructed. */
-	void InitItemStyle(const FTableRowStyle& InItemStyle);
-
 	/** Initialize the scrollbar style in the constructor before the SWidget is constructed. */
 	void InitScrollBarStyle(const FScrollBarStyle& InScrollBarStyle);
 
