Fix to allow BlueprintWrite for bIsFocusable property of userwidget

diff --git a/Engine/Plugins/Experimental/CommonUI/Source/CommonUI/Private/CommonButtonBase.cpp b/Engine/Plugins/Experimental/CommonUI/Source/CommonUI/Private/CommonButtonBase.cpp
--- a/Engine/Plugins/Experimental/CommonUI/Source/CommonUI/Private/CommonButtonBase.cpp
+++ b/Engine/Plugins/Experimental/CommonUI/Source/CommonUI/Private/CommonButtonBase.cpp
@@ -335,7 +335,7 @@ UCommonButtonBase::UCommonButtonBase(const FObjectInitializer& ObjectInitializer
 	, bButtonEnabled(true)
 	, bInteractionEnabled(true)
 {
-	InitIsFocusable(true);
+	SetIsFocusable(true);
 }
 
 void UCommonButtonBase::OnWidgetRebuilt()
@@ -1540,7 +1540,7 @@ void UCommonButtonBase::DisableButton()
 
 void UCommonButtonBase::SetIsFocusable(bool bInIsFocusable)
 {
-	InitIsFocusable(bInIsFocusable);
+	UUserWidget::SetIsFocusable(bInIsFocusable);
 
 	if (RootButton.IsValid())
 	{
diff --git a/Engine/Source/Runtime/UMG/Private/UserWidget.cpp b/Engine/Source/Runtime/UMG/Private/UserWidget.cpp
--- a/Engine/Source/Runtime/UMG/Private/UserWidget.cpp
+++ b/Engine/Source/Runtime/UMG/Private/UserWidget.cpp
@@ -1783,9 +1783,10 @@ bool UUserWidget::IsFocusable() const
 	return bIsFocusable;
 }
 
-void UUserWidget::InitIsFocusable(bool InIsFocusable)
+void UUserWidget::SetIsFocusable(bool InIsFocusable)
 {
 	bIsFocusable = InIsFocusable;
+	Invalidate(EInvalidateWidgetReason::Paint);
 }
 
 PRAGMA_ENABLE_DEPRECATION_WARNINGS
diff --git a/Engine/Source/Runtime/UMG/Public/Blueprint/UserWidget.h b/Engine/Source/Runtime/UMG/Public/Blueprint/UserWidget.h
--- a/Engine/Source/Runtime/UMG/Public/Blueprint/UserWidget.h
+++ b/Engine/Source/Runtime/UMG/Public/Blueprint/UserWidget.h
@@ -930,7 +930,7 @@ public:
 
 	UE_DEPRECATED(5.2, "Direct access to bIsFocusable is deprecated. Please use the getter. Note that this property is only set at construction and is not modifiable at runtime.")
 	/** Setting this flag to true, allows this widget to accept focus when clicked, or when navigated to. */
-	UPROPERTY(EditAnywhere, BlueprintReadWrite, Getter = "IsFocusable", Category = "Interaction")
+	UPROPERTY(EditAnywhere, BlueprintReadWrite, Getter = "IsFocusable", Setter = "SetIsFocusable", Category = "Interaction")
 	uint8 bIsFocusable : 1;
 	 
 	UE_DEPRECATED(5.2, "Direct access to bStopAction is deprecated. Please use the getter or setter.")
@@ -989,10 +989,12 @@ public:
 	bool IsInputActionBlocking() const;
 
 	/**
-	 * Returns whether this widget to accept focus when clicked, or when navigated to.
+	 * Sets whether this widget to accept focus when clicked, or when navigated to.
 	 */
 	bool IsFocusable() const;
 
+	void SetIsFocusable(bool InIsFocusable);
+
 	/**
 	 * Plays an animation in this widget a specified number of times
 	 * 
@@ -1459,9 +1461,6 @@ protected:
 
 	virtual void InitializeInputComponent();
 
-	/** Initialize IsFocusable in the constructor before the SWidget is constructed. */
-	void InitIsFocusable(bool InIsFocusable);
-
 private:
 	FVector2D MinimumDesiredSize;
 
