Add accessor to RichTextBlock

diff --git a/Engine/Plugins/Experimental/CommonUI/Source/CommonUI/Private/CommonRichTextBlock.cpp b/Engine/Plugins/Experimental/CommonUI/Source/CommonUI/Private/CommonRichTextBlock.cpp
--- a/Engine/Plugins/Experimental/CommonUI/Source/CommonUI/Private/CommonRichTextBlock.cpp
+++ b/Engine/Plugins/Experimental/CommonUI/Source/CommonUI/Private/CommonRichTextBlock.cpp
@@ -482,14 +482,14 @@ void UCommonRichTextBlock::Serialize(FArchive& Ar)
 
 	if (Ar.IsLoading() && bDisplayAllCaps_DEPRECATED)
 	{
-		TextTransformPolicy = ETextTransformPolicy::ToUpper;
+		SetTextTransformPolicy(ETextTransformPolicy::ToUpper);
 		bDisplayAllCaps_DEPRECATED = false;
 	}
 }
 
 void UCommonRichTextBlock::SetText(const FText& InText)
 {
-	Text = InText;
+	Super::SetText(InText);
 	if (MyRichTextBlock.IsValid())
 	{
 		if (CommonUIUtils::ShouldDisplayMobileUISizes())
@@ -530,8 +530,8 @@ bool UCommonRichTextBlock::CanEditChange(const FProperty* InProperty) const
 
 	if (DefaultTextStyleOverrideClass)
 	{
-		if (PropertyName.IsEqual(GET_MEMBER_NAME_CHECKED(UCommonRichTextBlock, DefaultTextStyleOverride)) ||
-			PropertyName.IsEqual(GET_MEMBER_NAME_CHECKED(UCommonRichTextBlock, MinDesiredWidth)))
+		if (PropertyName.IsEqual(TEXT("DefaultTextStyleOverride")) ||
+			PropertyName.IsEqual(TEXT("MinDesiredWidth")))
 		{
 			return false;
 		}
diff --git a/Engine/Plugins/Experimental/CommonUI/Source/CommonUI/Public/CommonRichTextBlock.h b/Engine/Plugins/Experimental/CommonUI/Source/CommonUI/Public/CommonRichTextBlock.h
--- a/Engine/Plugins/Experimental/CommonUI/Source/CommonUI/Public/CommonRichTextBlock.h
+++ b/Engine/Plugins/Experimental/CommonUI/Source/CommonUI/Public/CommonRichTextBlock.h
@@ -35,7 +35,7 @@ class COMMONUI_API UCommonRichTextBlock : public URichTextBlock
 	GENERATED_BODY()
 
 public:
-	ETextTransformPolicy GetTextTransformPolicy() const { return TextTransformPolicy; }
+	ETextTransformPolicy GetTextTransformPolicy() const { return GetTransformPolicy(); }
 	TSubclassOf<UCommonTextStyle> GetDefaultTextStyleClass() const { return DefaultTextStyleOverrideClass; }
 	float GetMobileTextBlockScale() const { return MobileTextBlockScale; }
 
diff --git a/Engine/Source/Runtime/UMG/Private/Components/RichTextBlock.cpp b/Engine/Source/Runtime/UMG/Private/Components/RichTextBlock.cpp
--- a/Engine/Source/Runtime/UMG/Private/Components/RichTextBlock.cpp
+++ b/Engine/Source/Runtime/UMG/Private/Components/RichTextBlock.cpp
@@ -46,6 +46,7 @@ FORCEINLINE TSharedPtr< ObjectType > MakeShareableDeferredCleanup(ObjectType* In
 	return MakeShareable(InObject, [](ObjectType* ObjectToDelete) { BeginCleanup(new FDeferredDeletor<ObjectType>(ObjectToDelete)); });
 }
 
+PRAGMA_DISABLE_DEPRECATION_WARNINGS
 URichTextBlock::URichTextBlock(const FObjectInitializer& ObjectInitializer)
 	: Super(ObjectInitializer)
 {
@@ -53,6 +54,7 @@ URichTextBlock::URichTextBlock(const FObjectInitializer& ObjectInitializer)
 	TextTransformPolicy = ETextTransformPolicy::None;
 	TextOverflowPolicy = ETextOverflowPolicy::Clip;
 }
+PRAGMA_ENABLE_DEPRECATION_WARNINGS
 
 void URichTextBlock::ReleaseSlateResources(bool bReleaseChildren)
 {
@@ -71,12 +73,12 @@ TSharedRef<SWidget> URichTextBlock::RebuildWidget()
 	CreateDecorators(CreatedDecorators);
 
 	TSharedRef<FRichTextLayoutMarshaller> Marshaller = FRichTextLayoutMarshaller::Create(CreateMarkupParser(), CreateMarkupWriter(), CreatedDecorators, StyleInstance.Get());
-
+PRAGMA_DISABLE_DEPRECATION_WARNINGS
 	MyRichTextBlock =
 		SNew(SRichTextBlock)
 		.TextStyle(bOverrideDefaultStyle ? &DefaultTextStyleOverride : &DefaultTextStyle)
 		.Marshaller(Marshaller);
-	
+PRAGMA_ENABLE_DEPRECATION_WARNINGS
 	return MyRichTextBlock.ToSharedRef();
 }
 
@@ -86,12 +88,13 @@ void URichTextBlock::SynchronizeProperties()
 
 	if (MyRichTextBlock.IsValid())
 	{
+PRAGMA_DISABLE_DEPRECATION_WARNINGS
 		MyRichTextBlock->SetText(Text);
 		MyRichTextBlock->SetTransformPolicy(TextTransformPolicy);
 		MyRichTextBlock->SetMinDesiredWidth(MinDesiredWidth);
 
 		MyRichTextBlock->SetOverflowPolicy(TextOverflowPolicy);
-
+PRAGMA_ENABLE_DEPRECATION_WARNINGS
 		Super::SynchronizeTextLayoutProperties(*MyRichTextBlock);
 	}
 }
@@ -121,6 +124,7 @@ void URichTextBlock::UpdateStyleData()
 	}
 }
 
+PRAGMA_DISABLE_DEPRECATION_WARNINGS
 FText URichTextBlock::GetText() const
 {
 	if (MyRichTextBlock.IsValid())
@@ -200,6 +204,7 @@ const FTextBlockStyle& URichTextBlock::GetCurrentDefaultTextStyle() const
 		return DefaultTextStyle;
 	}
 }
+PRAGMA_ENABLE_DEPRECATION_WARNINGS
 
 URichTextBlockDecorator* URichTextBlock::GetDecoratorByClass(TSubclassOf<URichTextBlockDecorator> DecoratorClass)
 {
@@ -239,6 +244,7 @@ TSharedPtr< IRichTextMarkupWriter > URichTextBlock::CreateMarkupWriter()
 	return FDefaultRichTextMarkupWriter::Create();
 }
 
+PRAGMA_DISABLE_DEPRECATION_WARNINGS
 void URichTextBlock::BeginDefaultStyleOverride()
 {
 	if (!bOverrideDefaultStyle)
@@ -248,6 +254,7 @@ void URichTextBlock::BeginDefaultStyleOverride()
 		DefaultTextStyleOverride = DefaultTextStyle;
 	}
 }
+PRAGMA_ENABLE_DEPRECATION_WARNINGS
 
 #if WITH_EDITOR
 
@@ -261,6 +268,7 @@ void URichTextBlock::OnCreationFromPalette()
 	//Decorators.Add(NewObject<URichTextBlockDecorator>(this, NAME_None, RF_Transactional));
 }
 
+PRAGMA_DISABLE_DEPRECATION_WARNINGS
 void URichTextBlock::ValidateCompiledDefaults(IWidgetCompilerLog& CompileLog) const
 {
 	Super::ValidateCompiledDefaults(CompileLog);
@@ -273,9 +281,11 @@ void URichTextBlock::ValidateCompiledDefaults(IWidgetCompilerLog& CompileLog) co
 			FText::AsCultureInvariant(TextStyleSet->GetPathName())));
 	}
 }
+PRAGMA_ENABLE_DEPRECATION_WARNINGS
 
 #endif //if WITH_EDITOR
 
+PRAGMA_DISABLE_DEPRECATION_WARNINGS
 void URichTextBlock::SetDefaultTextStyle(const FTextBlockStyle& InDefaultTextStyle)
 {
 	BeginDefaultStyleOverride();
@@ -298,10 +308,13 @@ void URichTextBlock::ClearAllDefaultStyleOverrides()
 		ApplyUpdatedDefaultTextStyle();
 	}
 }
+PRAGMA_ENABLE_DEPRECATION_WARNINGS
 
 UMaterialInstanceDynamic* URichTextBlock::GetDefaultDynamicMaterial()
 {
+	PRAGMA_DISABLE_DEPRECATION_WARNINGS
 	TObjectPtr<UObject>& MaterialObject = bOverrideDefaultStyle ? DefaultTextStyleOverride.Font.FontMaterial : DefaultTextStyle.Font.FontMaterial;
+	PRAGMA_ENABLE_DEPRECATION_WARNINGS
 
 	if (UMaterialInterface* Material = Cast<UMaterialInterface>(MaterialObject.Get()))
 	{
@@ -311,9 +324,11 @@ UMaterialInstanceDynamic* URichTextBlock::GetDefaultDynamicMaterial()
 		{
 			BeginDefaultStyleOverride();
 
+PRAGMA_DISABLE_DEPRECATION_WARNINGS
 			DynamicMaterial = UMaterialInstanceDynamic::Create(Material, this);
 			DefaultTextStyleOverride.Font.FontMaterial = DynamicMaterial;
-			
+PRAGMA_ENABLE_DEPRECATION_WARNINGS
+
 			ApplyUpdatedDefaultTextStyle();
 		}
 
@@ -340,6 +355,7 @@ void URichTextBlock::SetDecorators(const TArray<TSubclassOf<URichTextBlockDecora
 	}
 }
 
+PRAGMA_DISABLE_DEPRECATION_WARNINGS
 void URichTextBlock::SetDefaultColorAndOpacity(FSlateColor InColorAndOpacity)
 {
 	BeginDefaultStyleOverride();
@@ -375,6 +391,7 @@ void URichTextBlock::SetDefaultStrikeBrush(FSlateBrush& InStrikeBrush)
 	DefaultTextStyleOverride.StrikeBrush = InStrikeBrush;
 	ApplyUpdatedDefaultTextStyle();
 }
+PRAGMA_ENABLE_DEPRECATION_WARNINGS
 
 void URichTextBlock::SetJustification(ETextJustify::Type InJustification)
 {
@@ -385,6 +402,7 @@ void URichTextBlock::SetJustification(ETextJustify::Type InJustification)
 	}
 }
 
+PRAGMA_DISABLE_DEPRECATION_WARNINGS
 void URichTextBlock::SetMinDesiredWidth(float InMinDesiredWidth)
 {
 	MinDesiredWidth = InMinDesiredWidth;
@@ -393,6 +411,7 @@ void URichTextBlock::SetMinDesiredWidth(float InMinDesiredWidth)
 		MyRichTextBlock->SetMinDesiredWidth(InMinDesiredWidth);
 	}
 }
+PRAGMA_ENABLE_DEPRECATION_WARNINGS
 
 void URichTextBlock::SetAutoWrapText(bool InAutoTextWrap)
 {
@@ -403,6 +422,7 @@ void URichTextBlock::SetAutoWrapText(bool InAutoTextWrap)
 	}
 }
 
+PRAGMA_DISABLE_DEPRECATION_WARNINGS
 void URichTextBlock::SetTextTransformPolicy(ETextTransformPolicy InTransformPolicy)
 {
 	TextTransformPolicy = InTransformPolicy;
@@ -422,6 +442,26 @@ void URichTextBlock::SetTextOverflowPolicy(ETextOverflowPolicy InOverflowPolicy)
 	}
 }
 
+const FTextBlockStyle& URichTextBlock::GetDefaultTextStyleOverride() const
+{
+	return DefaultTextStyleOverride;
+}
+
+float URichTextBlock::GetMinDesiredWidth() const
+{
+	return MinDesiredWidth;
+}
+
+ETextTransformPolicy URichTextBlock::GetTransformPolicy() const
+{
+	return TextTransformPolicy;
+}
+
+ETextOverflowPolicy URichTextBlock::GetOverflowPolicy() const
+{
+	return TextOverflowPolicy;
+}
+
 void URichTextBlock::ApplyUpdatedDefaultTextStyle()
 {
 	if (MyRichTextBlock.IsValid())
@@ -429,6 +469,7 @@ void URichTextBlock::ApplyUpdatedDefaultTextStyle()
 		MyRichTextBlock->SetTextStyle(bOverrideDefaultStyle ? DefaultTextStyleOverride : DefaultTextStyle);
 	}
 }
+PRAGMA_ENABLE_DEPRECATION_WARNINGS
 
 void URichTextBlock::RefreshTextLayout()
 {
diff --git a/Engine/Source/Runtime/UMG/Public/Components/RichTextBlock.h b/Engine/Source/Runtime/UMG/Public/Components/RichTextBlock.h
--- a/Engine/Source/Runtime/UMG/Public/Components/RichTextBlock.h
+++ b/Engine/Source/Runtime/UMG/Public/Components/RichTextBlock.h
@@ -39,6 +39,44 @@ UCLASS()
 class UMG_API URichTextBlock : public UTextLayoutWidget
 {
 	GENERATED_BODY()
+protected:
+	UE_DEPRECATED(5.2, "Direct access to Text is deprecated. Please use the getter or setter.")
+	/** The text to display */
+	UPROPERTY(EditAnywhere, Category = Content, BlueprintSetter = "SetText", BlueprintGetter = "GetText", Getter, Setter, meta = (MultiLine = "true"))
+	FText Text;
+
+	UE_DEPRECATED(5.2, "Direct access to TextStyleSet is deprecated. Please use the getter or setter.")
+	/**  */
+	UPROPERTY(EditAnywhere, Category = Appearance, BlueprintSetter = "SetTextStyleSet", BlueprintGetter = "GetTextStyleSet", Getter, Setter, meta = (RequiredAssetDataTags = "RowStructure=/Script/UMG.RichTextStyleRow"))
+	TObjectPtr<UDataTable> TextStyleSet;
+
+	/**  */
+	UPROPERTY(EditAnywhere, Category = Appearance)
+	TArray<TSubclassOf<URichTextBlockDecorator>> DecoratorClasses;
+
+	/** True to specify the default text style for this rich text inline, overriding any default provided in the style set table */
+	UPROPERTY(EditAnywhere, Category = Appearance)
+	bool bOverrideDefaultStyle = false;
+
+	UE_DEPRECATED(5.2, "Direct access to TextOverflowPolicy is deprecated. Please use the getter or setter.")
+	/** Text style to apply by default to text in this block */
+	UPROPERTY(EditAnywhere, Category = Appearance, BlueprintSetter = "SetDefaultTextStyle", Setter = "SetDefaultTextStyle", Getter, meta = (EditCondition = bOverrideDefaultStyle))
+	FTextBlockStyle DefaultTextStyleOverride;
+
+	UE_DEPRECATED(5.2, "Direct access to MinDesiredWidth is deprecated. Please use the getter or setter.")
+	/** The minimum desired size for the text */
+	UPROPERTY(EditAnywhere, BlueprintReadWrite, BlueprintSetter = "SetMinDesiredWidth", Getter, Setter, Category = Appearance)
+	float MinDesiredWidth;
+
+	UE_DEPRECATED(5.2, "Direct access to TextTransformPolicy is deprecated. Please use the getter or setter.")
+	/** The text transformation policy to apply to this text block */
+	UPROPERTY(EditAnywhere, BlueprintReadWrite, BlueprintSetter = "SetTextTransformPolicy", Getter = "GetTransformPolicy", Setter, Category = Appearance, meta = (DisplayName = "Transform Policy"))
+	ETextTransformPolicy TextTransformPolicy;
+
+	UE_DEPRECATED(5.2, "Direct access to TextOverflowPolicy is deprecated. Please use the getter or setter.")
+	/** Sets what happens to text that is clipped and doesn't fit within the clip rect for this widget */
+	UPROPERTY(EditAnywhere, BlueprintReadWrite, BlueprintSetter = "SetTextOverflowPolicy", Getter = "GetOverflowPolicy", Setter, Category = Clipping, AdvancedDisplay, meta = (DisplayName = "Overflow Policy"))
+	ETextOverflowPolicy TextOverflowPolicy;
 
 public:
 	/**  
@@ -200,39 +238,12 @@ protected:
 	void BeginDefaultStyleOverride();
 	virtual void ApplyUpdatedDefaultTextStyle();
 
-protected:
-	/** The text to display */
-	UPROPERTY(EditAnywhere, Category=Content, meta=( MultiLine="true" ))
-	FText Text;
-
-	/**  */
-	UPROPERTY(EditAnywhere, Category=Appearance, meta=(RequiredAssetDataTags = "RowStructure=/Script/UMG.RichTextStyleRow"))
-	TObjectPtr<UDataTable> TextStyleSet;
-
-	/**  */
-	UPROPERTY(EditAnywhere, Category=Appearance)
-	TArray<TSubclassOf<URichTextBlockDecorator>> DecoratorClasses;
-
-	/** True to specify the default text style for this rich text inline, overriding any default provided in the style set table */
-	UPROPERTY(EditAnywhere, Category = Appearance)
-	bool bOverrideDefaultStyle = false;
-
-	/** Text style to apply by default to text in this block */
-	UPROPERTY(EditAnywhere, Category=Appearance, meta = (EditCondition = bOverrideDefaultStyle))
-	FTextBlockStyle DefaultTextStyleOverride;
-
-	/** The minimum desired size for the text */
-	UPROPERTY(EditAnywhere, BlueprintReadOnly, Category=Appearance)
-	float MinDesiredWidth;
-
-	/** The text transformation policy to apply to this text block */
-	UPROPERTY(EditAnywhere, BlueprintReadOnly, Category=Appearance, meta=(DisplayName="Transform Policy"))
-	ETextTransformPolicy TextTransformPolicy;
-
-	/** Sets what happens to text that is clipped and doesn't fit within the clip rect for this widget */
-	UPROPERTY(EditAnywhere, BlueprintReadWrite, Category = Clipping, AdvancedDisplay, meta = (DisplayName = "Overflow Policy"))
-	ETextOverflowPolicy TextOverflowPolicy;
+	const FTextBlockStyle& GetDefaultTextStyleOverride() const;
+	float GetMinDesiredWidth() const;
+	ETextTransformPolicy GetTransformPolicy() const;
+	ETextOverflowPolicy GetOverflowPolicy() const;
 
+protected:
 	UPROPERTY(Transient)
 	FTextBlockStyle DefaultTextStyle;
 
