UMG: Change GameViewportSubsystem from GameInstance to Engine subsystem. Fix widget created in the loading screen were removed when the game instance changed.

diff --git a/Engine/Source/Runtime/UMG/Private/Blueprint/GameViewportSubsystem.cpp b/Engine/Source/Runtime/UMG/Private/Blueprint/GameViewportSubsystem.cpp
--- a/Engine/Source/Runtime/UMG/Private/Blueprint/GameViewportSubsystem.cpp
+++ b/Engine/Source/Runtime/UMG/Private/Blueprint/GameViewportSubsystem.cpp
@@ -5,6 +5,7 @@
 #include "Blueprint/UserWidget.h"
 #include "Blueprint/WidgetLayoutLibrary.h"
 #include "Components/Widget.h"
+#include "Engine/Engine.h"
 #include "Engine/GameViewportClient.h"
 #include "Engine/LocalPlayer.h"
 #include "Engine/World.h"
@@ -32,12 +33,17 @@ namespace UE::UMG::Private
 void UGameViewportSubsystem::Initialize(FSubsystemCollectionBase& Collection)
 {
 	Super::Initialize(Collection);
-	FWorldDelegates::LevelRemovedFromWorld.AddUObject(this, &ThisClass::OnLevelRemovedFromWorld);
+	
+	FWorldDelegates::OnWorldCleanup.AddUObject(this, &ThisClass::HandleWorldCleanup);
+	FWorldDelegates::OnWorldBeginTearDown.AddUObject(this, &ThisClass::HandleRemoveWorld);
+	FWorldDelegates::OnPreWorldFinishDestroy.AddUObject(this, &ThisClass::HandleRemoveWorld);
 }
 
 void UGameViewportSubsystem::Deinitialize()
 {
-	FWorldDelegates::LevelRemovedFromWorld.RemoveAll(this);
+	FWorldDelegates::OnPreWorldFinishDestroy.RemoveAll(this);
+	FWorldDelegates::OnWorldBeginTearDown.RemoveAll(this);
+	FWorldDelegates::OnWorldCleanup.RemoveAll(this);
 
 	if (ViewportWidgets.Num() > 0)
 	{
@@ -68,14 +74,12 @@ void UGameViewportSubsystem::Deinitialize()
 
 UGameViewportSubsystem* UGameViewportSubsystem::Get(UWorld* World)
 {
-	if (World)
-	{
-		if (UGameInstance* GameInstance = World->GetGameInstance())
-		{
-			return GameInstance->GetSubsystem<UGameViewportSubsystem>();
-		}
-	}
-	return nullptr;
+	return Get();
+}
+
+UGameViewportSubsystem* UGameViewportSubsystem::Get()
+{
+	return GEngine->GetEngineSubsystem<UGameViewportSubsystem>();
 }
 
 bool UGameViewportSubsystem::IsWidgetAdded(const UWidget* Widget) const
@@ -128,7 +132,7 @@ void UGameViewportSubsystem::AddToScreen(UWidget* Widget, ULocalPlayer* Player,
 	}
 
 	// Add the widget to the current worlds viewport.
-	UWorld* World = GetGameInstance()->GetWorld();
+	UWorld* World = Widget->GetWorld();
 	if (!World || !World->IsGameWorld())
 	{
 		FFrame::KismetExecutionMessage(*FString::Printf(TEXT("The widget '%s' does not have a World."), *Widget->GetName()), ELogVerbosity::Warning);
@@ -200,7 +204,7 @@ void UGameViewportSubsystem::RemoveWidgetInternal(UWidget* Widget, const TWeakPt
 	if (TSharedPtr<SWidget> WidgetHost = FullScreenWidget.Pin())
 	{
 		// If this is a game world remove the widget from the current world's viewport.
-		UWorld* World = GetGameInstance()->GetWorld();
+		UWorld* World = Widget->GetWorld();
 		if (World && World->IsGameWorld())
 		{
 			if (UGameViewportClient* ViewportClient = World->GetGameViewport())
@@ -277,27 +281,32 @@ FGameViewportWidgetSlot UGameViewportSubsystem::SetWidgetSlotDesiredSize(FGameVi
 	return Slot;
 }
 
-void UGameViewportSubsystem::OnLevelRemovedFromWorld(ULevel* InLevel, UWorld* InWorld)
+void UGameViewportSubsystem::HandleWorldCleanup(UWorld* InWorld, bool bSessionEnded, bool bCleanupResoures)
 {
-	// If the InLevel is null, it's a signal that the entire world is about to disappear, so
-	// go ahead and remove this widget from the viewport, it could be holding onto too many
-	// dangerous actor references that won't carry over into the next world.
-	if (InLevel == nullptr)
+	HandleRemoveWorld(InWorld);
+}
+
+void UGameViewportSubsystem::HandleRemoveWorld(UWorld* InWorld)
+{
+	TArray<UWidget*, TInlineAllocator<16>> WidgetsToRemove;
+	for (FViewportWidgetList::TIterator Itt = ViewportWidgets.CreateIterator(); Itt; ++Itt)
 	{
-		for (FViewportWidgetList::TIterator Itt = ViewportWidgets.CreateIterator(); Itt; ++Itt)
+		if (UWidget* Widget = Itt.Key().Get())
 		{
-			if (UWidget* Widget = Itt.Key().Get())
-			{
-				if (InWorld == Widget->GetWorld())
-				{
-					Widget->RemoveFromParent();
-				}
-			}
-			else
+			if (InWorld == Widget->GetWorld())
 			{
-				Itt.RemoveCurrent();
+				WidgetsToRemove.Add(Widget);
 			}
 		}
+		else
+		{
+			Itt.RemoveCurrent();
+		}
+	}
+
+	for (UWidget* Widget : WidgetsToRemove)
+	{
+		Widget->RemoveFromParent();
 	}
 }
 
diff --git a/Engine/Source/Runtime/UMG/Public/Blueprint/GameViewportSubsystem.h b/Engine/Source/Runtime/UMG/Public/Blueprint/GameViewportSubsystem.h
--- a/Engine/Source/Runtime/UMG/Public/Blueprint/GameViewportSubsystem.h
+++ b/Engine/Source/Runtime/UMG/Public/Blueprint/GameViewportSubsystem.h
@@ -2,7 +2,7 @@
 
 #pragma once
 
-#include "Subsystems/GameInstanceSubsystem.h"
+#include "Subsystems/EngineSubsystem.h"
 #include "Layout/Margin.h"
 #include "Widgets/Layout/Anchors.h"
 #include "Widgets/Layout/SConstraintCanvas.h"
@@ -41,7 +41,7 @@ struct UMG_API FGameViewportWidgetSlot
  * 
  */
 UCLASS()
-class UMG_API UGameViewportSubsystem : public UGameInstanceSubsystem
+class UMG_API UGameViewportSubsystem : public UEngineSubsystem
 {
 	GENERATED_BODY()
 
@@ -51,6 +51,7 @@ public:
 	virtual void Deinitialize() override;
 	//~ End Subsystem
 
+	static UGameViewportSubsystem* Get();
 	static UGameViewportSubsystem* Get(UWorld* World);
 	
 public:
@@ -118,7 +119,9 @@ private:
 private:
 	void AddToScreen(UWidget* Widget, ULocalPlayer* Player, FGameViewportWidgetSlot& Slot);
 	void RemoveWidgetInternal(UWidget* Widget, const TWeakPtr<SConstraintCanvas>& FullScreenWidget, const TWeakObjectPtr<ULocalPlayer>& LocalPlayer);
-	void OnLevelRemovedFromWorld(ULevel* InLevel, UWorld* InWorld);
+	void HandleWorldCleanup(UWorld* InWorld, bool bSessionEnded, bool bCleanupResoures);
+	void HandleRemoveWorld(UWorld* InWorld);
+
 	FMargin GetFullScreenOffsetForWidget(UWidget* Widget) const;
 	TPair<FMargin, bool> GetOffsetAttribute(UWidget* Widget, const FGameViewportWidgetSlot& Slot) const;
 
