Added a IsItemSelectable check to SObjectTableRow::OnTouchEnded. Fixes mobile users being able to select elements that should be unselectable.

diff --git a/Engine/Source/Runtime/UMG/Public/Slate/SObjectTableRow.h b/Engine/Source/Runtime/UMG/Public/Slate/SObjectTableRow.h
--- a/Engine/Source/Runtime/UMG/Public/Slate/SObjectTableRow.h
+++ b/Engine/Source/Runtime/UMG/Public/Slate/SObjectTableRow.h
@@ -273,27 +273,30 @@ public:
 			TSharedRef<ITypedTableView<ItemType>> OwnerTable = OwnerTablePtr.Pin().ToSharedRef();
 			if (const ItemType* MyItemPtr = GetItemForThis(OwnerTable))
 			{
-				ESelectionMode::Type SelectionMode = GetSelectionMode();
-				if (SelectionMode != ESelectionMode::None)
+				if (IsItemSelectable())
 				{
-					const bool bIsSelected = OwnerTable->Private_IsItemSelected(*MyItemPtr);
-					if (!bIsSelected)
+					ESelectionMode::Type SelectionMode = GetSelectionMode();
+					if (SelectionMode != ESelectionMode::None)
 					{
-						if (SelectionMode != ESelectionMode::Multi)
+						const bool bIsSelected = OwnerTable->Private_IsItemSelected(*MyItemPtr);
+						if (!bIsSelected)
 						{
-							OwnerTable->Private_ClearSelection();
-						}
-						OwnerTable->Private_SetItemSelection(*MyItemPtr, true, true);
-						OwnerTable->Private_SignalSelectionChanged(ESelectInfo::OnMouseClick);
+							if (SelectionMode != ESelectionMode::Multi)
+							{
+								OwnerTable->Private_ClearSelection();
+							}
+							OwnerTable->Private_SetItemSelection(*MyItemPtr, true, true);
+							OwnerTable->Private_SignalSelectionChanged(ESelectInfo::OnMouseClick);
 
-						Reply = FReply::Handled();
-					}
-					else if (SelectionMode == ESelectionMode::SingleToggle || SelectionMode == ESelectionMode::Multi)
-					{
-						OwnerTable->Private_SetItemSelection(*MyItemPtr, true, true);
-						OwnerTable->Private_SignalSelectionChanged(ESelectInfo::OnMouseClick);
+							Reply = FReply::Handled();
+						}
+						else if (SelectionMode == ESelectionMode::SingleToggle || SelectionMode == ESelectionMode::Multi)
+						{
+							OwnerTable->Private_SetItemSelection(*MyItemPtr, true, true);
+							OwnerTable->Private_SignalSelectionChanged(ESelectInfo::OnMouseClick);
 
-						Reply = FReply::Handled();
+							Reply = FReply::Handled();
+						}
 					}
 				}
 
