UMG: Fix for drag and drop operation to not keep resources alive.

diff --git a/Engine/Source/Runtime/UMG/Private/Slate/UMGDragDropOp.cpp b/Engine/Source/Runtime/UMG/Private/Slate/UMGDragDropOp.cpp
--- a/Engine/Source/Runtime/UMG/Private/Slate/UMGDragDropOp.cpp
+++ b/Engine/Source/Runtime/UMG/Private/Slate/UMGDragDropOp.cpp
@@ -25,7 +25,6 @@ FUMGDragDropOp::FUMGDragDropOp()
 void FUMGDragDropOp::AddReferencedObjects(FReferenceCollector& Collector)
 {
 	Collector.AddReferencedObject(DragOperation);
-	Collector.AddReferencedObject(GameViewport);
 }
 
 FString FUMGDragDropOp::GetReferencerName() const
@@ -53,9 +52,9 @@ void FUMGDragDropOp::OnDrop( bool bDropWasHandled, const FPointerEvent& MouseEve
 		}
 		else
 		{
-			if ( SourceUserWidget.IsValid() )
+			if ( TSharedPtr<SObjectWidget> SourceUserWidgetPtr = SourceUserWidget.Pin() )
 			{
-				SourceUserWidget->OnDragCancelled(FDragDropEvent(MouseEvent, AsShared()), DragOperation);
+				SourceUserWidgetPtr->OnDragCancelled(FDragDropEvent(MouseEvent, AsShared()), DragOperation);
 			}
 
 			DragOperation->DragCancelled(MouseEvent);
@@ -145,12 +144,12 @@ FCursorReply FUMGDragDropOp::OnCursorQuery()
 		CursorReply = CursorReply.Cursor(EMouseCursor::Default);
 	}
 
-	if ( GameViewport )
+	if ( UGameViewportClient* GameViewportPtr = GameViewport.Get() )
 	{
-		TOptional<TSharedRef<SWidget>> CursorWidget = GameViewport->MapCursor(nullptr, CursorReply);
+		TOptional<TSharedRef<SWidget>> CursorWidget = GameViewportPtr->MapCursor(nullptr, CursorReply);
 		if ( CursorWidget.IsSet() )
 		{
-			CursorReply.SetCursorWidget(GameViewport->GetWindow(), CursorWidget.GetValue());
+			CursorReply.SetCursorWidget(GameViewportPtr->GetWindow(), CursorWidget.GetValue());
 		}
 	}
 
@@ -160,6 +159,7 @@ FCursorReply FUMGDragDropOp::OnCursorQuery()
 TSharedRef<FUMGDragDropOp> FUMGDragDropOp::New(UDragDropOperation* InOperation, const int32 PointerIndex, const FVector2D &PointerPosition, const FVector2D &ScreenPositionOfDragee, float DPIScale, TSharedPtr<SObjectWidget> SourceUserWidget)
 {
 	check(InOperation);
+	check(SourceUserWidget);
 
 	TSharedRef<FUMGDragDropOp> Operation = MakeShareable(new FUMGDragDropOp());
 	Operation->PointerIndex = PointerIndex;
diff --git a/Engine/Source/Runtime/UMG/Public/Slate/UMGDragDropOp.h b/Engine/Source/Runtime/UMG/Public/Slate/UMGDragDropOp.h
--- a/Engine/Source/Runtime/UMG/Public/Slate/UMGDragDropOp.h
+++ b/Engine/Source/Runtime/UMG/Public/Slate/UMGDragDropOp.h
@@ -47,10 +47,10 @@ private:
 	UDragDropOperation* DragOperation;
 
 	/** Source User Widget */
-	TSharedPtr<SObjectWidget> SourceUserWidget;
+	TWeakPtr<SObjectWidget> SourceUserWidget;
 
 	/** The viewport this drag/drop operation is associated with. */
-	UGameViewportClient* GameViewport;
+	TWeakObjectPtr<UGameViewportClient> GameViewport;
 
 	/** The widget used during the drag/drop action to show something being dragged. */
 	TSharedPtr<SWidget> DecoratorWidget;
