UMG: Fix GetExtension and GetNotifyField not calling the super class

diff --git a/Engine/Source/Runtime/UMG/Private/WidgetBlueprintGeneratedClass.cpp b/Engine/Source/Runtime/UMG/Private/WidgetBlueprintGeneratedClass.cpp
--- a/Engine/Source/Runtime/UMG/Private/WidgetBlueprintGeneratedClass.cpp
+++ b/Engine/Source/Runtime/UMG/Private/WidgetBlueprintGeneratedClass.cpp
@@ -505,7 +505,7 @@ void UWidgetBlueprintGeneratedClass::InitializeFieldNotification(const UUserWidg
 	}
 }
 
-void UWidgetBlueprintGeneratedClass::ForEachField(TFunctionRef<bool(::UE::FieldNotification::FFieldId FielId)> Callback) const
+void UWidgetBlueprintGeneratedClass::ForEachField(TFunctionRef<bool(::UE::FieldNotification::FFieldId FielId)> Callback, bool bIncludeSuper) const
 {
 	ensureMsgf(FieldNotifyStartBitNumber >= 0, TEXT("The FieldNotifyStartIndex is negative. The number of field should be positive."));
 	for (int32 Index = 0; Index < FieldNotifyNames.Num(); ++Index)
@@ -515,9 +515,16 @@ void UWidgetBlueprintGeneratedClass::ForEachField(TFunctionRef<bool(::UE::FieldN
 			break;
 		}
 	}
+	if (bIncludeSuper)
+	{
+		if (UWidgetBlueprintGeneratedClass* ParentClass = Cast<UWidgetBlueprintGeneratedClass>(GetSuperClass()))
+		{
+			ParentClass->ForEachField(Callback);
+		}
+	}
 }
 
-UWidgetBlueprintGeneratedClassExtension* UWidgetBlueprintGeneratedClass::GetExtension(TSubclassOf<UWidgetBlueprintGeneratedClassExtension> InExtensionType)
+UWidgetBlueprintGeneratedClassExtension* UWidgetBlueprintGeneratedClass::GetExtension(TSubclassOf<UWidgetBlueprintGeneratedClassExtension> InExtensionType, bool bIncludeSuper)
 {
 	for (UWidgetBlueprintGeneratedClassExtension* Extension : Extensions)
 	{
@@ -526,21 +533,39 @@ UWidgetBlueprintGeneratedClassExtension* UWidgetBlueprintGeneratedClass::GetExte
 			return Extension;
 		}
 	}
+	if (bIncludeSuper)
+	{
+		if (UWidgetBlueprintGeneratedClass* ParentClass = Cast<UWidgetBlueprintGeneratedClass>(GetSuperClass()))
+		{
+			return ParentClass->GetExtension(InExtensionType);
+		}
+	}
 	return nullptr;
 }
 
-
-TArray<UWidgetBlueprintGeneratedClassExtension*> UWidgetBlueprintGeneratedClass::GetExtensions(TSubclassOf<UWidgetBlueprintGeneratedClassExtension> InExtensionType)
+TArray<UWidgetBlueprintGeneratedClassExtension*> UWidgetBlueprintGeneratedClass::GetExtensions(TSubclassOf<UWidgetBlueprintGeneratedClassExtension> InExtensionType, bool bIncludeSuper)
 {
 	TArray<UWidgetBlueprintGeneratedClassExtension*> Result;
+	GetExtensions(Result, InExtensionType, bIncludeSuper);
+	return Result;
+}
+
+void UWidgetBlueprintGeneratedClass::GetExtensions(TArray<UWidgetBlueprintGeneratedClassExtension*>& OutExtensions, TSubclassOf<UWidgetBlueprintGeneratedClassExtension> InExtensionType, bool bIncludeSuper)
+{
 	for (UWidgetBlueprintGeneratedClassExtension* Extension : Extensions)
 	{
 		if (Extension->IsA(InExtensionType))
 		{
-			Result.Add(Extension);
+			OutExtensions.Add(Extension);
+		}
+	}
+	if (bIncludeSuper)
+	{
+		if (UWidgetBlueprintGeneratedClass* ParentClass = Cast<UWidgetBlueprintGeneratedClass>(GetSuperClass()))
+		{
+			ParentClass->GetExtensions(OutExtensions, InExtensionType, bIncludeSuper);
 		}
 	}
-	return Result;
 }
 
 #undef LOCTEXT_NAMESPACE
diff --git a/Engine/Source/Runtime/UMG/Public/Blueprint/WidgetBlueprintGeneratedClass.h b/Engine/Source/Runtime/UMG/Public/Blueprint/WidgetBlueprintGeneratedClass.h
--- a/Engine/Source/Runtime/UMG/Public/Blueprint/WidgetBlueprintGeneratedClass.h
+++ b/Engine/Source/Runtime/UMG/Public/Blueprint/WidgetBlueprintGeneratedClass.h
@@ -137,7 +137,7 @@ public:
 	UWidgetBlueprintGeneratedClass* FindWidgetTreeOwningClass() const;
 
 	// Execute the callback for every FieldId defined in the BP class
-	void ForEachField(TFunctionRef<bool(::UE::FieldNotification::FFieldId FielId)> Callback) const;
+	void ForEachField(TFunctionRef<bool(::UE::FieldNotification::FFieldId FielId)> Callback, bool bIncludeSuper = true) const;
 
 	//~ Begin UObject interface
 	virtual void Serialize(FArchive& Ar) override;
@@ -172,28 +172,37 @@ public:
 
 	/** Find the first extension of the requested type. */
 	template<typename ExtensionType>
-	ExtensionType* GetExtension()
+	ExtensionType* GetExtension(bool bIncludeSuper = true)
 	{
-		return Cast<ExtensionType>(GetExtension(ExtensionType::StaticClass()));
+		return Cast<ExtensionType>(GetExtension(ExtensionType::StaticClass(), bIncludeSuper));
 	}
 
 	/** Find the first extension of the requested type. */
-	UWidgetBlueprintGeneratedClassExtension* GetExtension(TSubclassOf<UWidgetBlueprintGeneratedClassExtension> InExtensionType);
+	UWidgetBlueprintGeneratedClassExtension* GetExtension(TSubclassOf<UWidgetBlueprintGeneratedClassExtension> InExtensionType, bool bIncludeSuper = true);
 
 	/** Find the extensions of the requested type. */
-	TArray<UWidgetBlueprintGeneratedClassExtension*> GetExtensions(TSubclassOf<UWidgetBlueprintGeneratedClassExtension> InExtensionType);
+	TArray<UWidgetBlueprintGeneratedClassExtension*> GetExtensions(TSubclassOf<UWidgetBlueprintGeneratedClassExtension> InExtensionType, bool bIncludeSuper = true);
 
 	template<typename Predicate>
-	void ForEachExtension(Predicate Pred) const
+	void ForEachExtension(Predicate Pred, bool bIncludeSuper = true) const
 	{
 		for (UWidgetBlueprintGeneratedClassExtension* Extension : Extensions)
 		{
 			check(Extension);
 			Pred(Extension);
 		}
+		if (bIncludeSuper)
+		{
+			if (UWidgetBlueprintGeneratedClass* ParentClass = Cast<UWidgetBlueprintGeneratedClass>(GetSuperClass()))
+			{
+				ParentClass->ForEachExtension(MoveTemp(Pred));
+			}
+		}
 	}
 
 private:
 	static void InitializeBindingsStatic(UUserWidget* UserWidget, const TArrayView<const FDelegateRuntimeBinding> InBindings, const TMap<FName, FObjectPropertyBase*>& InPropertyMap);
 	static void BindAnimationsStatic(UUserWidget* Instance, const TArrayView<UWidgetAnimation*> InAnimations, const TMap<FName, FObjectPropertyBase*>& InPropertyMap);
+
+	void GetExtensions(TArray<UWidgetBlueprintGeneratedClassExtension*>& OutExtensions, TSubclassOf<UWidgetBlueprintGeneratedClassExtension> InExtensionType, bool bIncludeSuper);
 };
