UMG: Do not return stopping players from UUserWidget::GetSequencePlayer

diff --git a/Engine/Source/Runtime/UMG/Private/UserWidget.cpp b/Engine/Source/Runtime/UMG/Private/UserWidget.cpp
--- a/Engine/Source/Runtime/UMG/Private/UserWidget.cpp
+++ b/Engine/Source/Runtime/UMG/Private/UserWidget.cpp
@@ -423,7 +423,7 @@ UUMGSequencePlayer* UUserWidget::GetSequencePlayer(const UWidgetAnimation* InAni
 	TObjectPtr<UUMGSequencePlayer> const* FoundPlayer = ActiveSequencePlayers.FindByPredicate(
 		[&](const UUMGSequencePlayer* Player)
 	{
-		return Player->GetAnimation() == InAnimation;
+		return Player->GetAnimation() == InAnimation && !Player->IsStopping();
 	});
 
 	return FoundPlayer ? *FoundPlayer : nullptr;
