Added a console variable to control whether the window title bar buttons should be forced hidden.

#rnx

diff --git a/Engine/Source/Runtime/UMG/Private/Components/WindowTitleBarArea.cpp b/Engine/Source/Runtime/UMG/Private/Components/WindowTitleBarArea.cpp
--- a/Engine/Source/Runtime/UMG/Private/Components/WindowTitleBarArea.cpp
+++ b/Engine/Source/Runtime/UMG/Private/Components/WindowTitleBarArea.cpp
@@ -12,6 +12,15 @@
 
 #include UE_INLINE_GENERATED_CPP_BY_NAME(WindowTitleBarArea)
 
+namespace UE::UMG::Private
+{
+	static TAutoConsoleVariable<bool> CVarForceWindowButtonsHidden(
+		TEXT("WindowTitleBar.ForceWindowButtonsHidden"),
+		false,
+		TEXT("If true, force the window title bar buttons to be hidden."),
+		ECVF_Default);
+}
+
 #define LOCTEXT_NAMESPACE "UMG"
 
 /////////////////////////////////////////////////////
@@ -67,7 +76,7 @@ PRAGMA_ENABLE_DEPRECATION_WARNINGS
 	}
 
 PRAGMA_DISABLE_DEPRECATION_WARNINGS
-	MyWindowTitleBarArea->SetWindowButtonsVisibility(bWindowButtonsEnabled);
+	MyWindowTitleBarArea->SetWindowButtonsVisibility(UE::UMG::Private::CVarForceWindowButtonsHidden.GetValueOnAnyThread() ? false : bWindowButtonsEnabled);
 PRAGMA_ENABLE_DEPRECATION_WARNINGS
 
 	return MyWindowTitleBarArea.ToSharedRef();
