UMG: Generate metadata in editor build.

diff --git a/Engine/Source/Runtime/UMG/Public/Components/Widget.h b/Engine/Source/Runtime/UMG/Public/Components/Widget.h
--- a/Engine/Source/Runtime/UMG/Public/Components/Widget.h
+++ b/Engine/Source/Runtime/UMG/Public/Components/Widget.h
@@ -25,7 +25,13 @@
 
 #include "Widget.generated.h"
 
-#define UE_HAS_WIDGET_GENERATED_BY_CLASS (!UE_BUILD_SHIPPING || WITH_EDITOR)
+#ifndef UE_HAS_WIDGET_GENERATED_BY_CLASS
+  #define UE_HAS_WIDGET_GENERATED_BY_CLASS (!UE_BUILD_SHIPPING || WITH_EDITOR)
+#endif
+
+#ifndef WIDGET_INCLUDE_RELFECTION_METADATA
+#define WIDGET_INCLUDE_RELFECTION_METADATA  (!UE_BUILD_SHIPPING || WITH_EDITOR)
+#endif
 
 class ULocalPlayer;
 class SObjectWidget;
@@ -200,12 +206,6 @@ public:
 #endif
 
 
-#ifndef WIDGET_INCLUDE_RELFECTION_METADATA
-	#define WIDGET_INCLUDE_RELFECTION_METADATA !UE_BUILD_SHIPPING
-#endif
-
-
-
 /**
  * This is the base class for all wrapped Slate controls that are exposed to UObjects.
  */
