UMG - Don't need the Set/GetPixelSnapping functions now that they're Setter/Getter meta data on the properties.

diff --git a/Engine/Source/Runtime/UMG/Public/Components/Widget.h b/Engine/Source/Runtime/UMG/Public/Components/Widget.h
--- a/Engine/Source/Runtime/UMG/Public/Components/Widget.h
+++ b/Engine/Source/Runtime/UMG/Public/Components/Widget.h
@@ -593,11 +593,9 @@ public:
 	void SetClipping(EWidgetClipping InClipping);
 	
 	/** Gets the pixel snapping method of this widget. */
-	UFUNCTION(BlueprintCallable, Category = "Widget")
 	EWidgetPixelSnapping GetPixelSnapping() const;
 
 	/** Sets the pixel snapping method of this widget. */
-	UFUNCTION(BlueprintCallable, Category = "Widget")
 	void SetPixelSnapping(EWidgetPixelSnapping InPixelSnapping);
 
 	/** Sets the forced volatility of the widget. */
