UMG: Change the description of UserWidget
#rnx

diff --git a/Engine/Source/Runtime/UMG/Public/Blueprint/UserWidget.h b/Engine/Source/Runtime/UMG/Public/Blueprint/UserWidget.h
--- a/Engine/Source/Runtime/UMG/Public/Blueprint/UserWidget.h
+++ b/Engine/Source/Runtime/UMG/Public/Blueprint/UserWidget.h
@@ -212,7 +212,7 @@ DECLARE_DYNAMIC_DELEGATE( FOnInputAction );
 DECLARE_DYNAMIC_MULTICAST_DELEGATE_OneParam(FOnVisibilityChangedEvent, ESlateVisibility, InVisibility);
 
 /**
- * The user widget is extensible by users through the WidgetBlueprint.
+ * A widget that enables UI extensibility through WidgetBlueprint.
  */
 UCLASS(Abstract, editinlinenew, BlueprintType, Blueprintable, meta=( DontUseGenericSpawnObject="True", DisableNativeTick) )
 class UMG_API UUserWidget : public UWidget, public INamedSlotInterface
