Slate: Change an error to warning in the retainer code.
#rnx

diff --git a/Engine/Source/Runtime/UMG/Private/Slate/SRetainerWidget.cpp b/Engine/Source/Runtime/UMG/Private/Slate/SRetainerWidget.cpp
--- a/Engine/Source/Runtime/UMG/Private/Slate/SRetainerWidget.cpp
+++ b/Engine/Source/Runtime/UMG/Private/Slate/SRetainerWidget.cpp
@@ -473,11 +473,11 @@ SRetainerWidget::EPaintRetainedContentResult SRetainerWidget::PaintRetainedConte
 				bInvalidSizeLogged = true;
 				if (bTextureIsTooLarge)
 				{
-					UE_LOG(LogUMG, Error, TEXT("The requested size for SRetainerWidget is too large. W:%i H:%i"), RenderTargetWidth, RenderTargetHeight);
+					UE_LOG(LogUMG, Warning, TEXT("The requested size for SRetainerWidget is too large. W:%i H:%i"), RenderTargetWidth, RenderTargetHeight);
 				}
 				else
 				{
-					UE_LOG(LogUMG, Error, TEXT("The requested size for SRetainerWidget is 0. W:%i H:%i"), RenderTargetWidth, RenderTargetHeight);
+					UE_LOG(LogUMG, Warning, TEXT("The requested size for SRetainerWidget is 0. W:%i H:%i"), RenderTargetWidth, RenderTargetHeight);
 				}
 			}
 			return bTextureIsTooLarge ? EPaintRetainedContentResult::TextureSizeTooBig : EPaintRetainedContentResult::TextureSizeZero;
