Fix ListView warning:
LogClass: Warning: Short type name "UserObjectListEntry" provided for TryFindType. Please convert it to a path name (suggested: "/Script/UMG.UserObjectListEntry").
#tests: Verified loading an instance of ListView did not cause UClass warning.

diff --git a/Engine/Source/Runtime/UMG/Public/Components/ListView.h b/Engine/Source/Runtime/UMG/Public/Components/ListView.h
--- a/Engine/Source/Runtime/UMG/Public/Components/ListView.h
+++ b/Engine/Source/Runtime/UMG/Public/Components/ListView.h
@@ -23,7 +23,7 @@ DECLARE_DYNAMIC_MULTICAST_DELEGATE_TwoParams(FOnListViewScrolledDynamic, float,
  *
  * To make a widget usable as an entry in a ListView, it must inherit from the IUserObjectListEntry interface.
  */
-UCLASS(meta = (EntryInterface = UserObjectListEntry))
+UCLASS(meta = (EntryInterface = "/Script/UMG.UserObjectListEntry"))
 class UMG_API UListView : public UListViewBase, public ITypedUMGListView<UObject*>
 {
 	GENERATED_BODY()
