TileView Features:
1. Separate vertical and horizontal spacing instead of having only one spacing value
2. Add an option for EntrySpacing to NOT affect an Item's EntryWidth and EntryHeight, when using alignment

diff --git a/Engine/Plugins/Experimental/CommonUI/Source/CommonUI/Private/CommonListView.cpp b/Engine/Plugins/Experimental/CommonUI/Source/CommonUI/Private/CommonListView.cpp
--- a/Engine/Plugins/Experimental/CommonUI/Source/CommonUI/Private/CommonListView.cpp
+++ b/Engine/Plugins/Experimental/CommonUI/Source/CommonUI/Private/CommonListView.cpp
@@ -19,7 +19,8 @@ UCommonListView::UCommonListView(const FObjectInitializer& ObjectInitializer)
 
 void UCommonListView::SetEntrySpacing(float InEntrySpacing)
 {
-	EntrySpacing = InEntrySpacing;
+	InitHorizontalEntrySpacing(InEntrySpacing);
+	InitVerticalEntrySpacing(InEntrySpacing);
 }
 
 #if WITH_EDITOR
diff --git a/Engine/Source/Runtime/UMG/Private/Components/ListView.cpp b/Engine/Source/Runtime/UMG/Private/Components/ListView.cpp
--- a/Engine/Source/Runtime/UMG/Private/Components/ListView.cpp
+++ b/Engine/Source/Runtime/UMG/Private/Components/ListView.cpp
@@ -330,12 +330,12 @@ FMargin UListView::GetDesiredEntryPadding(UObject* Item) const
 		if (Orientation == EOrientation::Orient_Horizontal)
 		{
 			// For all entries after the first one, add the spacing as left padding
-			return FMargin(EntrySpacing, 0.f, 0.0f, 0.f);
+			return FMargin(HorizontalEntrySpacing, 0.f, 0.0f, 0.f);
 		}
 		else
 		{
 			// For all entries after the first one, add the spacing as top padding
-			return FMargin(0.f, EntrySpacing, 0.f, 0.f);
+			return FMargin(0.f, VerticalEntrySpacing, 0.f, 0.f);
 		}
 	}
 
@@ -377,6 +377,35 @@ void UListView::OnListViewScrolledInternal(float ItemOffset, float DistanceRemai
 	BP_OnListViewScrolled.Broadcast(ItemOffset, DistanceRemaining);
 }
 
+void UListView::InitHorizontalEntrySpacing(float InHorizontalEntrySpacing)
+{
+	ensureMsgf(!MyListView.IsValid(), TEXT("The SWidget is already created."));
+	HorizontalEntrySpacing = InHorizontalEntrySpacing;
+}
+
+void UListView::InitVerticalEntrySpacing(float InVerticalEntrySpacing)
+{
+	ensureMsgf(!MyListView.IsValid(), TEXT("The SWidget is already created."));
+	VerticalEntrySpacing = InVerticalEntrySpacing;
+}
+
+void UListView::PostLoad()
+{
+	Super::PostLoad();
+
+#if WITH_EDITOR
+	PRAGMA_DISABLE_DEPRECATION_WARNINGS
+
+		if (EntrySpacing_DEPRECATED != 0.f)
+		{
+			HorizontalEntrySpacing = EntrySpacing_DEPRECATED;
+			VerticalEntrySpacing = EntrySpacing_DEPRECATED;
+			EntrySpacing_DEPRECATED = 0.f;
+		}
+
+	PRAGMA_ENABLE_DEPRECATION_WARNINGS
+#endif // WITH_EDITOR
+}
 /////////////////////////////////////////////////////
 
 #undef LOCTEXT_NAMESPACE
diff --git a/Engine/Source/Runtime/UMG/Private/Components/TileView.cpp b/Engine/Source/Runtime/UMG/Private/Components/TileView.cpp
--- a/Engine/Source/Runtime/UMG/Private/Components/TileView.cpp
+++ b/Engine/Source/Runtime/UMG/Private/Components/TileView.cpp
@@ -20,17 +20,31 @@ TSharedRef<STableViewBase> UTileView::RebuildListWidget()
 
 FMargin UTileView::GetDesiredEntryPadding(UObject* Item) const
 {
-	return FMargin(EntrySpacing * 0.5f);
+	return FMargin(GetHorizontalEntrySpacing() * 0.5f, GetVerticalEntrySpacing() * 0.5f);
 }
 
 float UTileView::GetTotalEntryHeight() const
 {
-	return EntryHeight + (EntrySpacing * 0.5f);
+	if (IsAligned() && !bEntrySizeIncludesEntrySpacing)
+	{
+		return EntryHeight + GetVerticalEntrySpacing();
+	}
+	else
+	{
+		return EntryHeight + GetVerticalEntrySpacing() * 0.5;
+	}
 }
 
 float UTileView::GetTotalEntryWidth() const
 {
-	return EntryWidth + (EntrySpacing * 0.5f);
+	if (IsAligned() && !bEntrySizeIncludesEntrySpacing)
+	{
+		return EntryWidth + GetHorizontalEntrySpacing();
+	}
+	else
+	{
+		return EntryWidth + GetHorizontalEntrySpacing() * 0.5;
+	}
 }
 
 void UTileView::SetEntryHeight(float NewHeight)
@@ -51,6 +65,13 @@ void UTileView::SetEntryWidth(float NewWidth)
 	}
 }
 
+bool UTileView::IsAligned() const
+{
+	return TileAlignment == EListItemAlignment::LeftAligned
+		|| TileAlignment == EListItemAlignment::RightAligned
+		|| TileAlignment == EListItemAlignment::CenterAligned;
+}
+
 void UTileView::ReleaseSlateResources(bool bReleaseChildren)
 {
 	Super::ReleaseSlateResources(bReleaseChildren);
diff --git a/Engine/Source/Runtime/UMG/Public/Components/ListView.h b/Engine/Source/Runtime/UMG/Public/Components/ListView.h
--- a/Engine/Source/Runtime/UMG/Public/Components/ListView.h
+++ b/Engine/Source/Runtime/UMG/Public/Components/ListView.h
@@ -34,7 +34,8 @@ public:
 	UListView(const FObjectInitializer& Initializer);
 
 	virtual void ReleaseSlateResources(bool bReleaseChildren) override;
-	
+	virtual void PostLoad() override;
+
 	/** Set the list of items to display within this listview */
 	template <typename ItemObjectT, typename AllocatorType = FDefaultAllocator>
 	void SetListItems(const TArray<ItemObjectT, AllocatorType>& InListItems)
@@ -197,9 +198,6 @@ protected:
 	UPROPERTY(EditAnywhere, BlueprintReadOnly, Category = ListView)
 	bool bIsFocusable = true;
 
-	UPROPERTY(EditAnywhere, BlueprintReadOnly, Category = ListView, meta = (ClampMin = 0))
-	float EntrySpacing = 0.f;
-
 	UPROPERTY(EditAnywhere, BlueprintReadOnly, Category = ListView)
 	bool bReturnFocusToSelection = false;
 
@@ -208,6 +206,35 @@ protected:
 
 	TSharedPtr<SListView<UObject*>> MyListView;
 
+#if WITH_EDITORONLY_DATA
+	/**
+	 * This deprecated property was originally BlueprintReadOnly. To satisfy the compiler requirment to have a BlueprintGetter for this property, 
+	 * it relies on the newly added UFunction GetHorizontalEntrySpacing() to act as its BlueprintGetter.
+	 */
+	UPROPERTY(meta = (DeprecatedProperty, DeprecationMessage = "EntrySpacing has been deprecated. Please use HorizontalEntrySpacing and VerticalEntrySpacing."), BlueprintReadOnly, BlueprintGetter = "GetHorizontalEntrySpacing", Category = ListView, meta = (ClampMin = 0))
+	float EntrySpacing_DEPRECATED = 0.f;
+#endif  // WITH_EDITOR
+
+private:
+	UPROPERTY(EditAnywhere, BlueprintReadOnly, Getter, Category = ListView, meta = (ClampMin = 0, AllowPrivateAccess = "true"))
+	float HorizontalEntrySpacing = 0.f;
+
+	UPROPERTY(EditAnywhere, BlueprintReadOnly, Getter, Category = ListView, meta = (ClampMin = 0, AllowPrivateAccess = "true"))
+	float VerticalEntrySpacing = 0.f;
+
+public:
+	/** Get the horizontal spacing between entries. */
+	UFUNCTION(BlueprintCallable, Category = ListView)
+	float GetHorizontalEntrySpacing() const { return HorizontalEntrySpacing; }
+
+	/** Get the vertical spacing between entries. */
+	UFUNCTION(BlueprintCallable, Category = ListView)
+	float GetVerticalEntrySpacing() const { return VerticalEntrySpacing; }
+
+protected:
+	void InitHorizontalEntrySpacing(float InHorizontalEntrySpacing);
+	void InitVerticalEntrySpacing(float InVerticalEntrySpacing);
+
 private:
 	// BP exposure of ITypedUMGListView API
 
diff --git a/Engine/Source/Runtime/UMG/Public/Components/TileView.h b/Engine/Source/Runtime/UMG/Public/Components/TileView.h
--- a/Engine/Source/Runtime/UMG/Public/Components/TileView.h
+++ b/Engine/Source/Runtime/UMG/Public/Components/TileView.h
@@ -55,17 +55,35 @@ protected:
 		Args.ConsumeMouseWheel = ConsumeMouseWheel;
 		Args.bReturnFocusToSelection = bReturnFocusToSelection;
 		Args.TileAlignment = TileAlignment;
-		Args.EntryHeight = EntryHeight;
-		Args.EntryWidth = EntryWidth;
 		Args.bWrapDirectionalNavigation = bWrapHorizontalNavigation;
 		Args.Orientation = Orientation;
 		Args.ScrollBarStyle = &ScrollBarStyle;
 
+		if (IsAligned() && !bEntrySizeIncludesEntrySpacing)
+		{
+			// This tells the underlying ListView to expect entries with the final width/height equal to EntryWidth/EntryHeight + spacing
+			// It allows the entry widget to always occupy the entire EntryWidth/EntryHeight.
+			Args.EntryWidth = GetHorizontalEntrySpacing();
+			Args.EntryHeight = GetVerticalEntrySpacing();
+		}
+		else
+		{
+			// This tells the underlying ListView to expect entries with the final width/height equal to EntryWidth/EntryHeight.
+			// It forces the entry widget size to adjust so that the summation of entry widget width/height and spacing does not exceed EntryWidth/EntryHeight.
+			Args.EntryWidth = EntryWidth;
+			Args.EntryHeight = EntryHeight;
+		}
+
 		MyListView = MyTileView = ITypedUMGListView<UObject*>::ConstructTileView<TileViewT>(this, ListItems, Args);
 		MyTileView->SetOnEntryInitialized(SListView<UObject*>::FOnEntryInitialized::CreateUObject(this, &UTileView::HandleOnEntryInitializedInternal));
 		return StaticCastSharedRef<TileViewT<UObject*>>(MyTileView.ToSharedRef());
 	}
 
+private:
+	/** Returns whether the TileView is left, right or center aligned. */
+	UFUNCTION()
+	bool IsAligned() const;
+
 protected:
 	/** The height of each tile */
 	UPROPERTY(EditAnywhere, Category = ListEntries)
@@ -84,4 +102,12 @@ protected:
 	bool bWrapHorizontalNavigation = false;
 
 	TSharedPtr<STileView<UObject*>> MyTileView;
+
+private:
+	/**
+	 * True if entry dimensions should be the sum of the entry widget dimensions and the spacing.
+	 * This means the size of the entry widget will be adjusted so that the summation of the widget size and entry spacing always equals entry size.
+	 */
+	UPROPERTY(EditAnywhere, Category = ListEntries, meta = (EditCondition = "IsAligned"))
+	bool bEntrySizeIncludesEntrySpacing = true;
 };
