UMG: Add class icon for StackBox

diff --git a/Engine/Source/Runtime/UMG/Private/UMGStyle.cpp b/Engine/Source/Runtime/UMG/Private/UMGStyle.cpp
--- a/Engine/Source/Runtime/UMG/Private/UMGStyle.cpp
+++ b/Engine/Source/Runtime/UMG/Private/UMGStyle.cpp
@@ -85,6 +85,7 @@ TSharedRef< FSlateStyleSet > FUMGStyle::Create()
 	Style->Set("ClassIcon.NativeWidgetHost", new IMAGE_BRUSH_SVG(TEXT("NativeWidgetHost"), Icon16x16));
 	Style->Set("ClassIcon.ScaleBox", new IMAGE_BRUSH_SVG(TEXT("ScaleBox"), Icon16x16));
 	Style->Set("ClassIcon.Sizebox", new IMAGE_BRUSH_SVG(TEXT("Sizebox"), Icon16x16));
+	Style->Set("ClassIcon.StackBox", new IMAGE_BRUSH_SVG(TEXT("StackBox"), Icon16x16));
 	Style->Set("ClassIcon.SpinBox", new IMAGE_BRUSH_SVG(TEXT("SpinBox"), Icon16x16));
 	Style->Set("ClassIcon.GridPanel", new IMAGE_BRUSH_SVG(TEXT("GridPanel"), Icon16x16));
 	Style->Set("ClassIcon.WrapBox", new IMAGE_BRUSH_SVG(TEXT("WrapBox"), Icon16x16));
@@ -145,6 +146,7 @@ TSharedRef< FSlateStyleSet > FUMGStyle::Create()
 	Style->Set("ClassThumbnail.Sizebox", new IMAGE_BRUSH_SVG(TEXT("Sizebox_64"), Icon64x64));
 	Style->Set("ClassThumbnail.Slider", new IMAGE_BRUSH_SVG(TEXT("Slider_64"), Icon64x64));
 	Style->Set("ClassThumbnail.SpinBox", new IMAGE_BRUSH_SVG(TEXT("SpinBox_64"), Icon64x64));
+	Style->Set("ClassThumbnail.StackBox", new IMAGE_BRUSH_SVG(TEXT("StackBox"), Icon64x64));
 	Style->Set("ClassThumbnail.MultilineEditableTextBox", new IMAGE_BRUSH_SVG(TEXT("TextBoxMultiline_64"), Icon64x64));
 	Style->Set("ClassThumbnail.EditableTextBox", new IMAGE_BRUSH_SVG(TEXT("TextBox_64"), Icon64x64));
 	Style->Set("ClassThumbnail.TextBlock", new IMAGE_BRUSH_SVG(TEXT("Text_64"), Icon64x64));
