Remove event from the CTAButton viewmodel. Change the category of StackBox Orientation property to be more visible in the DetailView.

diff --git a/Engine/Source/Runtime/UMG/Public/Components/StackBox.h b/Engine/Source/Runtime/UMG/Public/Components/StackBox.h
--- a/Engine/Source/Runtime/UMG/Public/Components/StackBox.h
+++ b/Engine/Source/Runtime/UMG/Public/Components/StackBox.h
@@ -26,7 +26,7 @@ class UMG_API UStackBox : public UPanelWidget
 
 private:
 	/** The orientation of the stack box. */
-	UPROPERTY(EditAnywhere, BlueprintReadWrite, Setter, Getter, Category = "Layout", meta = (AllowPrivateAccess = "true"))
+	UPROPERTY(EditAnywhere, BlueprintReadWrite, Setter, Getter, Category = "Behavior", meta = (AllowPrivateAccess = "true"))
 	TEnumAsByte<EOrientation> Orientation;
 
 public:
