Design View Widget: set desired size of selection border to 0 to prevent issue with designer view representing the layout of 0 sized widgets incorrectly

diff --git a/Engine/Source/Runtime/UMG/Private/Components/Widget.cpp b/Engine/Source/Runtime/UMG/Private/Components/Widget.cpp
--- a/Engine/Source/Runtime/UMG/Private/Components/Widget.cpp
+++ b/Engine/Source/Runtime/UMG/Private/Components/Widget.cpp
@@ -1069,6 +1069,7 @@ TSharedRef<SWidget> UWidget::CreateDesignerOutline(TSharedRef<SWidget> Content)
 		.VAlign(VAlign_Fill)
 		[
 			SNew(SBorder)
+			.DesiredSizeScale(FVector2D(0.0f, 0.0f))
 			.Visibility(HasAnyDesignerFlags(EWidgetDesignFlags::ShowOutline) ? EVisibility::HitTestInvisible : EVisibility::Collapsed)
 			.BorderImage(FUMGStyle::Get().GetBrush("MarchingAnts"))
 		];
