Fix retainer not rendering in editor

diff --git a/Engine/Source/Developer/SlateReflector/Private/Widgets/SSlateOptions.cpp b/Engine/Source/Developer/SlateReflector/Private/Widgets/SSlateOptions.cpp
--- a/Engine/Source/Developer/SlateReflector/Private/Widgets/SSlateOptions.cpp
+++ b/Engine/Source/Developer/SlateReflector/Private/Widgets/SSlateOptions.cpp
@@ -84,6 +84,7 @@ void SSlateOptions::Construct( const FArguments& InArgs )
 			AddMenuEntry(MenuBuilder, Icon, LOCTEXT("ShowClipping", "Show Clipping"), TEXT("Slate.ShowClipping"));
 			AddMenuEntry(MenuBuilder, Icon, LOCTEXT("DebugCulling", "Debug Culling"), TEXT("Slate.DebugCulling"));
 			AddMenuEntry(MenuBuilder, Icon, LOCTEXT("ShowHitTestGrid", "Show HitTestGrid"), TEXT("Slate.HitTestGridDebugging"));
+			AddMenuEntry(MenuBuilder, Icon, LOCTEXT("DesignerRetainedRendering", "Designer Retained Rendering"), TEXT("Slate.EnableDesignerRetainedRendering"));
 #endif // WITH_SLATE_DEBUGGING
 			AddMenuEntry(MenuBuilder, Icon, LOCTEXT("ShowBatching", "Show Batching"), TEXT("Slate.ShowBatching"));
 			AddMenuEntry(MenuBuilder, Icon, LOCTEXT("ShowOverdraw", "Show Overdraw"), TEXT("Slate.ShowOverdraw"));
diff --git a/Engine/Source/Runtime/UMG/Private/Components/RetainerBox.cpp b/Engine/Source/Runtime/UMG/Private/Components/RetainerBox.cpp
--- a/Engine/Source/Runtime/UMG/Private/Components/RetainerBox.cpp
+++ b/Engine/Source/Runtime/UMG/Private/Components/RetainerBox.cpp
@@ -25,6 +25,9 @@ PRAGMA_DISABLE_DEPRECATION_WARNINGS
 	RenderOnPhase = true;
 	RenderOnInvalidation = false;
 	TextureParameter = DefaultTextureParameterName;
+#if WITH_EDITOR
+	bShowEffectsInDesigner = true;
+#endif // WITH_EDITOR
 PRAGMA_ENABLE_DEPRECATION_WARNINGS
 }
 
@@ -151,6 +154,12 @@ PRAGMA_DISABLE_DEPRECATION_WARNINGS
 		.StatId( FName( *FString::Printf(TEXT("%s [%s]"), *GetFName().ToString(), *GetClass()->GetName() ) ) )
 #endif//STATS
 	;
+
+#if WITH_EDITOR
+	MyRetainerWidget->SetIsDesignTime(IsDesignTime());
+	MyRetainerWidget->SetShowEffectsInDesigner(bShowEffectsInDesigner);
+#endif // WITH_EDITOR
+
 PRAGMA_ENABLE_DEPRECATION_WARNINGS
 	if ( GetChildrenCount() > 0 )
 	{
@@ -170,7 +179,11 @@ void URetainerBox::SynchronizeProperties()
 	}
 	
 PRAGMA_DISABLE_DEPRECATION_WARNINGS
-	MyRetainerWidget->SetRetainedRendering(IsDesignTime() ? false : bRetainRender);
+#if WITH_EDITOR
+	MyRetainerWidget->SetIsDesignTime(IsDesignTime());
+	MyRetainerWidget->SetShowEffectsInDesigner(bShowEffectsInDesigner);
+#endif // WITH_EDITOR
+	MyRetainerWidget->SetRetainedRendering(bRetainRender);
 	MyRetainerWidget->SetEffectMaterial(EffectMaterial);
 	MyRetainerWidget->SetTextureParameter(TextureParameter);
 	MyRetainerWidget->SetWorld(GetWorld());
diff --git a/Engine/Source/Runtime/UMG/Private/Slate/SRetainerWidget.cpp b/Engine/Source/Runtime/UMG/Private/Slate/SRetainerWidget.cpp
--- a/Engine/Source/Runtime/UMG/Private/Slate/SRetainerWidget.cpp
+++ b/Engine/Source/Runtime/UMG/Private/Slate/SRetainerWidget.cpp
@@ -19,6 +19,20 @@ DECLARE_CYCLE_STAT(TEXT("Retainer Widget Paint"), STAT_SlateRetainerWidgetPaint,
 FOnRetainedModeChanged SRetainerWidget::OnRetainerModeChangedDelegate;
 #endif
 
+static void HandleDesignerRetainedRenderingToggled(IConsoleVariable* CVar)
+{
+	FSlateApplication::Get().InvalidateAllWidgets(false);
+}
+
+/** True if we should do retained rendering in designer. */
+bool GEnableDesignerRetainedRendering = true;
+FAutoConsoleVariableRef EnableDesignerRetainedRendering(
+	TEXT("Slate.EnableDesignerRetainedRendering"),
+	GEnableDesignerRetainedRendering,
+	TEXT("Controls if retainer renders in designer; 0 - Never; 1 - Per Widget Properties"),
+	FConsoleVariableDelegate::CreateStatic(&HandleDesignerRetainedRenderingToggled)
+);
+
 /** True if we should allow widgets to be cached in the UI at all. */
 int32 GEnableRetainedRendering = 1;
 FAutoConsoleVariableRef EnableRetainedRendering(
@@ -200,6 +214,11 @@ void SRetainerWidget::Construct(const FArguments& InArgs)
 	bEnableRenderWithLocalTransform = InArgs._RenderWithLocalTransform;
 	SetVolatilePrepass(bEnableRetainedRendering);
 
+#if WITH_EDITOR
+	bIsDesignTime = false;
+	bShowEffectsInDesigner = true;
+#endif // WITH_EDITOR
+
 	RefreshRenderingMode();
 	bRenderRequested = true;
 	bInvalidSizeLogged = false;
@@ -369,6 +388,18 @@ void SRetainerWidget::RequestRender()
 	InvalidateRootChildOrder();
 }
 
+#if WITH_EDITOR
+void SRetainerWidget::SetIsDesignTime(bool bInIsDesignTime)
+{
+	bIsDesignTime = bInIsDesignTime;
+}
+
+void SRetainerWidget::SetShowEffectsInDesigner(bool bInShowEffectsInDesigner)
+{
+	bShowEffectsInDesigner = bInShowEffectsInDesigner;
+}
+#endif // WITH_EDITOR
+
 SRetainerWidget::EPaintRetainedContentResult SRetainerWidget::PaintRetainedContentImpl(const FSlateInvalidationContext& Context, const FGeometry& AllottedGeometry, int32 LayerId)
 {
 	if (RenderOnPhase)
@@ -540,7 +571,14 @@ int32 SRetainerWidget::OnPaint(const FPaintArgs& Args, const FGeometry& Allotted
 
 	SRetainerWidget* MutableThis = const_cast<SRetainerWidget*>(this);
 
-	if (bEnableRetainedRendering && IsAnythingVisibleToRender())
+	bool bShouldRetainRendering = bEnableRetainedRendering;
+
+#if WITH_EDITOR
+	bool bShouldSkipDesignerRendering = bIsDesignTime && (!bShowEffectsInDesigner || !GEnableDesignerRetainedRendering);
+	bShouldRetainRendering = bEnableRetainedRendering && !bShouldSkipDesignerRendering;
+#endif // WITH_EDITOR
+
+	if (bShouldRetainRendering && IsAnythingVisibleToRender())
 	{
 		SCOPE_CYCLE_COUNTER(STAT_SlateRetainerWidgetPaint);
 
diff --git a/Engine/Source/Runtime/UMG/Public/Components/RetainerBox.h b/Engine/Source/Runtime/UMG/Public/Components/RetainerBox.h
--- a/Engine/Source/Runtime/UMG/Public/Components/RetainerBox.h
+++ b/Engine/Source/Runtime/UMG/Public/Components/RetainerBox.h
@@ -163,6 +163,14 @@ protected:
 	UPROPERTY(EditAnywhere, BlueprintReadWrite, Getter, Setter, BlueprintSetter = "SetTextureParameter", Category="Effect")
 	FName TextureParameter;
 
+#if WITH_EDITORONLY_DATA
+	/**
+	 * If true, retained rendering occurs in designer
+	 */
+	UPROPERTY(EditAnywhere, BlueprintReadOnly, Category="Effect", meta = (EditCondition = bRetainRender))
+	bool bShowEffectsInDesigner;
+#endif
+
 	//~ Begin UPanelWidget interface
 	virtual void OnSlotAdded(UPanelSlot* Slot) override;
 	virtual void OnSlotRemoved(UPanelSlot* Slot) override;
diff --git a/Engine/Source/Runtime/UMG/Public/Slate/SRetainerWidget.h b/Engine/Source/Runtime/UMG/Public/Slate/SRetainerWidget.h
--- a/Engine/Source/Runtime/UMG/Public/Slate/SRetainerWidget.h
+++ b/Engine/Source/Runtime/UMG/Public/Slate/SRetainerWidget.h
@@ -71,6 +71,11 @@ public:
 	/** Requests that the retainer redraw the hosted content next time it's painted. */
 	void RequestRender();
 
+#if WITH_EDITOR
+	void SetIsDesignTime(bool bInIsDesignTime);
+	void SetShowEffectsInDesigner(bool bInShowEffectsInDesigner);
+#endif
+
 	void SetRetainedRendering(bool bRetainRendering);
 
 	void SetContent(const TSharedRef< SWidget >& InContent);
@@ -148,6 +153,14 @@ private:
 	bool bEnableRetainedRendering;
 	bool bEnableRenderWithLocalTransform;
 
+#if WITH_EDITOR
+	/** True if widget is used in design time */
+	bool bIsDesignTime;
+
+	/** True if we should retain rendering in designer */
+	bool bShowEffectsInDesigner;
+#endif
+
 	bool RenderOnPhase;
 	bool RenderOnInvalidation;
 
