Initial addition of widget state, intended to be used by external systems to generically respond to widget changes regardless of type

diff --git a/Engine/Plugins/Experimental/CommonUI/Source/CommonUI/Private/CommonButtonBase.cpp b/Engine/Plugins/Experimental/CommonUI/Source/CommonUI/Private/CommonButtonBase.cpp
--- a/Engine/Plugins/Experimental/CommonUI/Source/CommonUI/Private/CommonButtonBase.cpp
+++ b/Engine/Plugins/Experimental/CommonUI/Source/CommonUI/Private/CommonButtonBase.cpp
@@ -2,6 +2,8 @@
 
 #include "CommonButtonBase.h"
 
+#include "Binding/States/WidgetStateBitfield.h"
+#include "Binding/States/WidgetStateRegistration.h"
 #include "CommonActionWidget.h"
 #include "CommonUISubsystemBase.h"
 #include "CommonInputSubsystem.h"
@@ -458,17 +460,27 @@ void UCommonButtonBase::NativeDestruct()
 
 void UCommonButtonBase::SetIsEnabled(bool bInIsEnabled)
 {
+	bool bValueChanged = bButtonEnabled != bInIsEnabled;
+
 	if (bInIsEnabled)
 	{
+		TGuardValue<bool> StateBroadcastGuard(bShouldBroadcastState, false);
 		Super::SetIsEnabled(bInIsEnabled);
 		EnableButton();
 	}
 	else
 	{
 		// Change the underlying enabled bool but do not call the case because we don't want to propogate it to the underlying SWidget
+		TGuardValue<bool> StateBroadcastGuard(bShouldBroadcastState, false);
 		Super::SetIsEnabled(bInIsEnabled);
 		DisableButton();
 	}
+
+	if (bValueChanged)
+	{
+		// Note: State is disabled, so we broadcast !bIsEnabled
+		BroadcastBinaryPostStateChange(UWidgetDisabledStateRegistration::Bit, !bInIsEnabled);
+	}
 }
 
 bool UCommonButtonBase::NativeIsInteractable() const
@@ -778,15 +790,24 @@ void UCommonButtonBase::SetIsSelected(bool InSelected, bool bGiveClickFeedback)
 
 void UCommonButtonBase::SetIsLocked(bool bInIsLocked)
 {
-	bLocked = bInIsLocked;
+	bool bValueChanged = bInIsLocked != bLocked;
 
-	SetButtonStyle();
+	if (bValueChanged)
+	{
+		bLocked = bInIsLocked;
 
-	BP_OnLockedChanged(bLocked);
+		SetButtonStyle();
+
+		BP_OnLockedChanged(bLocked);
+
+		BroadcastBinaryPostStateChange(UWidgetLockedStateRegistration::Bit, bLocked);
+	}
 }
 
 void UCommonButtonBase::SetSelectedInternal(bool bInSelected, bool bAllowSound /*= true*/, bool bBroadcast /*= true*/)
 {
+	bool bValueChanged = bInSelected != bSelected;
+
 	bSelected = bInSelected;
 
 	SetButtonStyle();
@@ -816,6 +837,11 @@ void UCommonButtonBase::SetSelectedInternal(bool bInSelected, bool bAllowSound /
 	}
 
 	UpdateInputActionWidgetVisibility();
+
+	if (bValueChanged)
+	{
+		BroadcastBinaryPostStateChange(UWidgetSelectedStateRegistration::Bit, bSelected);
+	}
 }
 
 void UCommonButtonBase::RefreshDimensions()
@@ -1641,3 +1667,22 @@ bool UCommonButtonBase::GetIsFocusable() const
 	return IsFocusable();
 }
 
+FName UWidgetLockedStateRegistration::GetStateName() const
+{
+	return StateName;
+};
+
+bool UWidgetLockedStateRegistration::GetRegisteredWidgetState(const UWidget* InWidget) const
+{
+	if (const UCommonButtonBase* CommonButton = Cast<UCommonButtonBase>(InWidget))
+	{
+		return CommonButton->GetLocked();
+	}
+
+	return false;
+}
+
+void UWidgetLockedStateRegistration::InitializeStaticBitfields() const
+{
+	Bit = FWidgetStateBitfield(GetStateName());
+}
\ No newline at end of file
diff --git a/Engine/Plugins/Experimental/CommonUI/Source/CommonUI/Public/CommonButtonBase.h b/Engine/Plugins/Experimental/CommonUI/Source/CommonUI/Public/CommonButtonBase.h
--- a/Engine/Plugins/Experimental/CommonUI/Source/CommonUI/Public/CommonButtonBase.h
+++ b/Engine/Plugins/Experimental/CommonUI/Source/CommonUI/Public/CommonButtonBase.h
@@ -2,6 +2,7 @@
 
 #pragma once
 
+#include "Binding/States/WidgetStateRegistration.h"
 #include "CommonUserWidget.h"
 #include "Components/Button.h"
 #include "Engine/DataTable.h"
@@ -900,6 +901,31 @@ protected:
 	TObjectPtr<UCommonActionWidget> InputActionWidget;
 };
 
+UCLASS(Transient)
+class COMMONUI_API UWidgetLockedStateRegistration : public UWidgetBinaryStateRegistration
+{
+	GENERATED_BODY()
+
+public:
+
+	/** Post-load initialized bit corresponding to this binary state */
+	static inline FWidgetStateBitfield Bit;
+
+	static const inline FName StateName = FName("Locked");
+
+	//~ Begin UWidgetBinaryStateRegistration Interface.
+	virtual FName GetStateName() const override;
+	virtual bool GetRegisteredWidgetState(const UWidget* InWidget) const override;
+	//~ End UWidgetBinaryStateRegistration Interface
+
+protected:
+	friend UWidgetStateSettings;
+
+	//~ Begin UWidgetBinaryStateRegistration Interface.
+	virtual void InitializeStaticBitfields() const override;
+	//~ End UWidgetBinaryStateRegistration Interface
+};
+
 #if UE_ENABLE_INCLUDE_ORDER_DEPRECATED_IN_5_2
 #include "CommonUITypes.h"
 #include "Misc/Optional.h"
diff --git a/Engine/Source/Runtime/UMG/Private/Binding/States/Tests/WidgetStateBitfieldTests.cpp b/Engine/Source/Runtime/UMG/Private/Binding/States/Tests/WidgetStateBitfieldTests.cpp
new file mode 100644
--- /dev/null
+++ b/Engine/Source/Runtime/UMG/Private/Binding/States/Tests/WidgetStateBitfieldTests.cpp
@@ -0,0 +1,332 @@
+// Copyright Epic Games, Inc. All Rights Reserved.
+
+#include "Binding/States/WidgetStateSettings.h"
+#include "Binding/States/WidgetStateBitfield.h"
+#include "Binding/States/WidgetStateRegistration.h"
+#include "Components/CheckBox.h"
+
+/**
+ * Coverage:
+ *
+ * OperatorBool
+ * OperatorBitwiseAnd (Intersection)
+ * OperatorBitwiseOr (Union)
+ * OperatorBitwiseNot
+ * HasFlag
+ *
+ * Note: Checkbox include is only needed if you want to use 'UWidgetCheckedStateRegistration' pre-defined
+ * bitfields. It is possible to just re-create these yourself if you don't want to include checkbox.
+ * Or more generally, a specific module that you may / may not know exists.
+ */
+
+#if WITH_DEV_AUTOMATION_TESTS
+
+DEFINE_LOG_CATEGORY_STATIC(LogWidgetStateBitfieldTest, Log, All);
+
+IMPLEMENT_SIMPLE_AUTOMATION_TEST(FOperatorBoolTest, "Slate.WidgetState.OperatorBoolTest", EAutomationTestFlags::ApplicationContextMask | EAutomationTestFlags::EngineFilter)
+bool FOperatorBoolTest::RunTest(const FString& Parameters)
+{
+	bool GPassing = true;
+
+	{
+		FWidgetStateBitfield Test = {};
+		GPassing &= TestEqual(TEXT("Unset is false"), Test, false);
+	}
+
+	{
+		FWidgetStateBitfield Test = {};
+
+		Test.SetBinaryStateSlow(FName("Pressed"), true);
+		GPassing &= TestEqual(TEXT("Any binary is true"), Test, true);
+
+		GPassing &= TestEqual(TEXT("Has binary is true"), Test.HasBinaryStates(), true);
+
+		Test.SetBinaryStateSlow(FName("Pressed"), false);
+		GPassing &= TestEqual(TEXT("No binary is false"), Test, false);
+
+		GPassing &= TestEqual(TEXT("Has binary is false"), Test.HasBinaryStates(), false);
+	}
+
+	{
+		FWidgetStateBitfield Test = {};
+
+		Test.SetEnumStateSlow(FName("CheckedState"), 1);
+		GPassing &= TestEqual(TEXT("Any enum is true"), Test, true);
+
+		GPassing &= TestEqual(TEXT("Has enum is true"), Test.HasEnumStates(), true);
+		GPassing &= TestEqual(TEXT("Has empty unused is false"), Test.HasEmptyUsedEnumStates(), false);
+
+		Test.ClearEnumState(FName("CheckedState"));
+		GPassing &= TestEqual(TEXT("No enum is false"), Test, false);
+
+		GPassing &= TestEqual(TEXT("Has enum is false"), Test.HasEnumStates(), false);
+		GPassing &= TestEqual(TEXT("Has empty unused is false - none in use"), Test.HasEmptyUsedEnumStates(), false);
+	}
+
+	return GPassing;
+}
+
+IMPLEMENT_SIMPLE_AUTOMATION_TEST(FOperatorBitwiseAndTest, "Slate.WidgetState.OperatorBitwiseAndTest", EAutomationTestFlags::ApplicationContextMask | EAutomationTestFlags::EngineFilter)
+bool FOperatorBitwiseAndTest::RunTest(const FString& Parameters)
+{
+	bool GPassing = true;
+
+	{
+		FWidgetStateBitfield Pressed = {};
+		FWidgetStateBitfield Hovered = {};
+		FWidgetStateBitfield PressedHovered = {};
+
+		Pressed.SetBinaryState(UWidgetPressedStateRegistration::Bit, true);
+		Hovered.SetBinaryState(UWidgetHoveredStateRegistration::Bit, true);
+		PressedHovered.SetBinaryState(UWidgetPressedStateRegistration::Bit, true);
+		PressedHovered.SetBinaryState(UWidgetHoveredStateRegistration::Bit, true);
+
+		GPassing &= TestEqual(TEXT("Pressed & Hovered is false"), Pressed.Intersect(Hovered), false);
+		GPassing &= TestEqual(TEXT("Hovered & Pressed is false"), Hovered.Intersect(Pressed), false);
+
+		GPassing &= TestEqual(TEXT("PressedHovered & Pressed is true"), PressedHovered.Intersect(Pressed), true);
+		GPassing &= TestEqual(TEXT("Pressed & PressedHovered is true"), Pressed.Intersect(PressedHovered), true);
+
+		GPassing &= TestEqual(TEXT("PressedHovered & Hovered is true"), PressedHovered.Intersect(Hovered), true);
+		GPassing &= TestEqual(TEXT("Hovered & PressedHovered is true"), Hovered.Intersect(PressedHovered), true);
+
+		GPassing &= TestEqual(TEXT("PressedHovered & PressedHovered is true"), PressedHovered.Intersect(PressedHovered), true);
+		GPassing &= TestEqual(TEXT("PressedHovered & PressedHovered is true"), PressedHovered.Intersect(PressedHovered), true);
+
+		GPassing &= TestEqual(TEXT("Bit: Pressed & Hovered is false"), UWidgetPressedStateRegistration::Bit.Intersect(UWidgetHoveredStateRegistration::Bit), false);
+		GPassing &= TestEqual(TEXT("Bit: Hovered & Pressed is false"), UWidgetHoveredStateRegistration::Bit.Intersect(UWidgetPressedStateRegistration::Bit), false);
+
+		GPassing &= TestEqual(TEXT("Bit: PressedHovered & Pressed is true"), PressedHovered.Intersect(UWidgetPressedStateRegistration::Bit), true);
+		GPassing &= TestEqual(TEXT("Bit: Pressed & PressedHovered is true"), UWidgetPressedStateRegistration::Bit.Intersect(PressedHovered), true);
+
+		GPassing &= TestEqual(TEXT("Bit: PressedHovered & Hovered is true"), PressedHovered.Intersect(UWidgetHoveredStateRegistration::Bit), true);
+		GPassing &= TestEqual(TEXT("Bit: Hovered & PressedHovered is true"), UWidgetHoveredStateRegistration::Bit.Intersect(PressedHovered), true);
+	}
+
+	{
+		FWidgetStateBitfield Checked = {};
+		FWidgetStateBitfield Unchecked = {};
+
+		Checked.SetEnumState(UWidgetCheckedStateRegistration::Checked);
+		Unchecked.SetEnumState(UWidgetCheckedStateRegistration::Unchecked);
+
+		GPassing &= TestEqual(TEXT("Checked & Unchecked is false"), Checked.Intersect(Unchecked), false);
+
+		GPassing &= TestEqual(TEXT("Bit: Checked & Unchecked is false"), UWidgetCheckedStateRegistration::Checked.Intersect(UWidgetCheckedStateRegistration::Unchecked), false);
+		GPassing &= TestEqual(TEXT("Bit: Checked & Unchecked is false"), Checked.Intersect(UWidgetCheckedStateRegistration::Unchecked), false);
+		GPassing &= TestEqual(TEXT("Bit: Checked & Unchecked is false"), UWidgetCheckedStateRegistration::Checked.Intersect(Unchecked), false);
+	}
+
+	{
+		FWidgetStateBitfield Checked = {};
+		FWidgetStateBitfield Hovered = {};
+		FWidgetStateBitfield CheckedHovered = {};
+
+		Checked.SetEnumState(UWidgetCheckedStateRegistration::Checked);
+		Hovered.SetBinaryState(UWidgetHoveredStateRegistration::Bit, true);
+
+		CheckedHovered.SetEnumState(UWidgetCheckedStateRegistration::Checked);
+		CheckedHovered.SetBinaryState(UWidgetHoveredStateRegistration::Bit, true);
+
+		GPassing &= TestEqual(TEXT("Checked & Hovered is false"), Checked.Intersect(Hovered), false);
+		GPassing &= TestEqual(TEXT("Checked & CheckedHovered is false"), Checked.Intersect(CheckedHovered), false);
+		GPassing &= TestEqual(TEXT("CheckedHovered & Checked is true"), CheckedHovered.Intersect(Checked), true);
+		GPassing &= TestEqual(TEXT("Hovered & CheckedHovered is false"), Hovered.Intersect(CheckedHovered), false);
+		GPassing &= TestEqual(TEXT("CheckedHovered & Hovered is true"), CheckedHovered.Intersect(Hovered), true);
+		GPassing &= TestEqual(TEXT("CheckedHovered & CheckedHovered is true"), CheckedHovered.Intersect(CheckedHovered), true);
+	}
+
+	return GPassing;
+}
+
+IMPLEMENT_SIMPLE_AUTOMATION_TEST(FOperatorBitwiseOrTest, "Slate.WidgetState.OperatorBitwiseOrTest", EAutomationTestFlags::ApplicationContextMask | EAutomationTestFlags::EngineFilter)
+bool FOperatorBitwiseOrTest::RunTest(const FString& Parameters)
+{
+	bool GPassing = true;
+
+	{
+		FWidgetStateBitfield Pressed = {};
+		FWidgetStateBitfield Hovered = {};
+		FWidgetStateBitfield PressedHovered = {};
+
+		Pressed.SetBinaryState(UWidgetPressedStateRegistration::Bit, true);
+		Hovered.SetBinaryState(UWidgetHoveredStateRegistration::Bit, true);
+		PressedHovered.SetBinaryState(UWidgetPressedStateRegistration::Bit, true);
+		PressedHovered.SetBinaryState(UWidgetHoveredStateRegistration::Bit, true);
+
+		GPassing &= TestEqual(TEXT("Pressed | Hovered is PressedHovered"), Pressed.Union(Hovered), PressedHovered);
+		GPassing &= TestEqual(TEXT("Hovered | Pressed is PressedHovered"), Hovered.Union(Pressed), PressedHovered);
+
+		GPassing &= TestEqual(TEXT("PressedHovered | Pressed is PressedHovered"), PressedHovered.Union(Pressed), PressedHovered);
+		GPassing &= TestEqual(TEXT("Pressed | PressedHovered is PressedHovered"), Pressed.Union(PressedHovered), PressedHovered);
+
+		GPassing &= TestEqual(TEXT("PressedHovered | Hovered is PressedHovered"), PressedHovered.Union(Hovered), PressedHovered);
+		GPassing &= TestEqual(TEXT("Hovered | PressedHovered is PressedHovered"), Hovered.Union(PressedHovered), PressedHovered);
+
+		GPassing &= TestEqual(TEXT("Bit: Pressed | Hovered is PressedHovered"), UWidgetPressedStateRegistration::Bit.Union(UWidgetHoveredStateRegistration::Bit), PressedHovered);
+		GPassing &= TestEqual(TEXT("Bit: Hovered | Pressed is PressedHovered"), UWidgetHoveredStateRegistration::Bit.Union(UWidgetPressedStateRegistration::Bit), PressedHovered);
+	}
+
+	{
+		FWidgetStateBitfield Checked = {};
+		FWidgetStateBitfield Unchecked = {};
+		FWidgetStateBitfield CheckedUnchecked = {};
+
+		Checked.SetEnumState(UWidgetCheckedStateRegistration::Checked);
+		Unchecked.SetEnumState(UWidgetCheckedStateRegistration::Unchecked);
+		CheckedUnchecked.SetEnumState(UWidgetCheckedStateRegistration::Checked.Union(UWidgetCheckedStateRegistration::Unchecked));
+
+		GPassing &= TestEqual(TEXT("Checked | Unchecked is CheckedUnchecked"), Checked.Union(Unchecked), CheckedUnchecked);
+		GPassing &= TestEqual(TEXT("Checked | Unchecked is true"), Checked.Union(Unchecked), true);
+
+		GPassing &= TestEqual(TEXT("CheckedUnchecked & Checked is true"), CheckedUnchecked.Intersect(Checked), true);
+		GPassing &= TestEqual(TEXT("CheckedUnchecked & Unchecked is true"), CheckedUnchecked.Intersect(Unchecked), true);
+
+		GPassing &= TestEqual(TEXT("Checked & CheckedUnchecked is true"), Checked.Intersect(CheckedUnchecked), true);
+		GPassing &= TestEqual(TEXT("Unchecked & CheckedUnchecked is true"), Unchecked.Intersect(CheckedUnchecked), true);
+	}
+
+	return GPassing;
+}
+
+IMPLEMENT_SIMPLE_AUTOMATION_TEST(FOperatorBitwiseNotTest, "Slate.WidgetState.OperatorBitwiseNotTest", EAutomationTestFlags::ApplicationContextMask | EAutomationTestFlags::EngineFilter)
+bool FOperatorBitwiseNotTest::RunTest(const FString& Parameters)
+{
+	bool GPassing = true;
+	{
+		FWidgetStateBitfield Test = {};
+		GPassing &= TestEqual(TEXT("~{} is true"), ~Test, true);
+	}
+
+	{
+		FWidgetStateBitfield Pressed = {};
+
+		Pressed.SetBinaryState(UWidgetPressedStateRegistration::Bit, true);
+
+		GPassing &= TestEqual(TEXT("~Pressed is true"), ~Pressed, true);
+
+		GPassing &= TestEqual(TEXT("Pressed & ~Pressed is false"), Pressed.Intersect(~Pressed), false);
+		GPassing &= TestEqual(TEXT("Pressed does not have any flags ~Pressed"), Pressed.HasAnyFlags(~Pressed), false);
+		GPassing &= TestEqual(TEXT("~Pressed does not have any flags Pressed"), (~Pressed).HasAnyFlags(Pressed), false);
+	}
+
+	{
+		FWidgetStateBitfield Checked = {};
+		FWidgetStateBitfield Unchecked = {};
+		FWidgetStateBitfield CheckedUnchecked = {};
+
+		Checked.SetEnumState(UWidgetCheckedStateRegistration::Checked);
+		Unchecked.SetEnumState(UWidgetCheckedStateRegistration::Unchecked);
+		CheckedUnchecked.SetEnumState(UWidgetCheckedStateRegistration::Checked.Union(UWidgetCheckedStateRegistration::Unchecked));
+
+		GPassing &= TestEqual(TEXT("Checked & ~Checked is false"), Checked.Intersect(~Checked), false);
+		GPassing &= TestEqual(TEXT("Unchecked & ~Unchecked is false"), Unchecked.Intersect(~Unchecked), false);
+		GPassing &= TestEqual(TEXT("CheckedUnchecked & ~CheckedUnchecked is false"), CheckedUnchecked.Intersect(~CheckedUnchecked), false);
+	}
+
+	return GPassing;
+}
+
+IMPLEMENT_SIMPLE_AUTOMATION_TEST(FHasFlagTest, "Slate.WidgetState.HasFlagTest", EAutomationTestFlags::ApplicationContextMask | EAutomationTestFlags::EngineFilter)
+bool FHasFlagTest::RunTest(const FString& Parameters)
+{
+	bool GPassing = true;
+
+	{
+		FWidgetStateBitfield Test = {};
+
+		Test.SetBinaryStateSlow(FName("Pressed"), true);
+		GPassing &= TestEqual(TEXT("Pressed any flag true"), Test.HasAnyFlags(UWidgetPressedStateRegistration::Bit), true);
+		GPassing &= TestEqual(TEXT("Pressed all flag true"), Test.HasAllFlags(UWidgetPressedStateRegistration::Bit), true);
+		GPassing &= TestEqual(TEXT("Pressed any binary flag true"), Test.HasAnyBinaryFlags(UWidgetPressedStateRegistration::Bit), true);
+		GPassing &= TestEqual(TEXT("Pressed all binary flag true"), Test.HasAllBinaryFlags(UWidgetPressedStateRegistration::Bit), true);
+
+		Test.SetBinaryStateSlow(FName("Pressed"), false);
+		GPassing &= TestEqual(TEXT("Pressed any flag false"), Test.HasAnyFlags(UWidgetPressedStateRegistration::Bit), false);
+		GPassing &= TestEqual(TEXT("Pressed all flag false"), Test.HasAllFlags(UWidgetPressedStateRegistration::Bit), false);
+		GPassing &= TestEqual(TEXT("Pressed any binary flag false"), Test.HasAnyBinaryFlags(UWidgetPressedStateRegistration::Bit), false);
+		GPassing &= TestEqual(TEXT("Pressed all binary flag false"), Test.HasAllBinaryFlags(UWidgetPressedStateRegistration::Bit), false);
+
+		Test.SetBinaryState(UWidgetPressedStateRegistration::Bit, true);
+		GPassing &= TestEqual(TEXT("Pressed any flag true"), Test.HasAnyFlags(UWidgetPressedStateRegistration::Bit), true);
+		GPassing &= TestEqual(TEXT("Pressed all flag true"), Test.HasAllFlags(UWidgetPressedStateRegistration::Bit), true);
+		GPassing &= TestEqual(TEXT("Pressed any binary flag true"), Test.HasAnyBinaryFlags(UWidgetPressedStateRegistration::Bit), true);
+		GPassing &= TestEqual(TEXT("Pressed all binary flag true"), Test.HasAllBinaryFlags(UWidgetPressedStateRegistration::Bit), true);
+
+		Test.SetBinaryState(UWidgetPressedStateRegistration::Bit, false);
+		GPassing &= TestEqual(TEXT("Pressed any flag false"), Test.HasAnyFlags(UWidgetPressedStateRegistration::Bit), false);
+		GPassing &= TestEqual(TEXT("Pressed all flag false"), Test.HasAllFlags(UWidgetPressedStateRegistration::Bit), false);
+		GPassing &= TestEqual(TEXT("Pressed any binary flag false"), Test.HasAnyBinaryFlags(UWidgetPressedStateRegistration::Bit), false);
+		GPassing &= TestEqual(TEXT("Pressed all binary flag false"), Test.HasAllBinaryFlags(UWidgetPressedStateRegistration::Bit), false);
+
+		// Slow cache results
+		uint8 PressedIndex = UWidgetStateSettings::Get()->GetBinaryStateIndex(FName("Pressed"));
+
+		Test.SetBinaryState(PressedIndex, true);
+		GPassing &= TestEqual(TEXT("Pressed any flag true"), Test.HasAnyFlags(UWidgetPressedStateRegistration::Bit), true);
+		GPassing &= TestEqual(TEXT("Pressed all flag true"), Test.HasAllFlags(UWidgetPressedStateRegistration::Bit), true);
+		GPassing &= TestEqual(TEXT("Pressed any binary flag true"), Test.HasAnyBinaryFlags(UWidgetPressedStateRegistration::Bit), true);
+		GPassing &= TestEqual(TEXT("Pressed all binary flag true"), Test.HasAllBinaryFlags(UWidgetPressedStateRegistration::Bit), true);
+
+		Test.SetBinaryState(PressedIndex, false);
+		GPassing &= TestEqual(TEXT("Pressed any flag false"), Test.HasAnyFlags(UWidgetPressedStateRegistration::Bit), false);
+		GPassing &= TestEqual(TEXT("Pressed all flag false"), Test.HasAllFlags(UWidgetPressedStateRegistration::Bit), false);
+		GPassing &= TestEqual(TEXT("Pressed any binary flag false"), Test.HasAnyBinaryFlags(UWidgetPressedStateRegistration::Bit), false);
+		GPassing &= TestEqual(TEXT("Pressed all binary flag false"), Test.HasAllBinaryFlags(UWidgetPressedStateRegistration::Bit), false);
+	}
+
+	{
+		FWidgetStateBitfield Test = {};
+		FWidgetStateBitfield PressedHovered = {};
+
+		PressedHovered.SetBinaryState(UWidgetPressedStateRegistration::Bit, true);
+		PressedHovered.SetBinaryState(UWidgetHoveredStateRegistration::Bit, true);
+
+		Test.SetBinaryState(UWidgetPressedStateRegistration::Bit, true);
+		Test.SetBinaryState(UWidgetHoveredStateRegistration::Bit, true);
+		GPassing &= TestEqual(TEXT("Pressed any flag true"), Test.HasAnyFlags(UWidgetPressedStateRegistration::Bit), true);
+		GPassing &= TestEqual(TEXT("Pressed all flag true"), Test.HasAllFlags(UWidgetPressedStateRegistration::Bit), true);
+		GPassing &= TestEqual(TEXT("Hovered any flag true"), Test.HasAnyFlags(UWidgetHoveredStateRegistration::Bit), true);
+		GPassing &= TestEqual(TEXT("Hovered all flag true"), Test.HasAllFlags(UWidgetHoveredStateRegistration::Bit), true);
+		GPassing &= TestEqual(TEXT("PressedHovered any flag true"), Test.HasAnyFlags(PressedHovered), true);
+		GPassing &= TestEqual(TEXT("PressedHovered all flag true"), Test.HasAllFlags(PressedHovered), true);
+
+		Test.SetBinaryState(UWidgetHoveredStateRegistration::Bit, false);
+		GPassing &= TestEqual(TEXT("Pressed any flag true"), Test.HasAnyFlags(UWidgetPressedStateRegistration::Bit), true);
+		GPassing &= TestEqual(TEXT("Pressed all flag true"), Test.HasAllFlags(UWidgetPressedStateRegistration::Bit), true);
+		GPassing &= TestEqual(TEXT("Hovered any flag true"), Test.HasAnyFlags(UWidgetHoveredStateRegistration::Bit), false);
+		GPassing &= TestEqual(TEXT("Hovered all flag true"), Test.HasAllFlags(UWidgetHoveredStateRegistration::Bit), false);
+		GPassing &= TestEqual(TEXT("PressedHovered any flag true"), Test.HasAnyFlags(PressedHovered), true);
+		GPassing &= TestEqual(TEXT("PressedHovered all flag true"), Test.HasAllFlags(PressedHovered), false);
+	}
+
+	{
+		FWidgetStateBitfield Test = {};
+
+		Test.SetEnumState(UWidgetCheckedStateRegistration::Checked);
+		GPassing &= TestEqual(TEXT("Checked any flag true"), Test.HasAnyFlags(UWidgetCheckedStateRegistration::Checked), true);
+		GPassing &= TestEqual(TEXT("Checked all flag true"), Test.HasAllFlags(UWidgetCheckedStateRegistration::Checked), true);
+		GPassing &= TestEqual(TEXT("Checked any enum flag true"), Test.HasAnyEnumFlags(UWidgetCheckedStateRegistration::Checked), true);
+		GPassing &= TestEqual(TEXT("Checked all enum flag true"), Test.HasAllEnumFlags(UWidgetCheckedStateRegistration::Checked), true);
+
+		Test.SetEnumState(UWidgetCheckedStateRegistration::Unchecked);
+		GPassing &= TestEqual(TEXT("Unchecked any flag true"), Test.HasAnyFlags(UWidgetCheckedStateRegistration::Unchecked), true);
+		GPassing &= TestEqual(TEXT("Unchecked all flag true"), Test.HasAllFlags(UWidgetCheckedStateRegistration::Unchecked), true);
+		GPassing &= TestEqual(TEXT("Unchecked any enum flag true"), Test.HasAnyEnumFlags(UWidgetCheckedStateRegistration::Unchecked), true);
+		GPassing &= TestEqual(TEXT("Unchecked all enum flag true"), Test.HasAllEnumFlags(UWidgetCheckedStateRegistration::Unchecked), true);
+		GPassing &= TestEqual(TEXT("Checked any flag false"), Test.HasAnyFlags(UWidgetCheckedStateRegistration::Checked), false);
+		GPassing &= TestEqual(TEXT("Checked all flag false"), Test.HasAllFlags(UWidgetCheckedStateRegistration::Checked), false);
+		GPassing &= TestEqual(TEXT("Checked any enum flag false"), Test.HasAnyEnumFlags(UWidgetCheckedStateRegistration::Checked), false);
+		GPassing &= TestEqual(TEXT("Checked all enum flag false"), Test.HasAllEnumFlags(UWidgetCheckedStateRegistration::Checked), false);
+
+		Test.ClearEnumState(UWidgetCheckedStateRegistration::Unchecked);
+		GPassing &= TestEqual(TEXT("Unchecked any flag false"), Test.HasAnyFlags(UWidgetCheckedStateRegistration::Unchecked), false);
+		GPassing &= TestEqual(TEXT("Unchecked all flag false"), Test.HasAllFlags(UWidgetCheckedStateRegistration::Unchecked), false);
+		GPassing &= TestEqual(TEXT("Unchecked any enum flag false"), Test.HasAnyEnumFlags(UWidgetCheckedStateRegistration::Unchecked), false);
+		GPassing &= TestEqual(TEXT("Unchecked all enum flag false"), Test.HasAllEnumFlags(UWidgetCheckedStateRegistration::Unchecked), false);
+	}
+
+	return GPassing;
+}
+
+#endif // WITH_DEV_AUTOMATION_TESTS
\ No newline at end of file
diff --git a/Engine/Source/Runtime/UMG/Private/Binding/States/WidgetStateBitfield.cpp b/Engine/Source/Runtime/UMG/Private/Binding/States/WidgetStateBitfield.cpp
new file mode 100644
--- /dev/null
+++ b/Engine/Source/Runtime/UMG/Private/Binding/States/WidgetStateBitfield.cpp
@@ -0,0 +1,322 @@
+// Copyright Epic Games, Inc. All Rights Reserved.
+ 
+#include "Binding/States/WidgetStateBitfield.h"
+
+#include "Binding/States/WidgetStateSettings.h"
+
+#include UE_INLINE_GENERATED_CPP_BY_NAME(WidgetStateBitfield)
+
+FWidgetStateBitfield::FWidgetStateBitfield()
+{
+}
+
+FWidgetStateBitfield::FWidgetStateBitfield(const FName InStateName)
+{
+	SetBinaryStateSlow(InStateName, true);
+}
+
+FWidgetStateBitfield::FWidgetStateBitfield(const FName InStateName, const uint8 InValue)
+{
+	SetEnumStateSlow(InStateName, InValue);
+}
+
+FWidgetStateBitfield FWidgetStateBitfield::operator~() const
+{
+	FWidgetStateBitfield Temp = *this;
+	Temp.NegateStates();
+	return Temp;
+}
+
+FWidgetStateBitfield FWidgetStateBitfield::Intersect(const FWidgetStateBitfield& Rhs) const
+{
+	FWidgetStateBitfield Temp;
+	Temp.BinaryStates = BinaryStates & Rhs.BinaryStates;
+	Temp.IsEnumStateNotAllowed = IsEnumStateUsed & Rhs.IsEnumStateUsed;
+
+	// See'IsEnumStateUsed' in header, we want to track enum state usage additively even with '&'
+	Temp.IsEnumStateUsed = IsEnumStateUsed | Rhs.IsEnumStateUsed;
+
+	// Only set enum states if all RHS binary state tests passed, else enum pass can shadow binary failures
+	bool bBinaryStatesPassed = !Rhs.BinaryStates || Temp.BinaryStates;
+	if (bBinaryStatesPassed && Temp.IsEnumStateUsed != 0)
+	{
+		for (uint8 EnumStateIndex = 0; EnumStateIndex < EnumStates.Num(); EnumStateIndex++)
+		{
+			bool bIsStateUsed = Temp.IsEnumStateUsed & (1 << EnumStateIndex);
+
+			if (bIsStateUsed)
+			{
+				bool bIsNotAllowedLhs = IsEnumStateNotAllowed & (1 << EnumStateIndex);
+				bool bIsNotAllowedRhs = Rhs.IsEnumStateNotAllowed & (1 << EnumStateIndex);
+				bool bIsAllowanceSame = bIsNotAllowedLhs == bIsNotAllowedRhs;
+
+				if (bIsAllowanceSame)
+				{
+					// If the RHS doesn't use this enum state, carry over the LHS enum state set
+					// Allows for tests against bitfields where RHS doesn't use LHS enum states
+					// to pass (Ex: {checked, hovered} & {hovered} should pass).
+
+					bool bRhsNotUsed = (Rhs.IsEnumStateUsed & (1 << EnumStateIndex)) == 0;
+
+					if (bRhsNotUsed)
+					{
+						Temp.EnumStates[EnumStateIndex] = EnumStates[EnumStateIndex];
+					}
+					else
+					{
+						Temp.EnumStates[EnumStateIndex] = EnumStates[EnumStateIndex].Intersect(Rhs.EnumStates[EnumStateIndex]);
+					}
+				}
+				else if (bIsNotAllowedRhs)
+				{
+					Temp.EnumStates[EnumStateIndex] = EnumStates[EnumStateIndex].Difference(Rhs.EnumStates[EnumStateIndex]);
+				}
+				else if (bIsNotAllowedLhs)
+				{
+					Temp.EnumStates[EnumStateIndex] = Rhs.EnumStates[EnumStateIndex].Difference(EnumStates[EnumStateIndex]);
+				}
+				else
+				{
+					checkf(false, TEXT("Should not be able to reach, all allowance combinations covered above"));
+				}
+			}
+		}
+	}
+
+	return Temp;
+}
+
+FWidgetStateBitfield FWidgetStateBitfield::Union(const FWidgetStateBitfield& Rhs) const
+{
+	FWidgetStateBitfield Temp;
+	Temp.BinaryStates = BinaryStates | Rhs.BinaryStates;
+	Temp.IsEnumStateUsed = IsEnumStateUsed | Rhs.IsEnumStateUsed;
+
+	// See'IsEnumStateNotAllowed' in header, we wain to maintain the allowed set so shrink the not allowed
+	Temp.IsEnumStateNotAllowed = IsEnumStateNotAllowed & Rhs.IsEnumStateNotAllowed;
+
+	if (Temp.IsEnumStateUsed != 0)
+	{
+		for (uint8 EnumStateIndex = 0; EnumStateIndex < EnumStates.Num(); EnumStateIndex++)
+		{
+			bool bIsStateUsed = Temp.IsEnumStateUsed & (1 << EnumStateIndex);
+			bool bIsStateAllowedUsed = Temp.IsEnumStateUsed & (1 << EnumStateIndex);
+			if (bIsStateUsed)
+			{
+				bool bIsNotAllowedLhs = IsEnumStateNotAllowed & (1 << EnumStateIndex);
+				bool bIsNotAllowedRhs = Rhs.IsEnumStateNotAllowed & (1 << EnumStateIndex);
+				bool bIsAllowanceSame = bIsNotAllowedLhs == bIsNotAllowedRhs;
+
+				if (bIsAllowanceSame)
+				{
+					Temp.EnumStates[EnumStateIndex] = EnumStates[EnumStateIndex].Union(Rhs.EnumStates[EnumStateIndex]);
+				}
+				else if (bIsNotAllowedRhs)
+				{
+					Temp.EnumStates[EnumStateIndex] = EnumStates[EnumStateIndex];
+				}
+				else if (bIsNotAllowedLhs)
+				{
+					Temp.EnumStates[EnumStateIndex] = Rhs.EnumStates[EnumStateIndex];
+				}
+				else
+				{
+					checkf(false, TEXT("Should not be able to reach, all allowance combinations covered above"));
+				}
+			}
+		}
+	}
+
+	return Temp;
+}
+
+FWidgetStateBitfield::operator bool() const
+{
+	// We have to check 'IsEnumStateUsed' here since otherwise a completely empty statefield would pass
+	return (BinaryStates != 0 || IsEnumStateUsed != 0) && !HasEmptyUsedEnumStates();
+}
+
+bool FWidgetStateBitfield::HasBinaryStates() const
+{
+	return BinaryStates != 0;
+}
+
+bool FWidgetStateBitfield::HasEnumStates() const
+{
+	return IsEnumStateUsed != 0;
+}
+
+bool FWidgetStateBitfield::HasEmptyUsedEnumStates() const
+{
+	if (IsEnumStateUsed != 0)
+	{
+		for (uint8 EnumStateIndex = 0; EnumStateIndex < EnumStates.Num(); EnumStateIndex++)
+		{
+			if (IsEnumStateUsed & (1 << EnumStateIndex))
+			{
+				if (EnumStates[EnumStateIndex].IsEmpty())
+				{
+					return true;
+				}
+			}
+		}
+	}
+	return false;
+}
+
+bool FWidgetStateBitfield::HasAnyFlags(const FWidgetStateBitfield& InBitfield)
+{
+	return HasAnyBinaryFlags(InBitfield) || HasAnyEnumFlags(InBitfield);
+}
+
+bool FWidgetStateBitfield::HasAllFlags(const FWidgetStateBitfield& InBitfield)
+{
+	return HasAllBinaryFlags(InBitfield) && HasAllEnumFlags(InBitfield);
+}
+
+bool FWidgetStateBitfield::HasAnyBinaryFlags(const FWidgetStateBitfield& InBitfield)
+{
+	return (BinaryStates & InBitfield.BinaryStates) != 0;
+}
+
+bool FWidgetStateBitfield::HasAllBinaryFlags(const FWidgetStateBitfield& InBitfield)
+{
+	return (BinaryStates & InBitfield.BinaryStates) == InBitfield.BinaryStates;
+}
+
+bool FWidgetStateBitfield::HasAnyEnumFlags(const FWidgetStateBitfield& InBitfield)
+{
+	const uint16 IsEnumStateShared = IsEnumStateUsed & InBitfield.IsEnumStateUsed;
+
+	if (IsEnumStateShared != 0)
+	{
+		for (uint8 EnumStateIndex = 0; EnumStateIndex < EnumStates.Num(); EnumStateIndex++)
+		{
+			if (IsEnumStateShared & (1 << EnumStateIndex))
+			{
+				if (!EnumStates[EnumStateIndex].Intersect(InBitfield.EnumStates[EnumStateIndex]).IsEmpty())
+				{
+					return true;
+				}
+			}
+		}
+	}
+	return false;
+}
+
+bool FWidgetStateBitfield::HasAllEnumFlags(const FWidgetStateBitfield& InBitfield)
+{
+	const uint16 IsEnumStateShared = IsEnumStateUsed & InBitfield.IsEnumStateUsed;
+
+	if (IsEnumStateShared == InBitfield.IsEnumStateUsed)
+	{
+		for (uint8 EnumStateIndex = 0; EnumStateIndex < EnumStates.Num(); EnumStateIndex++)
+		{
+			if (InBitfield.IsEnumStateUsed & (1 << EnumStateIndex))
+			{
+				if (!EnumStates[EnumStateIndex].Includes(InBitfield.EnumStates[EnumStateIndex]))
+				{
+					return false;
+				}
+			}
+		}
+
+		return true;
+	}
+
+	return false;
+}
+
+void FWidgetStateBitfield::SetState(const FWidgetStateBitfield& InBitfield)
+{
+	*this = InBitfield;
+}
+
+void FWidgetStateBitfield::NegateStates()
+{
+	NegateBinaryStates();
+	NegateEnumStates();
+}
+
+void FWidgetStateBitfield::SetBinaryState(uint8 BinaryStateIndex, bool BinaryStateValue)
+{
+	BinaryStates = (BinaryStates & ~(static_cast<uint64>(1) << BinaryStateIndex)) | (static_cast<uint64>(BinaryStateValue) << BinaryStateIndex);
+}
+
+void FWidgetStateBitfield::SetBinaryState(const FWidgetStateBitfield& BinaryStateBitfield, bool BinaryStateValue)
+{
+	if (BinaryStateValue)
+	{
+		BinaryStates = BinaryStates | BinaryStateBitfield.BinaryStates;
+	}
+	else
+	{
+		BinaryStates = BinaryStates & ~(BinaryStateBitfield.BinaryStates);
+	}
+}
+
+void FWidgetStateBitfield::SetBinaryStateSlow(FName BinaryStateName, bool BinaryStateValue)
+{
+	uint8 BinaryStateIndex = UWidgetStateSettings::Get()->GetBinaryStateIndex(BinaryStateName);
+	SetBinaryState(BinaryStateIndex, BinaryStateValue);
+}
+
+void FWidgetStateBitfield::NegateBinaryStates()
+{
+	BinaryStates = ~BinaryStates;
+}
+
+void FWidgetStateBitfield::SetEnumState(uint8 EnumStateIndex, uint8 EnumStateValue)
+{
+	IsEnumStateUsed = IsEnumStateUsed | (1 << EnumStateIndex);
+	EnumStates[EnumStateIndex] = { EnumStateValue };
+}
+
+void FWidgetStateBitfield::SetEnumState(const FWidgetStateBitfield& EnumStateBitfield)
+{
+	IsEnumStateUsed |= EnumStateBitfield.IsEnumStateUsed;
+
+	for (uint8 EnumStateIndex = 0; EnumStateIndex < EnumStates.Num(); EnumStateIndex++)
+	{
+		if (EnumStateBitfield.IsEnumStateUsed & (1 << EnumStateIndex))
+		{
+			EnumStates[EnumStateIndex] = EnumStateBitfield.EnumStates[EnumStateIndex];
+		}
+	}
+}
+
+void FWidgetStateBitfield::SetEnumStateSlow(FName EnumStateName, uint8 EnumStateValue)
+{
+	uint8 EnumStateIndex = UWidgetStateSettings::Get()->GetEnumStateIndex(EnumStateName);
+	SetEnumState(EnumStateIndex, EnumStateValue);
+}
+
+void FWidgetStateBitfield::ClearEnumState(const FWidgetStateBitfield& EnumStateBitfield)
+{
+	IsEnumStateUsed &= ~EnumStateBitfield.IsEnumStateUsed;
+
+	for (uint8 EnumStateIndex = 0; EnumStateIndex < EnumStates.Num(); EnumStateIndex++)
+	{
+		if (EnumStateBitfield.IsEnumStateUsed & (1 << EnumStateIndex))
+		{
+			EnumStates[EnumStateIndex].Empty();
+		}
+	}
+}
+
+void FWidgetStateBitfield::ClearEnumState(uint8 EnumStateIndex)
+{
+	IsEnumStateUsed = IsEnumStateUsed & ~(1 << EnumStateIndex);
+	EnumStates[EnumStateIndex].Empty();
+}
+
+void FWidgetStateBitfield::ClearEnumState(FName EnumStateName)
+{
+	uint8 EnumStateIndex = UWidgetStateSettings::Get()->GetEnumStateIndex(EnumStateName);
+	ClearEnumState(EnumStateIndex);
+}
+
+void FWidgetStateBitfield::NegateEnumStates()
+{
+	IsEnumStateNotAllowed = ~IsEnumStateNotAllowed;
+}
diff --git a/Engine/Source/Runtime/UMG/Private/Binding/States/WidgetStateRegistration.cpp b/Engine/Source/Runtime/UMG/Private/Binding/States/WidgetStateRegistration.cpp
new file mode 100644
--- /dev/null
+++ b/Engine/Source/Runtime/UMG/Private/Binding/States/WidgetStateRegistration.cpp
@@ -0,0 +1,100 @@
+// Copyright Epic Games, Inc. All Rights Reserved.
+
+#include "Binding/States/WidgetStateRegistration.h"
+
+#include "Components/Widget.h"
+
+#include UE_INLINE_GENERATED_CPP_BY_NAME(WidgetStateRegistration)
+
+FName UWidgetBinaryStateRegistration::GetStateName() const
+{ 
+	return FName(); 
+};
+
+bool UWidgetBinaryStateRegistration::GetRegisteredWidgetState(const UWidget* InWidget) const 
+{ 
+	return false; 
+}
+
+void UWidgetBinaryStateRegistration::InitializeStaticBitfields() const
+{
+}
+
+FName UWidgetHoveredStateRegistration::GetStateName() const
+{
+	return StateName;
+};
+
+bool UWidgetHoveredStateRegistration::GetRegisteredWidgetState(const UWidget* InWidget) const
+{
+	return false;
+}
+
+void UWidgetHoveredStateRegistration::InitializeStaticBitfields() const
+{
+	Bit = FWidgetStateBitfield(GetStateName());
+}
+
+FName UWidgetPressedStateRegistration::GetStateName() const
+{
+	return StateName;
+};
+
+bool UWidgetPressedStateRegistration::GetRegisteredWidgetState(const UWidget* InWidget) const
+{
+	return false;
+}
+
+void UWidgetPressedStateRegistration::InitializeStaticBitfields() const
+{
+	Bit = FWidgetStateBitfield(GetStateName());
+}
+
+FName UWidgetDisabledStateRegistration::GetStateName() const
+{
+	return StateName;
+};
+
+bool UWidgetDisabledStateRegistration::GetRegisteredWidgetState(const UWidget* InWidget) const
+{
+	return !InWidget->GetIsEnabled();
+}
+
+void UWidgetDisabledStateRegistration::InitializeStaticBitfields() const
+{
+	Bit = FWidgetStateBitfield(GetStateName());
+}
+
+FName UWidgetSelectedStateRegistration::GetStateName() const
+{
+	return StateName;
+};
+
+bool UWidgetSelectedStateRegistration::GetRegisteredWidgetState(const UWidget* InWidget) const
+{
+	return false;
+}
+
+void UWidgetSelectedStateRegistration::InitializeStaticBitfields() const
+{
+	Bit = FWidgetStateBitfield(GetStateName());
+}
+
+FName UWidgetEnumStateRegistration::GetStateName() const 
+{ 
+	return FName(); 
+}
+
+bool UWidgetEnumStateRegistration::GetRegisteredWidgetUsesState(const UWidget* InWidget) const 
+{ 
+	return false; 
+}
+
+uint8 UWidgetEnumStateRegistration::GetRegisteredWidgetState(const UWidget* InWidget) const 
+{ 
+	return 0; 
+}
+
+void UWidgetEnumStateRegistration::InitializeStaticBitfields() const
+{
+}
\ No newline at end of file
diff --git a/Engine/Source/Runtime/UMG/Private/Binding/States/WidgetStateSettings.cpp b/Engine/Source/Runtime/UMG/Private/Binding/States/WidgetStateSettings.cpp
new file mode 100644
--- /dev/null
+++ b/Engine/Source/Runtime/UMG/Private/Binding/States/WidgetStateSettings.cpp
@@ -0,0 +1,131 @@
+// Copyright Epic Games, Inc. All Rights Reserved.
+
+#include "Binding/States/WidgetStateSettings.h"
+
+#include "Binding/States/WidgetStateBitfield.h"
+#include "Binding/States/WidgetStateRegistration.h"
+
+#include "Components/Widget.h"
+
+#include UE_INLINE_GENERATED_CPP_BY_NAME(WidgetStateSettings)
+
+void UWidgetStateSettings::PostInitProperties()
+{
+	Super::PostInitProperties();
+
+	{
+		TArray<UClass*> BinaryStateClasses;
+		GetDerivedClasses(UWidgetBinaryStateRegistration::StaticClass(), BinaryStateClasses, false);
+
+		BinaryStateRegistrationCDOs.Empty(BinaryStateClasses.Num());
+		BinaryStateMap.Empty(BinaryStateClasses.Num());
+		BinaryStates.Empty(BinaryStateClasses.Num());
+
+		uint8 BinaryStateIndex = 0;
+		for (UClass* BinaryStateClass : BinaryStateClasses)
+		{
+			// CDO shouldn't ever be nullptr in non-compilation contexts
+			UWidgetBinaryStateRegistration* BinaryStateRegistrationCDO = CastChecked<UWidgetBinaryStateRegistration>(BinaryStateClass->GetDefaultObject());
+			BinaryStateRegistrationCDOs.Add(BinaryStateRegistrationCDO);
+
+			FName BinaryStateName = BinaryStateRegistrationCDO->GetStateName();
+			BinaryStateMap.Add(BinaryStateName, BinaryStateIndex);
+			BinaryStates.Add(BinaryStateName);
+			BinaryStateRegistrationCDO->InitializeStaticBitfields();
+
+			BinaryStateIndex++;
+		}
+	}
+
+	{
+		TArray<UClass*> EnumStateClasses;
+		GetDerivedClasses(UWidgetEnumStateRegistration::StaticClass(), EnumStateClasses, false);
+
+		EnumStateRegistrationCDOs.Empty(EnumStateClasses.Num());
+		EnumStateMap.Empty(EnumStateClasses.Num());
+		EnumStates.Empty(EnumStateClasses.Num());
+
+		uint8 EnumStateIndex = 0;
+		for (UClass* EnumStateClass : EnumStateClasses)
+		{
+			// CDO shouldn't ever be nullptr in non-compilation contexts
+			UWidgetEnumStateRegistration* EnumStateRegistrationCDO = CastChecked<UWidgetEnumStateRegistration>(EnumStateClass->GetDefaultObject());
+			EnumStateRegistrationCDOs.Add(EnumStateRegistrationCDO);
+
+			FName EnumStateName = EnumStateRegistrationCDO->GetStateName();
+			EnumStateMap.Add(EnumStateName, EnumStateIndex);
+			EnumStates.Add(EnumStateName);
+			EnumStateRegistrationCDO->InitializeStaticBitfields();
+
+			EnumStateIndex++;
+		}
+	}
+}
+
+void UWidgetStateSettings::GetAllStateNames(TArray<FName>& OutStateNames) const
+{
+	OutStateNames.Reset(BinaryStates.Num() + EnumStates.Num());
+	OutStateNames.Append(BinaryStates);
+	OutStateNames.Append(EnumStates);
+}
+
+void UWidgetStateSettings::GetBinaryStateNames(TArray<FName>& OutBinaryStateNames) const
+{
+	OutBinaryStateNames = BinaryStates;
+}
+
+void UWidgetStateSettings::GetEnumStateNames(TArray<FName>& OutEnumStateNames) const
+{
+	OutEnumStateNames = EnumStates;
+}
+
+uint8 UWidgetStateSettings::GetBinaryStateIndex(const FName InBinaryStateName) const
+{
+	return BinaryStateMap[InBinaryStateName];
+}
+
+uint8 UWidgetStateSettings::GetEnumStateIndex(const FName InEnumStateName) const
+{
+	return EnumStateMap[InEnumStateName];
+}
+
+FName UWidgetStateSettings::GetBinaryStateName(const uint8 InBinaryStateIndex) const
+{
+	return BinaryStates[InBinaryStateIndex];
+}
+
+FName UWidgetStateSettings::GetEnumStateName(const uint8 InEnumStateIndex) const
+{
+	return EnumStates[InEnumStateIndex];
+}
+
+FWidgetStateBitfield UWidgetStateSettings::GetInitialRegistrationBitfield(const UWidget* InWidget) const
+{
+	FWidgetStateBitfield InitBitfield = {};
+
+	uint8 BinaryStateIndex = 0;
+	for (TObjectPtr<UWidgetBinaryStateRegistration> BinaryStateRegistrationCDO : BinaryStateRegistrationCDOs)
+	{
+		bool InitBinaryState = BinaryStateRegistrationCDO->GetRegisteredWidgetState(InWidget);
+		InitBitfield.SetBinaryState(BinaryStateIndex, InitBinaryState);
+		BinaryStateIndex++;
+	}
+
+	uint8 EnumStateIndex = 0;
+	for (TObjectPtr<UWidgetEnumStateRegistration> EnumStateRegistrationCDO : EnumStateRegistrationCDOs)
+	{
+		bool InitUsesEnumState = EnumStateRegistrationCDO->GetRegisteredWidgetUsesState(InWidget);
+		if (InitUsesEnumState)
+		{
+			uint8 InitEnumState = EnumStateRegistrationCDO->GetRegisteredWidgetState(InWidget);
+			InitBitfield.SetEnumState(EnumStateIndex, InitEnumState);
+		}
+		else
+		{
+			InitBitfield.ClearEnumState(EnumStateIndex);
+		}
+		EnumStateIndex++;
+	}
+
+	return InitBitfield;
+}
diff --git a/Engine/Source/Runtime/UMG/Private/Components/Button.cpp b/Engine/Source/Runtime/UMG/Private/Components/Button.cpp
--- a/Engine/Source/Runtime/UMG/Private/Components/Button.cpp
+++ b/Engine/Source/Runtime/UMG/Private/Components/Button.cpp
@@ -1,6 +1,9 @@
 // Copyright Epic Games, Inc. All Rights Reserved.
 
 #include "Components/Button.h"
+
+#include "Binding/States/WidgetStateBitfield.h"
+#include "Binding/States/WidgetStateRegistration.h"
 #include "Widgets/SNullWidget.h"
 #include "Widgets/DeclarativeSyntaxSupport.h"
 #include "Widgets/Input/SButton.h"
@@ -283,21 +286,27 @@ FReply UButton::SlateHandleClicked()
 void UButton::SlateHandlePressed()
 {
 	OnPressed.Broadcast();
+	BroadcastBinaryPostStateChange(UWidgetPressedStateRegistration::Bit, true);
+
+
 }
 
 void UButton::SlateHandleReleased()
 {
 	OnReleased.Broadcast();
+	BroadcastBinaryPostStateChange(UWidgetPressedStateRegistration::Bit, false);
 }
 
 void UButton::SlateHandleHovered()
 {
 	OnHovered.Broadcast();
+	BroadcastBinaryPostStateChange(UWidgetHoveredStateRegistration::Bit, true);
 }
 
 void UButton::SlateHandleUnhovered()
 {
 	OnUnhovered.Broadcast();
+	BroadcastBinaryPostStateChange(UWidgetHoveredStateRegistration::Bit, false);
 }
 
 #if WITH_ACCESSIBILITY
diff --git a/Engine/Source/Runtime/UMG/Private/Components/CheckBox.cpp b/Engine/Source/Runtime/UMG/Private/Components/CheckBox.cpp
--- a/Engine/Source/Runtime/UMG/Private/Components/CheckBox.cpp
+++ b/Engine/Source/Runtime/UMG/Private/Components/CheckBox.cpp
@@ -1,6 +1,8 @@
 // Copyright Epic Games, Inc. All Rights Reserved.
 
 #include "Components/CheckBox.h"
+#include "Binding/States/WidgetStateBitfield.h"
+#include "Binding/States/WidgetStateRegistration.h"
 #include "Widgets/SNullWidget.h"
 #include "Widgets/DeclarativeSyntaxSupport.h"
 #include "Widgets/Input/SCheckBox.h"
@@ -209,9 +211,12 @@ ECheckBoxState UCheckBox::GetCheckedState() const
 
 void UCheckBox::SetIsChecked(bool InIsChecked)
 {
+	bool bValueChanged = false;
+
 	ECheckBoxState NewState = InIsChecked ? ECheckBoxState::Checked : ECheckBoxState::Unchecked;
 	if (NewState != CheckedState)
 	{
+		bValueChanged = true;
 		CheckedState = NewState;
 		BroadcastFieldValueChanged(FFieldNotificationClassDescriptor::CheckedState);
 	}
@@ -220,12 +225,20 @@ void UCheckBox::SetIsChecked(bool InIsChecked)
 	{
 		MyCheckbox->SetIsChecked(PROPERTY_BINDING(ECheckBoxState, CheckedState));
 	}
+
+	if (bValueChanged)
+	{
+		BroadcastEnumPostStateChange(InIsChecked ? UWidgetCheckedStateRegistration::Checked : UWidgetCheckedStateRegistration::Unchecked);
+	}
 }
 
 void UCheckBox::SetCheckedState(ECheckBoxState InCheckedState)
 {
+	bool bValueChanged = false;
+
 	if (CheckedState != InCheckedState)
 	{
+		bValueChanged = true;
 		CheckedState = InCheckedState;
 		BroadcastFieldValueChanged(FFieldNotificationClassDescriptor::CheckedState);
 	}
@@ -234,6 +247,11 @@ void UCheckBox::SetCheckedState(ECheckBoxState InCheckedState)
 	{
 		MyCheckbox->SetIsChecked(PROPERTY_BINDING(ECheckBoxState, CheckedState));
 	}
+
+	if (bValueChanged)
+	{
+		BroadcastEnumPostStateChange(UWidgetCheckedStateRegistration::GetBitfieldFromValue((uint8)CheckedState));
+	}
 }
 
 const FCheckBoxStyle& UCheckBox::GetWidgetStyle() const
@@ -300,7 +318,43 @@ const FText UCheckBox::GetPaletteCategory()
 
 #endif
 
+FName UWidgetCheckedStateRegistration::GetStateName() const
+{
+	return StateName;
+};
+
+uint8 UWidgetCheckedStateRegistration::GetRegisteredWidgetState(const UWidget* InWidget) const
+{
+	if (const UCheckBox* CheckBox = Cast<UCheckBox>(InWidget))
+	{
+		return (uint8)CheckBox->GetCheckedState();
+	}
+
+	return 0;
+}
+
+const FWidgetStateBitfield& UWidgetCheckedStateRegistration::GetBitfieldFromValue(uint8 InValue)
+{
+	switch ((ECheckBoxState)InValue)
+	{
+	case ECheckBoxState::Unchecked:
+		return UWidgetCheckedStateRegistration::Unchecked;
+	case ECheckBoxState::Checked:
+		return UWidgetCheckedStateRegistration::Checked;
+	case ECheckBoxState::Undetermined:
+		return UWidgetCheckedStateRegistration::Undetermined;
+	default:
+		return UWidgetCheckedStateRegistration::Undetermined;
+	}
+}
+
+void UWidgetCheckedStateRegistration::InitializeStaticBitfields() const
+{
+	Unchecked = FWidgetStateBitfield(GetStateName(), (uint8)ECheckBoxState::Unchecked);
+	Checked = FWidgetStateBitfield(GetStateName(), (uint8)ECheckBoxState::Checked);
+	Undetermined = FWidgetStateBitfield(GetStateName(), (uint8)ECheckBoxState::Undetermined);
+}
+
 /////////////////////////////////////////////////////
 
 #undef LOCTEXT_NAMESPACE
-
diff --git a/Engine/Source/Runtime/UMG/Private/Components/Widget.cpp b/Engine/Source/Runtime/UMG/Private/Components/Widget.cpp
--- a/Engine/Source/Runtime/UMG/Private/Components/Widget.cpp
+++ b/Engine/Source/Runtime/UMG/Private/Components/Widget.cpp
@@ -20,6 +20,9 @@
 #include "Widgets/Layout/SSpacer.h"
 #include "Widgets/SToolTip.h"
 #include "Binding/PropertyBinding.h"
+#include "Binding/States/WidgetStateBitfield.h"
+#include "Binding/States/WidgetStateSettings.h"
+#include "Binding/States/WidgetStateRegistration.h"
 #include "Binding/WidgetFieldNotificationExtension.h"
 #include "Logging/MessageLog.h"
 #include "Blueprint/GameViewportSubsystem.h"
@@ -201,6 +204,8 @@ PRAGMA_DISABLE_DEPRECATION_WARNINGS
 #endif
 	AccessibleWidgetData = nullptr;
 
+	bShouldBroadcastState = true;
+
 PRAGMA_ENABLE_DEPRECATION_WARNINGS
 
 #if WITH_EDITORONLY_DATA
@@ -306,10 +311,12 @@ bool UWidget::GetIsEnabled() const
 
 void UWidget::SetIsEnabled(bool bInIsEnabled)
 {
+	bool bValueChanged = false;
 	if (bIsEnabled != bInIsEnabled)
 	{
 		bIsEnabled = bInIsEnabled;
 		BroadcastFieldValueChanged(FFieldNotificationClassDescriptor::bIsEnabled);
+		bValueChanged = true;
 	}
 
 	TSharedPtr<SWidget> SafeWidget = GetCachedWidget();
@@ -317,6 +324,12 @@ void UWidget::SetIsEnabled(bool bInIsEnabled)
 	{
 		SafeWidget->SetEnabled(bInIsEnabled);
 	}
+
+	if (bValueChanged)
+	{
+		// Note: State is disabled, so we broadcast !bIsEnabled
+		BroadcastBinaryPostStateChange(UWidgetDisabledStateRegistration::Bit, !bIsEnabled);
+	}
 }
 
 bool UWidget::IsInViewport() const
@@ -1318,6 +1331,13 @@ void UWidget::PreSave(FObjectPreSaveContext ObjectSaveContext)
 	SynchronizeAccessibleData();
 }
 
+void UWidget::ReleaseSlateResources(bool bReleaseChildren)
+{
+	UVisual::ReleaseSlateResources(bReleaseChildren);
+
+	MyWidgetStateBitfield.Reset();
+}
+
 #if WITH_EDITOR
 bool UWidget::Modify(bool bAlwaysMarkDirty)
 {
@@ -1773,11 +1793,53 @@ bool UWidget::AddBinding(FDelegateProperty* DelegateProperty, UObject* SourceObj
 	return false;
 }
 
+FDelegateHandle UWidget::RegisterPostStateListener(const FOnWidgetStateBroadcast::FDelegate& ListenerDelegate, bool bBroadcastCurrentState)
+{
+	if (!MyWidgetStateBitfield.IsValid())
+	{
+		MyWidgetStateBitfield = MakeShared<FWidgetStateBitfield>(UWidgetStateSettings::Get()->GetInitialRegistrationBitfield(this));
+	}
+
+	if (bBroadcastCurrentState)
+	{
+		ListenerDelegate.ExecuteIfBound(*MyWidgetStateBitfield);
+	}
+
+	return PostWidgetStateChanged.Add(ListenerDelegate);
+}
+
+void UWidget::UnregisterPostStateListener(const FDelegateHandle& ListenerDelegate)
+{
+	PostWidgetStateChanged.Remove(ListenerDelegate);
+
+	if (!PostWidgetStateChanged.IsBound())
+	{
+		MyWidgetStateBitfield.Reset();
+	}
+}
+
 void UWidget::OnBindingChanged(const FName& Property)
 {
 
 }
 
+void UWidget::BroadcastBinaryPostStateChange(const FWidgetStateBitfield& StateChange, bool bInValue)
+{
+	if (bShouldBroadcastState && MyWidgetStateBitfield.IsValid())
+	{
+		MyWidgetStateBitfield->SetBinaryState(StateChange, bInValue);
+		PostWidgetStateChanged.Broadcast(*MyWidgetStateBitfield);
+	}
+}
+
+void UWidget::BroadcastEnumPostStateChange(const FWidgetStateBitfield& StateChange)
+{
+	if (bShouldBroadcastState && MyWidgetStateBitfield.IsValid())
+	{
+		MyWidgetStateBitfield->SetEnumState(StateChange);
+		PostWidgetStateChanged.Broadcast(*MyWidgetStateBitfield);
+	}
+}
 
 namespace UE::UMG::Private
 {
diff --git a/Engine/Source/Runtime/UMG/Public/Binding/States/WidgetStateBitfield.h b/Engine/Source/Runtime/UMG/Public/Binding/States/WidgetStateBitfield.h
new file mode 100644
--- /dev/null
+++ b/Engine/Source/Runtime/UMG/Public/Binding/States/WidgetStateBitfield.h
@@ -0,0 +1,194 @@
+// Copyright Epic Games, Inc. All Rights Reserved.
+
+#pragma once
+
+#include "CoreMinimal.h"
+#include "Containers/StaticArray.h"
+#include "Containers/ContainerAllocationPolicies.h"
+#include "WidgetStateBitfield.generated.h"
+
+/** 
+ * Expensive to create wrapper for binary and enumerated widget states that supports
+ * FName-based interaction and dynamic states populated in an IWidgetStateNameMapper
+ * May be used to represent either a target widget state combination, or an actual widgets state
+ * 
+ * All this does is cache the integer results of dynamic FName state lookup from an IWidgetStateNameMapper
+ * So it's prefertable to create these once and reuse. In particular, if broadcasting current state
+ * determine the correct bitfield to modify the existing state ahead of time. Save that bitfield and use it
+ * when performing state-modifications. Likewise when performing state tests, try to construct the relevant
+ * bitfield ahead of time & store that to later be compared against  
+ */
+USTRUCT()
+struct UMG_API FWidgetStateBitfield
+{
+	GENERATED_BODY()
+
+public:
+
+	/** Default, will convert to false as a bool */
+	FWidgetStateBitfield();
+
+	/** Interprets name as binary state on */
+	FWidgetStateBitfield(const FName InStateName);
+
+	/** Interprets name value pair as enum state */
+	FWidgetStateBitfield(const FName InStateName, const uint8 InValue);
+
+public:
+
+	/** 
+	 * Negation of binary and enum states, maintains enum state usage, but marks states as not allowed.
+	 */
+	FWidgetStateBitfield operator~() const;
+
+	/** 
+	 * Intersection of binary and enum states, with union on enum state usage so enum state mismatches can fail
+	 * Intersections will remove not allowed enum states, allowing for failure on bool conversion due to empty used states.
+	 * Approximation of bitwise '&'. Non-commutative, will use LHS enum state set if RHS does not use enum. Additionally,
+	 * if the RHS has binary states & the LHS does not meet those, all enum states will be marked unused to ensure failure
+	 */
+	FWidgetStateBitfield Intersect(const FWidgetStateBitfield& Rhs) const;
+
+	/** 
+	 * Union of binary and enum states, allows for a single enum state to have multiple values in a bitfield.
+	 * Unions will grow enum state set if allowance is same, otherwise the allowed set is maintained.
+	 * Approximation of bitwise '|'. Commutative.
+	 */
+	FWidgetStateBitfield Union(const FWidgetStateBitfield& Rhs) const;
+
+public:
+
+	/** True if any binary or enum state is set, and no empty used enum states exist */
+	operator bool() const;
+
+	/** True if any binary state is set */
+	bool HasBinaryStates() const;
+
+	/** True if any enum state is used, does not indicate if state is actually set. Use 'HasEmptyUsedEnumStates' for that. */
+	bool HasEnumStates() const;
+
+	/** True if any enum used enum states are empty, usually indicating a failed '&' test */
+	bool HasEmptyUsedEnumStates() const;
+
+public:
+
+	/**
+	 * Note: No need for HasNoFlags. Simply call !HasAnyFlags.
+	 */
+
+	/** True if any state flag is met */
+	bool HasAnyFlags(const FWidgetStateBitfield& InBitfield);
+
+	/** True if all state flags are met */
+	bool HasAllFlags(const FWidgetStateBitfield& InBitfield);
+
+	/** True if any binary state flag is met */
+	bool HasAnyBinaryFlags(const FWidgetStateBitfield& InBitfield);
+
+	/** True if all binary state flags are met */
+	bool HasAllBinaryFlags(const FWidgetStateBitfield& InBitfield);
+
+	/** True if any enum state flag is met */
+	bool HasAnyEnumFlags(const FWidgetStateBitfield& InBitfield);
+
+	/** True if all enum state flags are met */
+	bool HasAllEnumFlags(const FWidgetStateBitfield& InBitfield);
+
+public:
+
+	/** Set state to given value, just calls assignment operator for you */
+	void SetState(const FWidgetStateBitfield& InBitfield);
+
+	/** Negate both binary and enum states */
+	void NegateStates();
+
+	/**
+	 * Set binary state to given value
+	 *
+	 * @param BinaryStateIndex Index for Binary state to assign value to
+	 * @param BinaryStateValue Value to assign for given Binary state
+	 */
+	void SetBinaryState(uint8 BinaryStateIndex, bool BinaryStateValue);
+
+	/**
+	 * Set Binary state to given value, using a bitfield to indicate states that should be changed
+	 *
+	 * @param BinaryStateBitfield state bitfield, all marked states will be set with given value
+	 * @param BinaryStateValue Value to assign for given Binary state
+	 */
+	void SetBinaryState(const FWidgetStateBitfield& BinaryStateBitfield, bool BinaryStateValue);
+
+	/** Set Binary state to given value, slow version w/ FName map lookup */
+	void SetBinaryStateSlow(FName BinaryStateName, bool BinaryStateValue);
+
+	/** Negate only binary states */
+	void NegateBinaryStates();
+
+	/** 
+	 * Set enum state to given value, clearing all existing values for that state  
+	 * 
+	 * @param EnumStateIndex Index for enum state to assign value to
+	 * @param EnumStateValue Value to assign for given enum state
+	 */
+	void SetEnumState(uint8 EnumStateIndex, uint8 EnumStateValue);
+
+	/**
+	 * Set enum state to given value, clearing all existing values for states in use
+	 *
+	 * @param EnumStateBitfield state bitfield, all used states will copy-overwrite current enum state
+	 */
+	void SetEnumState(const FWidgetStateBitfield& EnumStateBitfield);
+
+	/** Set enum state to given value, clearing all existing values for that state, slow version w/ FName map lookup */
+	void SetEnumStateSlow(FName EnumStateName, uint8 EnumStateValue);
+
+	/**
+	 * Clear all existing values for given states & mark the states unused
+	 *
+	 * @param EnumStateBitfield state bitfield, all used states will be cleared
+	 */
+	void ClearEnumState(const FWidgetStateBitfield& EnumStateBitfield);
+
+	/**
+	 * Clear all existing values for given state & mark the state unused
+	 *
+	 * @param EnumStateIndex Index for enum state to mark as unused
+	 */
+	void ClearEnumState(uint8 EnumStateIndex);
+
+	/** Clear all existing values for given state & mark the state unused, slow version w/ map lookup */
+	void ClearEnumState(FName EnumStateName);
+
+	/** Negate only enum states, does not change usage, but marks existing values as not allowed */
+	void NegateEnumStates();
+
+private:
+
+	/** 
+	 * Note: This implementation is private, we choose to limit 64 binary and 16 enum states. 
+	 * We believe that is more than enum, but if not it's possible to use:
+	 * 
+	 * TBitArray<TInlineAllocator<32>> BinaryStates;
+	 * TArray<TSet<uint8, TInlineAllocator<1>>, TInlineAllocator<8>> EnumStates;
+	 * 
+	 * With some additional modifications to the .cpp impl.
+	 */
+
+	/** Binary states captured by this state field */
+	uint64 BinaryStates = 0;
+
+	/** 
+	 * An empty intersection of enum states may / may not be a reason to convert this field to true, so track states used 
+	 * In particular for '&' IsEnumStateUsed will still grow by '|'. Ex: ECheckboxState & EAssetViewType => false.
+	 */
+	uint16 IsEnumStateUsed = 0;
+
+	/**
+	 * A negation of enum states does not change usage, but indicates that states should be removed on intersection.
+	 * On union, if allowance is the same the state set is union'd, else the allowed set is maintained
+	 */
+	uint16 IsEnumStateNotAllowed = 0;
+
+	/** Binary states captured by this state field. Note: 16 == sizeof(IsEnumStateUsed) * CHAR_BIT */
+	TStaticArray<TSet<uint8, DefaultKeyFuncs<uint8, false>, TInlineSetAllocator<1>>, 16> EnumStates;
+};
\ No newline at end of file
diff --git a/Engine/Source/Runtime/UMG/Public/Binding/States/WidgetStateRegistration.h b/Engine/Source/Runtime/UMG/Public/Binding/States/WidgetStateRegistration.h
new file mode 100644
--- /dev/null
+++ b/Engine/Source/Runtime/UMG/Public/Binding/States/WidgetStateRegistration.h
@@ -0,0 +1,169 @@
+// Copyright Epic Games, Inc. All Rights Reserved.
+
+#pragma once
+
+#include "CoreMinimal.h"
+#include "UObject/Object.h"
+#include "UObject/SoftObjectPath.h"
+#include "Binding/States/WidgetStateBitfield.h"
+
+#include "WidgetStateRegistration.generated.h"
+
+class UWidget;
+class UWidgetStateSettings;
+
+/**
+ * Derive from to add a new widget binary state
+ * 
+ * Technically these can be created in BP, but for now we don't want to encourage
+ * that workflow as it involves requring overrides for the virtuals which is technical.
+ */
+UCLASS(Transient)
+class UMG_API UWidgetBinaryStateRegistration : public UObject
+{
+	GENERATED_BODY()
+
+public:
+	UWidgetBinaryStateRegistration() = default;
+	virtual ~UWidgetBinaryStateRegistration() = default;
+
+	/** Called once during WidgetStateSettings initialization to get this widget state's name */
+	virtual FName GetStateName() const;
+
+	/** Called on widget registration to correctly initialize widget state based on the current widget */
+	virtual bool GetRegisteredWidgetState(const UWidget* InWidget) const;
+
+protected:
+	friend UWidgetStateSettings;
+
+	/** Called to give CDO chance to initialize any static state bitfields that might be declared for convenience */
+	virtual void InitializeStaticBitfields() const;
+};
+
+UCLASS(Transient)
+class UMG_API UWidgetHoveredStateRegistration : public UWidgetBinaryStateRegistration
+{
+	GENERATED_BODY()
+
+public:
+
+	/** Post-load initialized bit corresponding to this binary state */
+	static inline FWidgetStateBitfield Bit;
+
+	static const inline FName StateName = FName("Hovered");
+
+	//~ Begin UWidgetBinaryStateRegistration Interface.
+	virtual FName GetStateName() const override;
+	virtual bool GetRegisteredWidgetState(const UWidget* InWidget) const override;
+	//~ End UWidgetBinaryStateRegistration Interface
+
+protected:
+	friend UWidgetStateSettings;
+
+	//~ Begin UWidgetBinaryStateRegistration Interface.
+	virtual void InitializeStaticBitfields() const override;
+	//~ End UWidgetBinaryStateRegistration Interface
+};
+
+UCLASS(Transient)
+class UMG_API UWidgetPressedStateRegistration : public UWidgetBinaryStateRegistration
+{
+	GENERATED_BODY()
+
+public:
+
+	/** Post-load initialized bit corresponding to this binary state */
+	static inline FWidgetStateBitfield Bit;
+
+	static const inline FName StateName = FName("Pressed");
+
+	//~ Begin UWidgetBinaryStateRegistration Interface.
+	virtual FName GetStateName() const override;
+	virtual bool GetRegisteredWidgetState(const UWidget* InWidget) const override;
+	//~ End UWidgetBinaryStateRegistration Interface
+
+protected:
+	friend UWidgetStateSettings;
+
+	//~ Begin UWidgetBinaryStateRegistration Interface.
+	virtual void InitializeStaticBitfields() const override;
+	//~ End UWidgetBinaryStateRegistration Interface
+};
+
+UCLASS(Transient)
+class UMG_API UWidgetDisabledStateRegistration : public UWidgetBinaryStateRegistration
+{
+	GENERATED_BODY()
+
+public:
+
+	/** Post-load initialized bit corresponding to this binary state */
+	static inline FWidgetStateBitfield Bit;
+
+	static const inline FName StateName = FName("Disabled");
+
+	//~ Begin UWidgetBinaryStateRegistration Interface.
+	virtual FName GetStateName() const override;
+	virtual bool GetRegisteredWidgetState(const UWidget* InWidget) const override;
+	//~ End UWidgetBinaryStateRegistration Interface
+
+protected:
+	friend UWidgetStateSettings;
+
+	//~ Begin UWidgetBinaryStateRegistration Interface.
+	virtual void InitializeStaticBitfields() const override;
+	//~ End UWidgetBinaryStateRegistration Interface
+};
+
+UCLASS(Transient)
+class UMG_API UWidgetSelectedStateRegistration : public UWidgetBinaryStateRegistration
+{
+	GENERATED_BODY()
+
+public:
+
+	/** Post-load initialized bit corresponding to this binary state */
+	static inline FWidgetStateBitfield Bit;
+
+	static const inline FName StateName = FName("Selected");
+
+	//~ Begin UWidgetBinaryStateRegistration Interface.
+	virtual FName GetStateName() const override;
+	virtual bool GetRegisteredWidgetState(const UWidget* InWidget) const override;
+	//~ End UWidgetBinaryStateRegistration Interface
+
+protected:
+	friend UWidgetStateSettings;
+
+	//~ Begin UWidgetBinaryStateRegistration Interface.
+	virtual void InitializeStaticBitfields() const override;
+	//~ End UWidgetBinaryStateRegistration Interface
+};
+
+/**
+ * Derive from to add a new Enum binary state
+ */
+UCLASS(Transient)
+class UMG_API UWidgetEnumStateRegistration : public UObject
+{
+	GENERATED_BODY()
+
+public:
+	UWidgetEnumStateRegistration() = default;
+	virtual ~UWidgetEnumStateRegistration() = default;
+
+	/** Called once during WidgetStateSettings initialization to get this widget state's name */
+	virtual FName GetStateName() const;
+
+	/** Called on widget registration to determine if this widget uses the given state */
+	virtual bool GetRegisteredWidgetUsesState(const UWidget* InWidget) const;
+
+	/** Called on widget registration to correctly initialize widget state based on the current widget */
+	virtual uint8 GetRegisteredWidgetState(const UWidget* InWidget) const;
+
+protected:
+	friend UWidgetStateSettings;
+
+	/** Called to give CDO chance to initialize any static state bitfields that might be declared for convenience */
+	virtual void InitializeStaticBitfields() const;
+};
diff --git a/Engine/Source/Runtime/UMG/Public/Binding/States/WidgetStateSettings.h b/Engine/Source/Runtime/UMG/Public/Binding/States/WidgetStateSettings.h
new file mode 100644
--- /dev/null
+++ b/Engine/Source/Runtime/UMG/Public/Binding/States/WidgetStateSettings.h
@@ -0,0 +1,85 @@
+// Copyright Epic Games, Inc. All Rights Reserved.
+
+#pragma once
+
+#include "CoreMinimal.h"
+#include "Containers/ContainerAllocationPolicies.h"
+#include "Engine/DeveloperSettings.h"
+#include "UObject/Object.h"
+#include "UObject/SoftObjectPath.h"
+
+#include "WidgetStateSettings.generated.h"
+
+class UWidget;
+
+struct FWidgetStateBitfield;
+class UWidgetBinaryStateRegistration;
+class UWidgetEnumStateRegistration;
+
+/**
+ * Settings used to map widget states to indexes.
+ * Does not perform any input validation, will crash if given invalid searches
+ * 
+ * Note: Currently doesn't really have any meaningful settings, used more as global singleton. May change.
+ */
+UCLASS(config = Game, defaultconfig)
+class UMG_API UWidgetStateSettings : public UDeveloperSettings
+{
+	GENERATED_BODY()
+
+public:
+		
+	static UWidgetStateSettings* Get() 
+	{ 
+		return GetMutableDefault<UWidgetStateSettings>();
+	}
+
+	//~ Begin UObject Interface.
+	virtual void PostInitProperties() override;
+	//~ End UObject Interface
+
+public:
+
+	/** Obtain a list of all known States. */
+	void GetAllStateNames(TArray<FName>& OutStateNames) const;
+
+	/** Obtain a list of all known Binary States. */
+	void GetBinaryStateNames(TArray<FName>& OutBinaryStateNames) const;
+
+	/** Obtain a list of all known Enum States. */
+	void GetEnumStateNames(TArray<FName>& OutEnumStateNames) const;
+
+	/** Get Index corresponding to this binary state, which may dynamically change - but not within a game session. */
+	uint8 GetBinaryStateIndex(const FName InBinaryStateName) const;
+
+	/** Get Index corresponding to this enum state, which may dynamically change - but not within a game session. */
+	uint8 GetEnumStateIndex(const FName InEnumStateName) const;
+
+	/** Get Name corresponding to this binary state index, which may dynamically change - but not within a game session. */
+	FName GetBinaryStateName(const uint8 InBinaryStateIndex) const;
+
+	/** Get Name corresponding to this enum state index, which may dynamically change - but not within a game session. */
+	FName GetEnumStateName(const uint8 InEnumStateIndex) const;
+
+	/** Gets bitfield of widget based on current state and given on registration state initializers */
+	FWidgetStateBitfield GetInitialRegistrationBitfield(const UWidget* InWidget) const;
+
+private:
+	/** Map of Binary State Names to index (In order of addition) */
+	TMap<FName, uint8, TInlineSetAllocator<32>> BinaryStateMap;
+
+	/** Map of Enum State Names to index (In order of addition) */
+	TMap<FName, uint8, TInlineSetAllocator<8>> EnumStateMap;
+
+	/** Ordered list of Binary Widget States */
+	TArray<FName, TInlineAllocator<32>> BinaryStates;
+
+	/** Ordered list of Enum Widget States */
+	TArray<FName, TInlineAllocator<8>> EnumStates;
+
+	/** Ordered list of Binary Widget States */
+	TArray<TObjectPtr<UWidgetBinaryStateRegistration>, TInlineAllocator<32>> BinaryStateRegistrationCDOs;
+
+	/** Ordered list of Enum Widget States */
+	TArray<TObjectPtr<UWidgetEnumStateRegistration>, TInlineAllocator<8>> EnumStateRegistrationCDOs;
+};
\ No newline at end of file
diff --git a/Engine/Source/Runtime/UMG/Public/Components/CheckBox.h b/Engine/Source/Runtime/UMG/Public/Components/CheckBox.h
--- a/Engine/Source/Runtime/UMG/Public/Components/CheckBox.h
+++ b/Engine/Source/Runtime/UMG/Public/Components/CheckBox.h
@@ -8,6 +8,7 @@
 #include "Styling/SlateColor.h"
 #include "Styling/SlateTypes.h"
 #include "Widgets/SWidget.h"
+#include "Binding/States/WidgetStateRegistration.h"
 #include "Components/ContentWidget.h"
 
 #include "CheckBox.generated.h"
@@ -178,3 +179,33 @@ PRAGMA_DISABLE_DEPRECATION_WARNINGS
 	PROPERTY_BINDING_IMPLEMENTATION(ECheckBoxState, CheckedState)
 PRAGMA_ENABLE_DEPRECATION_WARNINGS
 };
+
+UCLASS(Transient)
+class UMG_API UWidgetCheckedStateRegistration : public UWidgetEnumStateRegistration
+{
+	GENERATED_BODY()
+
+public:
+
+	/** Post-load initialized values corresponding to this enum state */
+	static inline FWidgetStateBitfield Unchecked;
+	static inline FWidgetStateBitfield Checked;
+	static inline FWidgetStateBitfield Undetermined;
+
+	static const inline FName StateName = FName("CheckedState");
+
+	//~ Begin UWidgetEnumStateRegistration Interface.
+	virtual FName GetStateName() const override;
+	virtual uint8 GetRegisteredWidgetState(const UWidget* InWidget) const override;
+	//~ End UWidgetEnumStateRegistration Interface
+
+	/** Convenience method to get widget state bitfield from enum value */
+	static const FWidgetStateBitfield& GetBitfieldFromValue(uint8 InValue);
+
+protected:
+	friend UWidgetStateSettings;
+
+	//~ Begin UWidgetEnumStateRegistration Interface.
+	virtual void InitializeStaticBitfields() const override;
+	//~ End UWidgetEnumStateRegistration Interface
+};
\ No newline at end of file
diff --git a/Engine/Source/Runtime/UMG/Public/Components/Widget.h b/Engine/Source/Runtime/UMG/Public/Components/Widget.h
--- a/Engine/Source/Runtime/UMG/Public/Components/Widget.h
+++ b/Engine/Source/Runtime/UMG/Public/Components/Widget.h
@@ -41,6 +41,7 @@ class UPanelSlot;
 class UPropertyBinding;
 class UUserWidget;
 struct FDynamicPropertyPath;
+struct FWidgetStateBitfield;
 enum class ECheckBoxState : uint8;
 
 namespace UMWidget
@@ -251,6 +252,8 @@ public:
 	DECLARE_DYNAMIC_DELEGATE_RetVal(FEventReply, FOnReply);
 	DECLARE_DYNAMIC_DELEGATE_RetVal_TwoParams(FEventReply, FOnPointerEvent, FGeometry, MyGeometry, const FPointerEvent&, MouseEvent);
 
+	DECLARE_MULTICAST_DELEGATE_OneParam(FOnWidgetStateBroadcast, const FWidgetStateBitfield& /*InStateBitfield */);
+
 	typedef TFunctionRef<TSharedPtr<SObjectWidget>( UUserWidget*, TSharedRef<SWidget> )> ConstructMethodType;
 
 	/**
@@ -975,6 +978,23 @@ public:
 	/**  */
 	bool AddBinding(FDelegateProperty* DelegateProperty, UObject* SourceObject, const FDynamicPropertyPath& BindingPath);
 
+	/**
+	 * Add a post-state-changed listener to this widget, will fire after a state changed and all related side effects are resolved.
+	 * 
+	 * Note: Currently we only support post-state-changed broadcasts.
+	 *
+	 * @param ListenerDelegate Delegate to fire when state changes
+	 * @param bBroadcastCurrentState true if we should trigger this delegate once on registration with the current widget state (Does not globally broadcast).
+	 */
+	FDelegateHandle RegisterPostStateListener(const FOnWidgetStateBroadcast::FDelegate& ListenerDelegate, bool bBroadcastCurrentState = true);
+
+	/**
+	 * Remove a post-state-changed  listener from this widget, resets state bitfield if no other state listeners exist
+	 *
+	 * @param ListenerDelegate Delegate to remove
+	 */
+	void UnregisterPostStateListener(const FDelegateHandle& ListenerDelegate);
+
 	static TSubclassOf<UPropertyBinding> FindBinderClassForDestination(FProperty* Property);
 
 	//~ Begin UObject
@@ -984,6 +1004,10 @@ public:
 	virtual void PreSave(FObjectPreSaveContext ObjectSaveContext) override;
 	//~ End UObject
 
+	//~ Begin UVisual
+	virtual void ReleaseSlateResources(bool bReleaseChildren) override;
+	//~ End UVisual
+
 	FORCEINLINE bool CanSafelyRouteEvent()
 	{
 		return !IsDesignTime() && CanSafelyRouteCall();
@@ -1074,6 +1098,21 @@ protected:
 
 	virtual void OnBindingChanged(const FName& Property);
 
+	/**
+	 * Broadcast a binary state post change
+	 *
+	 * @param StateChange bitfield marking states that should be changed
+	 * @param bInValue true if marked states should be enabled, false otherwise
+	 */
+	void BroadcastBinaryPostStateChange(const FWidgetStateBitfield& StateChange, bool bInValue);
+
+	/**
+	 * Broadcast an enum state post change
+	 *
+	 * @param StateChange bitfield marking states that should be changed
+	 */
+	void BroadcastEnumPostStateChange(const FWidgetStateBitfield& StateChange);
+
 protected:
 	UObject* GetSourceAssetOrClass() const;
 
@@ -1130,6 +1169,15 @@ protected:
 	/** The underlying SWidget contained in a SObjectWidget */
 	TWeakPtr<SObjectWidget> MyGCWidget;
 
+	/** The bitfield for this widget's state */
+	TSharedPtr<FWidgetStateBitfield> MyWidgetStateBitfield;
+
+	/** False will skip state broadcasts. Useful for child classes to call Super methods without broadcasting early / late. */
+	bool bShouldBroadcastState;
+
+	/** Delegate that broadcasts after current widget state has fully changed, including all state-related side effects */
+	FOnWidgetStateBroadcast PostWidgetStateChanged;
+
 	/** Native property bindings. */
 	UPROPERTY(Transient)
 	TArray<TObjectPtr<UPropertyBinding>> NativeBindings;
diff --git a/Engine/Source/Runtime/UMG/UMG.Build.cs b/Engine/Source/Runtime/UMG/UMG.Build.cs
--- a/Engine/Source/Runtime/UMG/UMG.Build.cs
+++ b/Engine/Source/Runtime/UMG/UMG.Build.cs
@@ -10,6 +10,7 @@ public class UMG : ModuleRules
 			new string[] {
 				"Core",
 				"CoreUObject",
+				"DeveloperSettings",
                 "Engine",
                 "InputCore",
 				"Slate",
