Fix Nonunity WidgetStateBitfieldTests

diff --git a/Engine/Source/Runtime/UMG/Private/Binding/States/Tests/WidgetStateBitfieldTests.cpp b/Engine/Source/Runtime/UMG/Private/Binding/States/Tests/WidgetStateBitfieldTests.cpp
--- a/Engine/Source/Runtime/UMG/Private/Binding/States/Tests/WidgetStateBitfieldTests.cpp
+++ b/Engine/Source/Runtime/UMG/Private/Binding/States/Tests/WidgetStateBitfieldTests.cpp
@@ -1,5 +1,8 @@
 // Copyright Epic Games, Inc. All Rights Reserved.
 
+#include "CoreTypes.h"
+#include "Logging/LogMacros.h"
+#include "Misc/AutomationTest.h"
 #include "Binding/States/WidgetStateSettings.h"
 #include "Binding/States/WidgetStateBitfield.h"
 #include "Binding/States/WidgetStateRegistration.h"
