PR #10257: Bug Fix - ITypedUMGListView Improper nullptr conversion (Contributed by [REDACTED_PERSON])

diff --git a/Engine/Source/Runtime/UMG/Public/Components/ListViewBase.h b/Engine/Source/Runtime/UMG/Public/Components/ListViewBase.h
--- a/Engine/Source/Runtime/UMG/Public/Components/ListViewBase.h
+++ b/Engine/Source/Runtime/UMG/Public/Components/ListViewBase.h
@@ -107,7 +107,7 @@ public:
 				return SelectedItems[0];
 			}
 		}
-		return nullptr;
+		return TListTypeTraits<ItemType>::MakeNullPtr();
 	}
 
 	const ItemType* ItemFromEntryWidget(const UUserWidget& EntryWidget) const
