Rich text - updating SetDefaultStrikeBrush to take a const reference. Fixes blueprint call to the function to offer correct input parameter type

diff --git a/Engine/Source/Runtime/UMG/Private/Components/RichTextBlock.cpp b/Engine/Source/Runtime/UMG/Private/Components/RichTextBlock.cpp
--- a/Engine/Source/Runtime/UMG/Private/Components/RichTextBlock.cpp
+++ b/Engine/Source/Runtime/UMG/Private/Components/RichTextBlock.cpp
@@ -385,7 +385,7 @@ void URichTextBlock::SetDefaultFont(FSlateFontInfo InFontInfo)
 	ApplyUpdatedDefaultTextStyle();
 }
 
-void URichTextBlock::SetDefaultStrikeBrush(FSlateBrush& InStrikeBrush)
+void URichTextBlock::SetDefaultStrikeBrush(const FSlateBrush& InStrikeBrush)
 {
 	BeginDefaultStyleOverride();
 	DefaultTextStyleOverride.StrikeBrush = InStrikeBrush;
diff --git a/Engine/Source/Runtime/UMG/Public/Components/RichTextBlock.h b/Engine/Source/Runtime/UMG/Public/Components/RichTextBlock.h
--- a/Engine/Source/Runtime/UMG/Public/Components/RichTextBlock.h
+++ b/Engine/Source/Runtime/UMG/Public/Components/RichTextBlock.h
@@ -113,7 +113,7 @@ public:
 	 * @param InStrikeBrush The new brush to use to strike through text
 	 */
 	UFUNCTION(BlueprintCallable, Category = "Appearance")
-	void SetDefaultStrikeBrush(FSlateBrush& InStrikeBrush);
+	void SetDefaultStrikeBrush(const FSlateBrush& InStrikeBrush);
 
 	/**
 	 *  Set the minimum desired width for this rich text block
