UMG: Fix slot caching in the viewport subsystem. In the BP Construct a user can set the size and the ViewportWidgets was not updated yet.

diff --git a/Engine/Source/Runtime/UMG/Private/Blueprint/GameViewportSubsystem.cpp b/Engine/Source/Runtime/UMG/Private/Blueprint/GameViewportSubsystem.cpp
--- a/Engine/Source/Runtime/UMG/Private/Blueprint/GameViewportSubsystem.cpp
+++ b/Engine/Source/Runtime/UMG/Private/Blueprint/GameViewportSubsystem.cpp
@@ -146,42 +146,45 @@ void UGameViewportSubsystem::AddToScreen(UWidget* Widget, ULocalPlayer* Player,
 		return;
 	}
 
-	FSlotInfo& SlotInfo = ViewportWidgets.FindOrAdd(Widget);
-	Widget->bIsManagedByGameViewportSubsystem = true;
-
-	if (SlotInfo.FullScreenWidget.IsValid())
+	SConstraintCanvas::FSlot* RawSlot = nullptr;
+	TSharedPtr<SConstraintCanvas> FullScreenCanvas;
 	{
-		FFrame::KismetExecutionMessage(*FString::Printf(TEXT("The widget '%s' was already added to the screen."), *Widget->GetName()), ELogVerbosity::Warning);
-		return;
+		FSlotInfo& SlotInfo = ViewportWidgets.FindOrAdd(Widget);
+		Widget->bIsManagedByGameViewportSubsystem = true;
+
+		if (SlotInfo.FullScreenWidget.IsValid())
+		{
+			FFrame::KismetExecutionMessage(*FString::Printf(TEXT("The widget '%s' was already added to the screen."), *Widget->GetName()), ELogVerbosity::Warning);
+			return;
+		}
+
+		TPair<FMargin, bool> OffsetArgument = UE::UMG::Private::CalculateOffsetArgument(Slot);
+		FullScreenCanvas = SNew(SConstraintCanvas)
+			+ SConstraintCanvas::Slot()
+			.Offset(OffsetArgument.Get<0>())
+			.AutoSize(OffsetArgument.Get<1>())
+			.Anchors(Slot.Anchors)
+			.Alignment(Slot.Alignment)
+			.Expose(RawSlot);
+
+		SlotInfo.Slot = Slot;
+		SlotInfo.LocalPlayer = Player;
+		SlotInfo.FullScreenWidget = FullScreenCanvas;
+		SlotInfo.FullScreenWidgetSlot = RawSlot;
 	}
 
-	TPair<FMargin, bool> OffsetArgument = UE::UMG::Private::CalculateOffsetArgument(Slot);
-	SConstraintCanvas::FSlot* RawSlot = nullptr;
-	TSharedRef<SConstraintCanvas> FullScreenCanvas = SNew(SConstraintCanvas)
-		+ SConstraintCanvas::Slot()
-		.Offset(OffsetArgument.Get<0>())
-		.AutoSize(OffsetArgument.Get<1>())
-		.Anchors(Slot.Anchors)
-		.Alignment(Slot.Alignment)
-		.Expose(RawSlot)
-		[
-			Widget->TakeWidget()
-		];
-
-	SlotInfo.Slot = Slot;
-	SlotInfo.LocalPlayer = Player;
-	SlotInfo.FullScreenWidget = FullScreenCanvas;
-	SlotInfo.FullScreenWidgetSlot = RawSlot;
+	check(RawSlot);
+	RawSlot->AttachWidget(Widget->TakeWidget());
 
 	if (Player)
 	{
-		ViewportClient->AddViewportWidgetForPlayer(Player, FullScreenCanvas, Slot.ZOrder);
+		ViewportClient->AddViewportWidgetForPlayer(Player, FullScreenCanvas.ToSharedRef(), Slot.ZOrder);
 	}
 	else
 	{
 		// We add 10 to the zorder when adding to the viewport to avoid 
 		// displaying below any built-in controls, like the virtual joysticks on mobile builds.
-		ViewportClient->AddViewportWidgetContent(FullScreenCanvas, Slot.ZOrder + 10);
+		ViewportClient->AddViewportWidgetContent(FullScreenCanvas.ToSharedRef(), Slot.ZOrder + 10);
 	}
 	OnWidgetAdded.Broadcast(Widget, Player);
 }
