PR #8949: Faster index fetching for listview items (Contributed by [REDACTED_PERSON])

diff --git a/Engine/Source/Runtime/UMG/Private/Components/ListView.cpp b/Engine/Source/Runtime/UMG/Private/Components/ListView.cpp
--- a/Engine/Source/Runtime/UMG/Private/Components/ListView.cpp
+++ b/Engine/Source/Runtime/UMG/Private/Components/ListView.cpp
@@ -131,7 +131,8 @@ int32 UListView::GetNumItems() const
 
 int32 UListView::GetIndexForItem(const UObject* Item) const
 {
-	return ListItems.IndexOfByKey(Item);
+	TSharedPtr<SObjectTableRow<UObject*>> TableRow = GetObjectRowFromItem<SObjectTableRow>(const_cast<UObject*>(Item));
+	return TableRow.IsValid() ? TableRow->GetIndexInList() : INDEX_NONE;
 }
 
 void UListView::ClearListItems()
