Added the way for UListViewBase (and it's descendants) to be acknowledged about listEntry was generated or released.
See the thread for details: https://epicgames.slack.com/archives/C045T964405/p1680826364199749

diff --git a/Engine/Source/Runtime/UMG/Private/Components/ListViewBase.cpp b/Engine/Source/Runtime/UMG/Private/Components/ListViewBase.cpp
--- a/Engine/Source/Runtime/UMG/Private/Components/ListViewBase.cpp
+++ b/Engine/Source/Runtime/UMG/Private/Components/ListViewBase.cpp
@@ -204,6 +204,7 @@ void UListViewBase::HandleRowReleased(const TSharedRef<ITableRow>& Row)
 		{
 			GeneratedEntriesToAnnounce.Remove(EntryWidget);
 			OnEntryWidgetReleased().Broadcast(*EntryWidget);
+			NativeOnEntryReleased(EntryWidget);
 			BP_OnEntryReleased.Broadcast(EntryWidget);
 		}
 	}
@@ -234,6 +235,7 @@ void UListViewBase::HandleAnnounceGeneratedEntries()
 		if (EntryWidget.IsValid())
 		{
 			OnEntryWidgetGenerated().Broadcast(*EntryWidget);
+			NativeOnEntryGenerated(EntryWidget.Get());
 			BP_OnEntryGenerated.Broadcast(EntryWidget.Get());
 		}
 	}
diff --git a/Engine/Source/Runtime/UMG/Public/Components/ListViewBase.h b/Engine/Source/Runtime/UMG/Public/Components/ListViewBase.h
--- a/Engine/Source/Runtime/UMG/Public/Components/ListViewBase.h
+++ b/Engine/Source/Runtime/UMG/Public/Components/ListViewBase.h
@@ -568,6 +568,7 @@ protected:
     /** Called when a row widget is generated for a list item */
     UPROPERTY(BlueprintAssignable, Category = Events, meta = (DisplayName = "On Entry Generated"))
     FOnListEntryGeneratedDynamic BP_OnEntryGenerated;
+	virtual void NativeOnEntryGenerated(UUserWidget* EntryWidget) {};
 
 	/**
 	* Normally these are processed by UListViewBase::FinishGeneratingEntry which uses World->GetTimerManager() to generate entries next frame
@@ -688,12 +689,13 @@ protected:
 	UPROPERTY(EditAnywhere, BlueprintReadOnly, Category = Input)
 	bool bAllowDragging = true;
 
-private:
-	virtual void HandleAnnounceGeneratedEntries();
-
 	/** Called when a row widget is released by the list (i.e. when it no longer represents a list item) */
 	UPROPERTY(BlueprintAssignable, Category = Events, meta = (DisplayName = "On Entry Released"))
 	FOnListEntryReleasedDynamic BP_OnEntryReleased;
+	virtual void NativeOnEntryReleased(UUserWidget* EntryWidget) {};
+
+private:
+	virtual void HandleAnnounceGeneratedEntries();
 
 #if WITH_EDITORONLY_DATA
 	bool bNeedsToCallRefreshDesignerItems = false;;
