Slate: Add the widget table name name to track crash.

diff --git a/Engine/Source/Runtime/UMG/Public/Slate/SObjectTableRow.h b/Engine/Source/Runtime/UMG/Public/Slate/SObjectTableRow.h
--- a/Engine/Source/Runtime/UMG/Public/Slate/SObjectTableRow.h
+++ b/Engine/Source/Runtime/UMG/Public/Slate/SObjectTableRow.h
@@ -4,6 +4,7 @@
 
 #include "Slate/SObjectWidget.h"
 #include "Blueprint/IUserObjectListEntry.h"
+#include "Types/ReflectionMetadata.h"
 #include "Widgets/Views/STableRow.h"
 #include "Widgets/Views/STableViewBase.h"
 #include "Widgets/Text/STextBlock.h"
@@ -499,7 +500,7 @@ protected:
 		}
 		else
 		{
-			checkf(OwnerTable->Private_IsPendingRefresh(), TEXT("We were unable to find the item for this widget.  If it was removed from the source collection, the list should be pending a refresh."));
+			checkf(OwnerTable->Private_IsPendingRefresh(), TEXT("We were unable to find the item for this widget. If it was removed from the source collection, the list should be pending a refresh. %s"), *FReflectionMetaData::GetWidgetPath(this, false, false));
 		}
 
 		return nullptr;
