Widget State: Add missing const qualifier to necessary member function of WidgetStateBitfield

diff --git a/Engine/Source/Runtime/UMG/Private/Binding/States/WidgetStateBitfield.cpp b/Engine/Source/Runtime/UMG/Private/Binding/States/WidgetStateBitfield.cpp
--- a/Engine/Source/Runtime/UMG/Private/Binding/States/WidgetStateBitfield.cpp
+++ b/Engine/Source/Runtime/UMG/Private/Binding/States/WidgetStateBitfield.cpp
@@ -164,27 +164,27 @@ bool FWidgetStateBitfield::HasEmptyUsedEnumStates() const
 	return false;
 }
 
-bool FWidgetStateBitfield::HasAnyFlags(const FWidgetStateBitfield& InBitfield)
+bool FWidgetStateBitfield::HasAnyFlags(const FWidgetStateBitfield& InBitfield) const
 {
 	return HasAnyBinaryFlags(InBitfield) || HasAnyEnumFlags(InBitfield);
 }
 
-bool FWidgetStateBitfield::HasAllFlags(const FWidgetStateBitfield& InBitfield)
+bool FWidgetStateBitfield::HasAllFlags(const FWidgetStateBitfield& InBitfield) const
 {
 	return HasAllBinaryFlags(InBitfield) && HasAllEnumFlags(InBitfield);
 }
 
-bool FWidgetStateBitfield::HasAnyBinaryFlags(const FWidgetStateBitfield& InBitfield)
+bool FWidgetStateBitfield::HasAnyBinaryFlags(const FWidgetStateBitfield& InBitfield) const
 {
 	return (BinaryStates & InBitfield.BinaryStates) != 0;
 }
 
-bool FWidgetStateBitfield::HasAllBinaryFlags(const FWidgetStateBitfield& InBitfield)
+bool FWidgetStateBitfield::HasAllBinaryFlags(const FWidgetStateBitfield& InBitfield) const
 {
 	return (BinaryStates & InBitfield.BinaryStates) == InBitfield.BinaryStates;
 }
 
-bool FWidgetStateBitfield::HasAnyEnumFlags(const FWidgetStateBitfield& InBitfield)
+bool FWidgetStateBitfield::HasAnyEnumFlags(const FWidgetStateBitfield& InBitfield) const
 {
 	const uint16 IsEnumStateShared = IsEnumStateUsed & InBitfield.IsEnumStateUsed;
 
@@ -204,7 +204,7 @@ bool FWidgetStateBitfield::HasAnyEnumFlags(const FWidgetStateBitfield& InBitfiel
 	return false;
 }
 
-bool FWidgetStateBitfield::HasAllEnumFlags(const FWidgetStateBitfield& InBitfield)
+bool FWidgetStateBitfield::HasAllEnumFlags(const FWidgetStateBitfield& InBitfield) const
 {
 	const uint16 IsEnumStateShared = IsEnumStateUsed & InBitfield.IsEnumStateUsed;
 
diff --git a/Engine/Source/Runtime/UMG/Public/Binding/States/WidgetStateBitfield.h b/Engine/Source/Runtime/UMG/Public/Binding/States/WidgetStateBitfield.h
--- a/Engine/Source/Runtime/UMG/Public/Binding/States/WidgetStateBitfield.h
+++ b/Engine/Source/Runtime/UMG/Public/Binding/States/WidgetStateBitfield.h
@@ -77,22 +77,22 @@ public:
 	 */
 
 	/** True if any state flag is met */
-	bool HasAnyFlags(const FWidgetStateBitfield& InBitfield);
+	bool HasAnyFlags(const FWidgetStateBitfield& InBitfield) const;
 
 	/** True if all state flags are met */
-	bool HasAllFlags(const FWidgetStateBitfield& InBitfield);
+	bool HasAllFlags(const FWidgetStateBitfield& InBitfield) const;
 
 	/** True if any binary state flag is met */
-	bool HasAnyBinaryFlags(const FWidgetStateBitfield& InBitfield);
+	bool HasAnyBinaryFlags(const FWidgetStateBitfield& InBitfield) const;
 
 	/** True if all binary state flags are met */
-	bool HasAllBinaryFlags(const FWidgetStateBitfield& InBitfield);
+	bool HasAllBinaryFlags(const FWidgetStateBitfield& InBitfield) const;
 
 	/** True if any enum state flag is met */
-	bool HasAnyEnumFlags(const FWidgetStateBitfield& InBitfield);
+	bool HasAnyEnumFlags(const FWidgetStateBitfield& InBitfield) const;
 
 	/** True if all enum state flags are met */
-	bool HasAllEnumFlags(const FWidgetStateBitfield& InBitfield);
+	bool HasAllEnumFlags(const FWidgetStateBitfield& InBitfield) const;
 
 public:
 
