Widget State: Add widget to the FOnWidgetStateBroadcast

diff --git a/Engine/Source/Runtime/UMG/Private/Components/Widget.cpp b/Engine/Source/Runtime/UMG/Private/Components/Widget.cpp
--- a/Engine/Source/Runtime/UMG/Private/Components/Widget.cpp
+++ b/Engine/Source/Runtime/UMG/Private/Components/Widget.cpp
@@ -1802,7 +1802,7 @@ FDelegateHandle UWidget::RegisterPostStateListener(const FOnWidgetStateBroadcast
 
 	if (bBroadcastCurrentState)
 	{
-		ListenerDelegate.ExecuteIfBound(*MyWidgetStateBitfield);
+		ListenerDelegate.ExecuteIfBound(this, *MyWidgetStateBitfield);
 	}
 
 	return PostWidgetStateChanged.Add(ListenerDelegate);
@@ -1828,7 +1828,7 @@ void UWidget::BroadcastBinaryPostStateChange(const FWidgetStateBitfield& StateCh
 	if (bShouldBroadcastState && MyWidgetStateBitfield.IsValid())
 	{
 		MyWidgetStateBitfield->SetBinaryState(StateChange, bInValue);
-		PostWidgetStateChanged.Broadcast(*MyWidgetStateBitfield);
+		PostWidgetStateChanged.Broadcast(this, *MyWidgetStateBitfield);
 	}
 }
 
@@ -1837,7 +1837,7 @@ void UWidget::BroadcastEnumPostStateChange(const FWidgetStateBitfield& StateChan
 	if (bShouldBroadcastState && MyWidgetStateBitfield.IsValid())
 	{
 		MyWidgetStateBitfield->SetEnumState(StateChange);
-		PostWidgetStateChanged.Broadcast(*MyWidgetStateBitfield);
+		PostWidgetStateChanged.Broadcast(this, *MyWidgetStateBitfield);
 	}
 }
 
diff --git a/Engine/Source/Runtime/UMG/Public/Components/Widget.h b/Engine/Source/Runtime/UMG/Public/Components/Widget.h
--- a/Engine/Source/Runtime/UMG/Public/Components/Widget.h
+++ b/Engine/Source/Runtime/UMG/Public/Components/Widget.h
@@ -252,7 +252,7 @@ public:
 	DECLARE_DYNAMIC_DELEGATE_RetVal(FEventReply, FOnReply);
 	DECLARE_DYNAMIC_DELEGATE_RetVal_TwoParams(FEventReply, FOnPointerEvent, FGeometry, MyGeometry, const FPointerEvent&, MouseEvent);
 
-	DECLARE_MULTICAST_DELEGATE_OneParam(FOnWidgetStateBroadcast, const FWidgetStateBitfield& /*InStateBitfield */);
+	DECLARE_MULTICAST_DELEGATE_TwoParams(FOnWidgetStateBroadcast, UWidget* /*InWidget*/, const FWidgetStateBitfield& /*InStateBitfield*/);
 
 	typedef TFunctionRef<TSharedPtr<SObjectWidget>( UUserWidget*, TSharedRef<SWidget> )> ConstructMethodType;
 
