[25.00] Add OnIsItemSelectableOrNavigable delegate to list views.

UListViewBase:
- Added pure virtual OnIsItemSelectableOrNavigable function.
- OnIsSelectableOrNavigableInternal checks if OnIsItemSelectableOrNavigable is bound and returns its value if so.

IMPLEMENT_TYPED_UMG_LIST:
- Added OnIsItemSelectableOrNavigableDelegate.
- Overrides OnIsItemSelectableOrNavigable and returns OnIsItemSelectableOrNavigableDelegate.

diff --git a/Engine/Source/Runtime/UMG/Public/Components/ListViewBase.h b/Engine/Source/Runtime/UMG/Public/Components/ListViewBase.h
--- a/Engine/Source/Runtime/UMG/Public/Components/ListViewBase.h
+++ b/Engine/Source/Runtime/UMG/Public/Components/ListViewBase.h
@@ -67,6 +67,9 @@ public:
 
 	virtual TSubclassOf<UUserWidget> GetDefaultEntryClass() const = 0;
 
+	DECLARE_DELEGATE_RetVal_OneParam(bool, FOnIsItemSelectableOrNavigable, ItemType);
+	virtual FOnIsItemSelectableOrNavigable& OnIsItemSelectableOrNavigable() const = 0;
+
 protected:
 	virtual SListView<ItemType>* GetMyListView() const = 0;
 	virtual uint32 GetOwningUserIndex() const = 0;
@@ -374,7 +377,7 @@ protected:
 	virtual void OnItemClickedInternal(ItemType Item) {}
 	virtual void OnItemDoubleClickedInternal(ItemType Item) {}
 	virtual void OnSelectionChangedInternal(NullableItemType FirstSelectedItem) {}
-	virtual bool OnIsSelectableOrNavigableInternal(ItemType FirstSelectedItem) { return true; }
+	virtual bool OnIsSelectableOrNavigableInternal(ItemType FirstSelectedItem) { return OnIsItemSelectableOrNavigable().IsBound() ? OnIsItemSelectableOrNavigable().Execute(FirstSelectedItem) : true; }
 	virtual void OnItemScrolledIntoViewInternal(ItemType Item, UUserWidget& EntryWidget) {}
 	virtual void OnListViewScrolledInternal(float ItemOffset, float DistanceRemaining) {}
 	virtual void OnItemExpansionChangedInternal(ItemType Item, bool bIsExpanded) {}
@@ -740,6 +743,7 @@ private:	\
 	mutable FOnFinishedScrolling OnFinishedScrollingEvent;	\
 	mutable FOnItemExpansionChanged OnItemExpansionChangedEvent;	\
 	mutable FOnGetEntryClassForItem OnGetEntryClassForItemDelegate;	\
+	mutable FOnIsItemSelectableOrNavigable OnIsItemSelectableOrNavigableDelegate; \
 public:	\
 	virtual TSubclassOf<UUserWidget> GetDefaultEntryClass() const override { return EntryWidgetClass; }	\
 	virtual FSimpleListItemEvent& OnItemClicked() const override { return OnItemClickedEvent; }	\
@@ -750,4 +754,5 @@ public:	\
 	virtual FOnListViewScrolled& OnListViewScrolled() const override { return OnListViewScrolledEvent; }	\
 	virtual FOnFinishedScrolling& OnFinishedScrolling() const override { return OnFinishedScrollingEvent; }	\
 	virtual FOnItemExpansionChanged& OnItemExpansionChanged() const override { return OnItemExpansionChangedEvent; }	\
-	virtual FOnGetEntryClassForItem& OnGetEntryClassForItem() const override { return OnGetEntryClassForItemDelegate; }
\ No newline at end of file
+	virtual FOnGetEntryClassForItem& OnGetEntryClassForItem() const override { return OnGetEntryClassForItemDelegate; } \
+	virtual FOnIsItemSelectableOrNavigable& OnIsItemSelectableOrNavigable() const override { return OnIsItemSelectableOrNavigableDelegate; }
\ No newline at end of file
