UImage SetBrush fix: ensure we call the base class setter when setting the brush in the body of UImage.

diff --git a/Engine/Source/Runtime/UMG/Private/Components/Image.cpp b/Engine/Source/Runtime/UMG/Private/Components/Image.cpp
--- a/Engine/Source/Runtime/UMG/Private/Components/Image.cpp
+++ b/Engine/Source/Runtime/UMG/Private/Components/Image.cpp
@@ -85,7 +85,7 @@ void UImage::SetOpacity(float InOpacity)
 const FSlateBrush* UImage::ConvertImage(TAttribute<FSlateBrush> InImageAsset) const
 {
 	UImage* MutableThis = const_cast<UImage*>( this );
-	MutableThis->SetBrush(InImageAsset.Get());
+	MutableThis->UImage::SetBrush(InImageAsset.Get());
 
 	return &Brush;
 }
@@ -159,11 +159,11 @@ void UImage::SetBrushFromAsset(USlateBrushAsset* Asset)
 		CancelImageStreaming();
 		if (Asset)
 		{
-			SetBrush(Asset->Brush);
+			UImage::SetBrush(Asset->Brush);
 		}
 		else
 		{
-			SetBrush(FSlateBrush());
+			UImage::SetBrush(FSlateBrush());
 		}
 	}
 }
