Fix safezone offset regression

diff --git a/Engine/Source/Runtime/UMG/Private/Components/SafeZone.cpp b/Engine/Source/Runtime/UMG/Private/Components/SafeZone.cpp
--- a/Engine/Source/Runtime/UMG/Private/Components/SafeZone.cpp
+++ b/Engine/Source/Runtime/UMG/Private/Components/SafeZone.cpp
@@ -103,7 +103,18 @@ void USafeZone::SetSidesToPad(bool InPadLeft, bool InPadRight, bool InPadTop, bo
 
 TSharedRef<SWidget> USafeZone::RebuildWidget()
 {	
+	USafeZoneSlot* SafeSlot = nullptr;
+	if (GetChildrenCount() > 0)
+	{
+		SafeSlot = Cast<USafeZoneSlot>(GetContentSlot());;
+	}
+
 	MySafeZone = SNew( SSafeZone )
+		.IsTitleSafe(SafeSlot ? SafeSlot->bIsTitleSafe : false)
+		.SafeAreaScale(SafeSlot ? SafeSlot->SafeAreaScale : FMargin(1, 1, 1, 1))
+		.HAlign(SafeSlot ? SafeSlot->HAlign.GetValue() : HAlign_Fill)
+		.VAlign(SafeSlot ? SafeSlot->VAlign.GetValue() : VAlign_Fill)
+		.Padding(SafeSlot ? SafeSlot->Padding : FMargin())
 		.PadLeft( PadLeft )
 		.PadRight( PadRight )
 		.PadTop( PadTop )
@@ -116,12 +127,9 @@ TSharedRef<SWidget> USafeZone::RebuildWidget()
 			GetChildAt( 0 ) ? GetChildAt( 0 )->TakeWidget() : SNullWidget::NullWidget
 		];
 
-	if (GetChildrenCount() > 0)
+	if (SafeSlot)
 	{
-		if (USafeZoneSlot* SafeSlot = Cast<USafeZoneSlot>(GetContentSlot()))
-		{
-			SafeSlot->BuildSlot(MySafeZone.ToSharedRef());
-		}
+		SafeSlot->BuildSlot(MySafeZone.ToSharedRef());
 	}
 
 	return MySafeZone.ToSharedRef();
