UMG: Fixed calling stopping animations inside Destruct calls

  - When Stop is called for an animation on a widget inside its OnDestruct (or NativeDestruct) function, we have to queue up an additional update at time=0 for the animation (for legacy design reasons).
  - This final update was getting thrown away inside UUMGAnimationTickManager because the widget was no longer constructed.
  - We now special case this codepath to allow animations to remain alive on destructed widgets if they are currently being stopped (ie IsStopping() == true). This allows them to finish correctly.

diff --git a/Engine/Source/Runtime/UMG/Private/Animation/UMGSequenceTickManager.cpp b/Engine/Source/Runtime/UMG/Private/Animation/UMGSequenceTickManager.cpp
--- a/Engine/Source/Runtime/UMG/Private/Animation/UMGSequenceTickManager.cpp
+++ b/Engine/Source/Runtime/UMG/Private/Animation/UMGSequenceTickManager.cpp
@@ -153,14 +153,18 @@ void UUMGSequenceTickManager::TickWidgetAnimations(float DeltaSeconds)
 			{
 				if (!bIsCurrentlyEvaluating)
 				{
-					// Wait until the runner has finished evaluating before we tear stuff down
-					UserWidget->TearDownAnimations();
+					// Tear down any animations that are not currently being stopped
+					UserWidget->ConditionalTearDownAnimations();
 					UserWidget->UpdateCanTick();
 
-					// Resetting the animation tick manager is ok here because TearDownAnimations will always clear out all animations
-					UserWidget->AnimationTickManager = nullptr;
+					// If there are no more animations playing, we can remove this widget altogether
+					if (!UserWidget->IsAnyAnimationPlaying())
+					{
+						// Resetting the animation tick manager is ok here because TearDownAnimations will always clear out all animations
+						UserWidget->AnimationTickManager = nullptr;
 
-					WidgetIter.RemoveCurrent();
+						WidgetIter.RemoveCurrent();
+					}
 				}
 			}
 			else if (!WidgetData.bIsTicking)
diff --git a/Engine/Source/Runtime/UMG/Private/UserWidget.cpp b/Engine/Source/Runtime/UMG/Private/UserWidget.cpp
--- a/Engine/Source/Runtime/UMG/Private/UserWidget.cpp
+++ b/Engine/Source/Runtime/UMG/Private/UserWidget.cpp
@@ -474,6 +474,29 @@ UUMGSequencePlayer* UUserWidget::GetOrAddSequencePlayer(UWidgetAnimation* InAnim
 	return nullptr;
 }
 
+void UUserWidget::ConditionalTearDownAnimations()
+{
+	for (auto It = ActiveSequencePlayers.CreateIterator(); It; ++It)
+	{
+		UUMGSequencePlayer* Player = *It;
+		if (!Player)
+		{
+			It.RemoveCurrent();
+		}
+		else if (!Player->IsStopping())
+		{
+			Player->TearDown();
+			It.RemoveCurrent();
+		}
+	}
+
+	for (UUMGSequencePlayer* Player : StoppedSequencePlayers)
+	{
+		Player->TearDown();
+	}
+	StoppedSequencePlayers.Empty();
+}
+
 void UUserWidget::TearDownAnimations()
 {
 	for (UUMGSequencePlayer* Player : ActiveSequencePlayers)
diff --git a/Engine/Source/Runtime/UMG/Public/Blueprint/UserWidget.h b/Engine/Source/Runtime/UMG/Public/Blueprint/UserWidget.h
--- a/Engine/Source/Runtime/UMG/Public/Blueprint/UserWidget.h
+++ b/Engine/Source/Runtime/UMG/Public/Blueprint/UserWidget.h
@@ -1427,6 +1427,8 @@ protected:
 	UUMGSequencePlayer* GetSequencePlayer(const UWidgetAnimation* InAnimation) const;
 	UUMGSequencePlayer* GetOrAddSequencePlayer(UWidgetAnimation* InAnimation);
 
+	void ConditionalTearDownAnimations();
+
 	void TearDownAnimations();
 
 	void DisableAnimations();
