UMG: Invalidate the Image when the brush value changes.

diff --git a/Engine/Source/Runtime/UMG/Private/Components/Image.cpp b/Engine/Source/Runtime/UMG/Private/Components/Image.cpp
--- a/Engine/Source/Runtime/UMG/Private/Components/Image.cpp
+++ b/Engine/Source/Runtime/UMG/Private/Components/Image.cpp
@@ -92,7 +92,7 @@ const FSlateBrush* UImage::ConvertImage(TAttribute<FSlateBrush> InImageAsset) co
 
 void UImage::SetBrush(const FSlateBrush& InBrush)
 {
-	if (Brush != InBrush)
+	if(Brush != InBrush)
 	{
 		Brush = InBrush;
 		BroadcastFieldValueChanged(FFieldNotificationClassDescriptor::Brush);
@@ -133,7 +133,7 @@ void UImage::SetBrushTintColor(FSlateColor TintColor)
 
 		if (MyImage.IsValid())
 		{
-			MyImage->InvalidateImage();;
+			MyImage->InvalidateImage();
 		}
 	}
 }
