Set WidgetComponent propeties that should not be serialized to Transient/DuplicateTransient

diff --git a/Engine/Source/Runtime/UMG/Private/Components/WidgetComponent.cpp b/Engine/Source/Runtime/UMG/Private/Components/WidgetComponent.cpp
--- a/Engine/Source/Runtime/UMG/Private/Components/WidgetComponent.cpp
+++ b/Engine/Source/Runtime/UMG/Private/Components/WidgetComponent.cpp
@@ -611,6 +611,7 @@ UWidgetComponent::UWidgetComponent( const FObjectInitializer& PCIP )
 	, bRedrawRequested(true)
 	, RedrawTime(0)
 	, LastWidgetRenderTime(0)
+	, CurrentDrawSize(FIntPoint(0, 0))
 	, bReceiveHardwareInput(false)
 	, bWindowFocusable(true)
 	, WindowVisibility(EWindowVisibility::SelfHitTestInvisible)
diff --git a/Engine/Source/Runtime/UMG/Public/Components/WidgetComponent.h b/Engine/Source/Runtime/UMG/Public/Components/WidgetComponent.h
--- a/Engine/Source/Runtime/UMG/Public/Components/WidgetComponent.h
+++ b/Engine/Source/Runtime/UMG/Public/Components/WidgetComponent.h
@@ -454,7 +454,7 @@ protected:
 	bool bManuallyRedraw;
 
 	/** Has anyone requested we redraw? */
-	UPROPERTY()
+	UPROPERTY(Transient, DuplicateTransient)
 	bool bRedrawRequested;
 
 	/**
@@ -475,7 +475,7 @@ protected:
 	 * The actual draw size, this changes based on DrawSize - or the desired size of the widget if
 	 * bDrawAtDesiredSize is true.
 	 */
-	UPROPERTY()
+	UPROPERTY(Transient, DuplicateTransient)
 	FIntPoint CurrentDrawSize;
 
 	/**
