UImage: remove redundant invalidation

diff --git a/Engine/Source/Runtime/UMG/Private/Components/Image.cpp b/Engine/Source/Runtime/UMG/Private/Components/Image.cpp
--- a/Engine/Source/Runtime/UMG/Private/Components/Image.cpp
+++ b/Engine/Source/Runtime/UMG/Private/Components/Image.cpp
@@ -174,7 +174,8 @@ void UImage::SetBrushFromTexture(UTexture2D* Texture, bool bMatchSize)
 
 	if(Brush.GetResourceObject() != Texture)
 	{
-		SetBrushResourceObject(Texture);
+		Brush.SetResourceObject(Texture);
+		BroadcastFieldValueChanged(FFieldNotificationClassDescriptor::Brush);
 
 		if (Texture) // Since this texture is used as UI, don't allow it affected by budget.
 		{
@@ -193,7 +194,6 @@ void UImage::SetBrushFromTexture(UTexture2D* Texture, bool bMatchSize)
 			{
 				Brush.ImageSize = FVector2D(0, 0);
 			}
-			BroadcastFieldValueChanged(FFieldNotificationClassDescriptor::Brush);
 		}
 
 		if (MyImage.IsValid())
@@ -208,7 +208,8 @@ void UImage::SetBrushFromAtlasInterface(TScriptInterface<ISlateTextureAtlasInter
 	if(Brush.GetResourceObject() != AtlasRegion.GetObject())
 	{
 		CancelImageStreaming();
-		SetBrushResourceObject(AtlasRegion.GetObject());
+		Brush.SetResourceObject(AtlasRegion.GetObject());
+		BroadcastFieldValueChanged(FFieldNotificationClassDescriptor::Brush);
 
 		if (bMatchSize)
 		{
@@ -221,7 +222,6 @@ void UImage::SetBrushFromAtlasInterface(TScriptInterface<ISlateTextureAtlasInter
 			{
 				Brush.ImageSize = FVector2D(0, 0);
 			}
-			BroadcastFieldValueChanged(FFieldNotificationClassDescriptor::Brush);
 		}
 
 		if (MyImage.IsValid())
@@ -236,13 +236,13 @@ void UImage::SetBrushFromTextureDynamic(UTexture2DDynamic* Texture, bool bMatchS
 	if(Brush.GetResourceObject() != Texture)
 	{
 		CancelImageStreaming();
-		SetBrushResourceObject(Texture);
+		Brush.SetResourceObject(Texture);
+		BroadcastFieldValueChanged(FFieldNotificationClassDescriptor::Brush);
 
 		if (bMatchSize && Texture)
 		{
 			Brush.ImageSize.X = Texture->SizeX;
 			Brush.ImageSize.Y = Texture->SizeY;
-			BroadcastFieldValueChanged(FFieldNotificationClassDescriptor::Brush);
 		}
 
 		if (MyImage.IsValid())
@@ -257,7 +257,8 @@ void UImage::SetBrushFromMaterial(UMaterialInterface* Material)
 	if(Brush.GetResourceObject() != Material)
 	{
 		CancelImageStreaming();
-		SetBrushResourceObject(Material);
+		Brush.SetResourceObject(Material);
+		BroadcastFieldValueChanged(FFieldNotificationClassDescriptor::Brush);
 
 		//TODO UMG Check if the material can be used with the UI
 
