UMG: Add comments on WIdgetBllueprintGeneratedClassExtension's function
#rnx

diff --git a/Engine/Source/Runtime/UMG/Public/Extensions/WidgetBlueprintGeneratedClassExtension.h b/Engine/Source/Runtime/UMG/Public/Extensions/WidgetBlueprintGeneratedClassExtension.h
--- a/Engine/Source/Runtime/UMG/Public/Extensions/WidgetBlueprintGeneratedClassExtension.h
+++ b/Engine/Source/Runtime/UMG/Public/Extensions/WidgetBlueprintGeneratedClassExtension.h
@@ -17,12 +17,21 @@ class UMG_API UWidgetBlueprintGeneratedClassExtension : public UObject
 	GENERATED_BODY()
 	
 public:
-	/** Extend the UUserWidget::Initialize function */
+	/**
+	 * Extend the UUserWidget::Initialize function.
+	 * Called once when the UUserWidget is created.
+	 */
 	virtual void Initialize(UUserWidget* UserWidget) { }
 	/** Extend the UUserWidget::PreConstruct function */
 	virtual void PreConstruct(UUserWidget* UserWidget, bool IsDesignTime) { }
-	/** Extend the UUserWidget::Construct function */
+	/**
+	 * Extend the UUserWidget::Construct function.
+	 * Called after when the SWidget is created. (TakeWidget)
+	 */
 	virtual void Construct(UUserWidget* UserWidget) { }
-	/** Extend the UUserWidget::Destruct function */
+	/**
+	 * Extend the UUserWidget::Destruct function.
+	 * Called when the SWidget is going to be destroy. (ReleaseSlateWidget)
+	 */
 	virtual void Destruct(UUserWidget* UserWidget) { }
 };
