Remove "_DEPRECATED" from the deprecated property name in ListView to patch the new property correctly.

diff --git a/Engine/Source/Runtime/UMG/Private/Components/ListView.cpp b/Engine/Source/Runtime/UMG/Private/Components/ListView.cpp
--- a/Engine/Source/Runtime/UMG/Private/Components/ListView.cpp
+++ b/Engine/Source/Runtime/UMG/Private/Components/ListView.cpp
@@ -361,11 +361,11 @@ void UListView::PostLoad()
 #if WITH_EDITOR
 	PRAGMA_DISABLE_DEPRECATION_WARNINGS
 
-		if (EntrySpacing_DEPRECATED != 0.f)
+		if (EntrySpacing != 0.f)
 		{
-			HorizontalEntrySpacing = EntrySpacing_DEPRECATED;
-			VerticalEntrySpacing = EntrySpacing_DEPRECATED;
-			EntrySpacing_DEPRECATED = 0.f;
+			HorizontalEntrySpacing = EntrySpacing;
+			VerticalEntrySpacing = EntrySpacing;
+			EntrySpacing = 0.f;
 		}
 
 	PRAGMA_ENABLE_DEPRECATION_WARNINGS
diff --git a/Engine/Source/Runtime/UMG/Public/Components/ListView.h b/Engine/Source/Runtime/UMG/Public/Components/ListView.h
--- a/Engine/Source/Runtime/UMG/Public/Components/ListView.h
+++ b/Engine/Source/Runtime/UMG/Public/Components/ListView.h
@@ -212,7 +212,7 @@ protected:
 	 * it relies on the newly added UFunction GetHorizontalEntrySpacing() to act as its BlueprintGetter.
 	 */
 	UPROPERTY(meta = (DeprecatedProperty, DeprecationMessage = "EntrySpacing has been deprecated. Please use HorizontalEntrySpacing and VerticalEntrySpacing."), BlueprintReadOnly, BlueprintGetter = "GetHorizontalEntrySpacing", Category = ListView, meta = (ClampMin = 0))
-	float EntrySpacing_DEPRECATED = 0.f;
+	float EntrySpacing = 0.f;
 #endif  // WITH_EDITOR
 
 private:
