Fix incorrectly displayed text. In Slate, wrapping text needs 2 draws to be displayed properly, otherwise some incorrect value are used. In the case of this bug, the text was offseted with a value based on the text width without wrapping (as if it was single line)  instead of the correct width with wrapping applied. An invalidation of the widget was called, but as the invalidation was not supported in this code path (WidgetComponent + WidgetRendere::DrawWindow), it was ignored.

Now the WidgetRenderer::DrawWindow support Invalidation Root, and the WidgetComponent enabled the Invalidation on the VirtualWindow it uses, and do only the necessary amount of Tick.

#tests Repro steps from the jira, no longer occurs. Tested that hit testing was still working properly by using a widget with buttons in the widgetcomponent. Tested with and without the new setting bIsInvalidationEnabled activated.

diff --git a/Engine/Source/Runtime/SlateCore/Private/FastUpdate/SlateInvalidationRoot.cpp b/Engine/Source/Runtime/SlateCore/Private/FastUpdate/SlateInvalidationRoot.cpp
--- a/Engine/Source/Runtime/SlateCore/Private/FastUpdate/SlateInvalidationRoot.cpp
+++ b/Engine/Source/Runtime/SlateCore/Private/FastUpdate/SlateInvalidationRoot.cpp
@@ -199,6 +199,7 @@ FSlateInvalidationRoot::FSlateInvalidationRoot()
 	, InvalidationRootWidget(nullptr)
 	, RootHittestGrid(nullptr)
 	, CachedMaxLayerId(0)
+	, InvalidationRequestCount(0)
 	, bNeedsSlowPath(true)
 	, bNeedScreenPositionShift(false)
 	, bProcessingPreUpdate(false)
@@ -290,7 +291,17 @@ void FSlateInvalidationRoot::InvalidateRootLayout(const SWidget* Investigator)
 	UE_TRACE_SLATE_ROOT_INVALIDATED(InvalidationRootWidget, Investigator);
 }
 
-void FSlateInvalidationRoot::InvalidateWidget(FWidgetProxy& Proxy, EInvalidateWidgetReason InvalidateReason)
+void FSlateInvalidationRoot::NotifyInvalidate()
+{
+	if (InvalidationRequestCount == 0)
+	{
+		InvalidationUpdateNeeded.Broadcast();
+	}
+	InvalidationRequestCount++;
+	ensure(InvalidationRequestCount < std::numeric_limits<uint32>::max());
+}
+
+void FSlateInvalidationRoot::InvalidateWidget(FWidgetProxy & Proxy, EInvalidateWidgetReason InvalidateReason)
 {
 	ensureMsgf(bProcessingChildOrderInvalidation == false, TEXT("A widget got invalidated while building the childorder."));
 
@@ -303,6 +314,8 @@ void FSlateInvalidationRoot::InvalidateWidget(FWidgetProxy& Proxy, EInvalidateWi
 		}
 	}
 
+	NotifyInvalidate();
+
 	if (!bNeedsSlowPath)
 	{
 		Proxy.CurrentInvalidateReason |= InvalidateReason;
@@ -1081,7 +1094,7 @@ void FSlateInvalidationRoot::ProcessPrepassUpdate()
 	TGuardValue<bool> Tmp(bProcessingPrepassUpdate, true);
 
 #if WITH_SLATE_DEBUGGING
-	if (GSlateInvalidationRootDumpPostInvalidationList)
+	if (GSlateInvalidationRootDumpPrepassInvalidationList)
 	{
 		UE_LOG(LogSlate, Log, TEXT("Dumping Prepass Invalidation List"));
 		UE_LOG(LogSlate, Log, TEXT("-------------------"));
@@ -1203,6 +1216,7 @@ bool FSlateInvalidationRoot::ProcessInvalidation()
 	FScopedDurationTimer TmpPerformance_ProcessInvalidation(PerformanceStat.InvalidationProcessing);
 #endif
 
+	InvalidationRequestCount = 0;
 	TGuardValue<bool> OnFastPathGuard(GSlateIsOnFastProcessInvalidation, true);
 
 	bool bWidgetsNeedRepaint = false;
diff --git a/Engine/Source/Runtime/SlateCore/Private/Widgets/SWindow.cpp b/Engine/Source/Runtime/SlateCore/Private/Widgets/SWindow.cpp
--- a/Engine/Source/Runtime/SlateCore/Private/Widgets/SWindow.cpp
+++ b/Engine/Source/Runtime/SlateCore/Private/Widgets/SWindow.cpp
@@ -726,6 +726,17 @@ void SWindow::SetAllowFastUpdate(bool bInAllowFastUpdate)
 	}
 }
 
+void SWindow::SetLocallyEnabledInvalidation(bool bInEnableInvalidation)
+{
+	if (bLocallyEnabledInvalidation != bInEnableInvalidation)
+	{
+		bLocallyEnabledInvalidation = bInEnableInvalidation;
+		if (bLocallyEnabledInvalidation)
+		{
+			InvalidateRootChildOrder();
+		}
+	}
+}
 
 void SWindow::Tick( const FGeometry& AllottedGeometry, const double InCurrentTime, const float InDeltaTime )
 {
@@ -1102,19 +1113,24 @@ void SWindow::StartMorph()
 	}
 }
 
+bool SWindow::CanUseFastUpdate() const
+{
+	return bAllowFastUpdate && (bLocallyEnabledInvalidation || GSlateEnableGlobalInvalidation);
+}
+
 bool SWindow::Advanced_IsInvalidationRoot() const
 {
-	return bAllowFastUpdate && GSlateEnableGlobalInvalidation;
+	return CanUseFastUpdate();
 }
 
 const FSlateInvalidationRoot* SWindow::Advanced_AsInvalidationRoot() const
 {
-	return (bAllowFastUpdate && GSlateEnableGlobalInvalidation) ? this : nullptr;
+	return CanUseFastUpdate()  ? this : nullptr;
 }
 
 void SWindow::ProcessWindowInvalidation()
 {
-	if (bAllowFastUpdate && GSlateEnableGlobalInvalidation)
+	if (CanUseFastUpdate())
 	{
 		ProcessInvalidation();
 	}
@@ -1122,7 +1138,7 @@ void SWindow::ProcessWindowInvalidation()
 
 bool SWindow::CustomPrepass(float LayoutScaleMultiplier)
 {
-	if (bAllowFastUpdate && GSlateEnableGlobalInvalidation)
+	if (CanUseFastUpdate())
 	{
 		return NeedsPrepass();
 	}
@@ -2088,7 +2104,7 @@ int32 SWindow::PaintWindow( double CurrentTime, float DeltaTime, FSlateWindowEle
 	FSlateInvalidationContext Context(OutDrawElements, InWidgetStyle);
 	Context.bParentEnabled = bParentEnabled;
 	// Fast path at the window level should only be enabled if global invalidation is allowed
-	Context.bAllowFastPathUpdate = bAllowFastUpdate && GSlateEnableGlobalInvalidation;
+	Context.bAllowFastPathUpdate = CanUseFastUpdate();
 	Context.LayoutScaleMultiplier = FSlateApplicationBase::Get().GetApplicationScale() * GetDPIScaleFactor();
 	Context.PaintArgs = &PaintArgs;
 	Context.IncomingLayerId = 0;
diff --git a/Engine/Source/Runtime/SlateCore/Public/FastUpdate/SlateInvalidationRoot.h b/Engine/Source/Runtime/SlateCore/Public/FastUpdate/SlateInvalidationRoot.h
--- a/Engine/Source/Runtime/SlateCore/Public/FastUpdate/SlateInvalidationRoot.h
+++ b/Engine/Source/Runtime/SlateCore/Public/FastUpdate/SlateInvalidationRoot.h
@@ -80,6 +80,9 @@ class FSlateInvalidationRoot : public FGCObject, public FNoncopyable
 	friend class FWidgetProxyHandle;
 
 public:
+	/** Delegate fired when the InvalidationRoot got invalidated, and requires an update (process invalidation and paint). */
+	DECLARE_MULTICAST_DELEGATE(FOnInvalidationUpdateNeeded);
+
 	SLATECORE_API FSlateInvalidationRoot();
 	SLATECORE_API virtual ~FSlateInvalidationRoot();
 
@@ -126,6 +129,8 @@ public:
 
 	SLATECORE_API void Advanced_ResetInvalidation(bool bClearResourcesImmediately);
 
+	SLATECORE_API FOnInvalidationUpdateNeeded& GetOnInvalidationUpdateNeeded() { return InvalidationUpdateNeeded; }
+
 #if WITH_SLATE_DEBUGGING
 	/** @return the last paint type the invalidation root handle used. */
 	ESlateInvalidationPaintType GetLastPaintType() const { return LastPaintType; }
@@ -171,6 +176,8 @@ private:
 	/** Call when an invalidation occurred. */
 	void InvalidateWidget(FWidgetProxy& Proxy, EInvalidateWidgetReason InvalidateReason);
 
+	void NotifyInvalidate();
+
 	void BuildFastPathWidgetList(const TSharedRef<SWidget>& RootWidget);
 	void AdjustWidgetsDesktopGeometry(UE::Slate::FDeprecateVector2DParameter WindowToDesktopTransform);
 
@@ -214,6 +221,11 @@ private:
 	FHittestGrid* RootHittestGrid;
 
 	int32 CachedMaxLayerId;
+	/** Count of invalidation requests received since the last call to ProcessInvalidation. */
+	uint32 InvalidationRequestCount;
+
+	/** This delegate will be called at most a single time between ProcessInvalidation calls, for the first invalidation.*/
+	FOnInvalidationUpdateNeeded InvalidationUpdateNeeded;
 
 	FSlateInvalidationRootHandle InvalidationRootHandle;
 
diff --git a/Engine/Source/Runtime/SlateCore/Public/Widgets/SWindow.h b/Engine/Source/Runtime/SlateCore/Public/Widgets/SWindow.h
--- a/Engine/Source/Runtime/SlateCore/Public/Widgets/SWindow.h
+++ b/Engine/Source/Runtime/SlateCore/Public/Widgets/SWindow.h
@@ -945,6 +945,9 @@ private:
 	void ResizeWindowSize( FVector2f NewWindowSize );
 
 	void OnGlobalInvalidationToggled(bool bGlobalInvalidationEnabled);
+
+	bool CanUseFastUpdate() const;
+
 public:
 	/**
 	 * For a given client size, calculate the window size required to accommodate any potential non-OS borders and title bars
@@ -1062,6 +1065,8 @@ public:
 	void SetSizeLimits(const FWindowSizeLimits& InSizeLimits);
 
 	void SetAllowFastUpdate(bool bInAllowFastUpdate);
+
+	void SetLocallyEnabledInvalidation(bool bInEnableInvalidation);
 public:
 
 	//~ SWidget overrides
@@ -1164,6 +1169,9 @@ protected:
 	/** True if this window allows global invalidation of its contents */
 	bool bAllowFastUpdate : 1;
 
+	/** True to enable invalidation on this particular window, even if the global invalidation has not been enabled. */
+	bool bLocallyEnabledInvalidation : 1;
+
 	/** When should the window be activated upon being shown */
 	EWindowActivationPolicy WindowActivationPolicy;
 
diff --git a/Engine/Source/Runtime/UMG/Private/Components/WidgetComponent.cpp b/Engine/Source/Runtime/UMG/Private/Components/WidgetComponent.cpp
--- a/Engine/Source/Runtime/UMG/Private/Components/WidgetComponent.cpp
+++ b/Engine/Source/Runtime/UMG/Private/Components/WidgetComponent.cpp
@@ -612,6 +612,7 @@ UWidgetComponent::UWidgetComponent( const FObjectInitializer& PCIP )
 	, RedrawTime(0)
 	, LastWidgetRenderTime(0)
 	, CurrentDrawSize(FIntPoint(0, 0))
+	, bIsInvalidationEnabled(false)
 	, bReceiveHardwareInput(false)
 	, bWindowFocusable(true)
 	, WindowVisibility(EWindowVisibility::SelfHitTestInvisible)
@@ -1136,6 +1137,10 @@ void UWidgetComponent::RegisterWindow()
 		if (!CanReceiveHardwareInput() && FSlateApplication::IsInitialized() )
 		{
 			FSlateApplication::Get().RegisterVirtualWindow(SlateWindow.ToSharedRef());
+			if (bIsInvalidationEnabled)
+			{
+				SlateWindow->GetOnInvalidationUpdateNeeded().AddUObject(this, &ThisClass::RequestRenderUpdate);
+			}
 		}
 
 		if (Widget && !Widget->IsDesignTime())
@@ -1163,6 +1168,7 @@ void UWidgetComponent::UnregisterWindow()
 	{
 		if ( !CanReceiveHardwareInput() && FSlateApplication::IsInitialized() )
 		{
+			SlateWindow->GetOnInvalidationUpdateNeeded().RemoveAll(this);
 			FSlateApplication::Get().UnregisterVirtualWindow(SlateWindow.ToSharedRef());
 		}
 
@@ -1304,7 +1310,7 @@ bool UWidgetComponent::ShouldDrawWidget() const
 	if ( IsVisible() )
 	{
 		// If we don't tick when off-screen, don't bother ticking if it hasn't been rendered recently
-		if ( TickWhenOffscreen || GetWorld()->TimeSince(GetLastRenderTime()) <= RenderTimeThreshold )
+		if ( TickWhenOffscreen || GetWorld()->TimeSince(GetLastRenderTime()) <= RenderTimeThreshold)
 		{
 			if ( ( GetCurrentTime() - LastWidgetRenderTime) >= RedrawTime )
 			{
@@ -1385,6 +1391,10 @@ void UWidgetComponent::DrawWidgetToRenderTarget(float DeltaTime)
 	if(RenderTarget)
 	{
 		bRedrawRequested = false;
+		if (TickMode == ETickMode::Disabled && IsComponentTickEnabled())
+		{
+			SetComponentTickEnabled(false);
+		}
 
 		WidgetRenderer->DrawWindow(
 			RenderTarget,
@@ -1395,11 +1405,6 @@ void UWidgetComponent::DrawWidgetToRenderTarget(float DeltaTime)
 			DeltaTime);
 
 		LastWidgetRenderTime = GetCurrentTime();
-
-		if (TickMode == ETickMode::Disabled && IsComponentTickEnabled())
-		{
-			SetComponentTickEnabled(false);
-		}
 	}
 }
 
@@ -1732,6 +1737,8 @@ void UWidgetComponent::UpdateWidget()
 				SlateWindow = SNew(SVirtualWindow).Size(CurrentDrawSize);
 				SlateWindow->SetIsFocusable(bWindowFocusable);
 				SlateWindow->SetVisibility(ConvertWindowVisibilityToVisibility(WindowVisibility));
+				SlateWindow->SetAllowFastUpdate(bIsInvalidationEnabled);
+				SlateWindow->SetLocallyEnabledInvalidation(bIsInvalidationEnabled);
 				RegisterWindow();
 
 				bNeededNewWindow = true;
@@ -1782,7 +1789,7 @@ void UWidgetComponent::UpdateWidget()
 			}
 		}
 
-		UpdateWidgetOnScreen();		
+		UpdateWidgetOnScreen();
 	}
 }
 
diff --git a/Engine/Source/Runtime/UMG/Private/Slate/WidgetRenderer.cpp b/Engine/Source/Runtime/UMG/Private/Slate/WidgetRenderer.cpp
--- a/Engine/Source/Runtime/UMG/Private/Slate/WidgetRenderer.cpp
+++ b/Engine/Source/Runtime/UMG/Private/Slate/WidgetRenderer.cpp
@@ -223,6 +223,7 @@ void FWidgetRenderer::DrawWindow(
 		if (bPrepassNeeded)
 		{
 			// Ticking can cause geometry changes.  Recompute
+			Window->ProcessWindowInvalidation();
 			Window->SlatePrepass(WindowGeometry.Scale);
 		}
 
@@ -239,35 +240,56 @@ void FWidgetRenderer::DrawWindow(
 			ISlate3DRenderer::FScopedAcquireDrawBuffer ScopedDrawBuffer{ *Renderer, bDeferRenderTargetUpdate };
 			FSlateWindowElementList& WindowElementList = ScopedDrawBuffer.GetDrawBuffer().AddWindowElementList(Window);
 
-			// Paint the window
-			int32 MaxLayerId = Window->Paint(
-				PaintArgs,
-				WindowGeometry, WindowClipRect,
-				WindowElementList,
-				0,
-				FWidgetStyle(),
-				Window->IsEnabled());
+			bool bRepaintedWidgets = true;
 
-			//MaxLayerId = WindowElementList.PaintDeferred(MaxLayerId);
-			DeferredPaints = WindowElementList.GetDeferredPaintList();
+			if (Window->Advanced_IsInvalidationRoot())
+			{
+				// Need to set a new window element list so make a copy
+				FSlateInvalidationContext Context{ WindowElementList, FWidgetStyle() };
+				Context.bParentEnabled = Window->IsEnabled();
+				Context.bAllowFastPathUpdate = true;
+				Context.LayoutScaleMultiplier = WindowGeometry.Scale;
+				Context.PaintArgs = &PaintArgs;
+				Context.IncomingLayerId = 0;
+				Context.CullingRect = WindowClipRect;
+				Context.ViewOffset = ViewOffset;
+				FSlateInvalidationResult Result = Window->PaintInvalidationRoot(Context);
+				bRepaintedWidgets = Result.bRepaintedWidgets;
+			}
+			else
+			{
+				// Paint the window
+				Window->Paint(
+					PaintArgs,
+					WindowGeometry, WindowClipRect,
+					WindowElementList,
+					0,
+					FWidgetStyle(),
+					Window->IsEnabled());
+			}
 
-			Renderer->DrawWindow_GameThread(ScopedDrawBuffer.GetDrawBuffer());
+			if (bRepaintedWidgets)
+			{
+				DeferredPaints = WindowElementList.GetDeferredPaintList();
+
+				Renderer->DrawWindow_GameThread(ScopedDrawBuffer.GetDrawBuffer());
 
-			ScopedDrawBuffer.GetDrawBuffer().ViewOffset = ViewOffset;
+				ScopedDrawBuffer.GetDrawBuffer().ViewOffset = ViewOffset;
 
-			FRenderThreadUpdateContext RenderThreadUpdateContext =
-			{
-				&(ScopedDrawBuffer.GetDrawBuffer()),
-				(FApp::GetCurrentTime() - GStartTime),
-				static_cast<float>(FApp::GetDeltaTime()),
-				(FPlatformTime::Seconds() - GStartTime),
-				static_cast<float>(FApp::GetDeltaTime()),
-				RenderTarget,
-				Renderer.Get(),
-				bClearTarget
-			};
-
-			MainSlateRenderer->AddWidgetRendererUpdate(RenderThreadUpdateContext, bDeferRenderTargetUpdate);
+				FRenderThreadUpdateContext RenderThreadUpdateContext =
+				{
+					&(ScopedDrawBuffer.GetDrawBuffer()),
+					(FApp::GetCurrentTime() - GStartTime),
+					static_cast<float>(FApp::GetDeltaTime()),
+					(FPlatformTime::Seconds() - GStartTime),
+					static_cast<float>(FApp::GetDeltaTime()),
+					RenderTarget,
+					Renderer.Get(),
+					bClearTarget
+				};
+
+				MainSlateRenderer->AddWidgetRendererUpdate(RenderThreadUpdateContext, bDeferRenderTargetUpdate);
+			}
 		}
 	}
 #endif // !UE_SERVER
diff --git a/Engine/Source/Runtime/UMG/Public/Components/WidgetComponent.h b/Engine/Source/Runtime/UMG/Public/Components/WidgetComponent.h
--- a/Engine/Source/Runtime/UMG/Public/Components/WidgetComponent.h
+++ b/Engine/Source/Runtime/UMG/Public/Components/WidgetComponent.h
@@ -478,6 +478,10 @@ protected:
 	UPROPERTY(Transient, DuplicateTransient)
 	FIntPoint CurrentDrawSize;
 
+	/** Should the component use UMG Invalidation system? Use it only if the widgets used support invalidation*/
+	UPROPERTY(EditAnywhere, Category= UserInterface)
+	bool bIsInvalidationEnabled;
+
 	/**
 	 * Causes the render target to automatically match the desired size.
 	 * 
