Uses GetContent().Num() instead of GetContentLenght() to prevents unbouded  access to data.
#rnx

diff --git a/Engine/Source/Runtime/UMG/Private/Blueprint/AsyncTaskDownloadImage.cpp b/Engine/Source/Runtime/UMG/Private/Blueprint/AsyncTaskDownloadImage.cpp
--- a/Engine/Source/Runtime/UMG/Private/Blueprint/AsyncTaskDownloadImage.cpp
+++ b/Engine/Source/Runtime/UMG/Private/Blueprint/AsyncTaskDownloadImage.cpp
@@ -57,7 +57,7 @@ void UAsyncTaskDownloadImage::HandleImageRequest(FHttpRequestPtr HttpRequest, FH
 
 	RemoveFromRoot();
 
-	if ( bSucceeded && HttpResponse.IsValid() && HttpResponse->GetContentLength() > 0 )
+	if ( bSucceeded && HttpResponse.IsValid() && HttpResponse->GetContentLength() > 0 && HttpResponse->GetContent().Num() > 0 )
 	{
 		IImageWrapperModule& ImageWrapperModule = FModuleManager::LoadModuleChecked<IImageWrapperModule>(FName("ImageWrapper"));
 		TSharedPtr<IImageWrapper> ImageWrappers[3] =
@@ -69,7 +69,7 @@ void UAsyncTaskDownloadImage::HandleImageRequest(FHttpRequestPtr HttpRequest, FH
 
 		for ( auto ImageWrapper : ImageWrappers )
 		{
-			if ( ImageWrapper.IsValid() && ImageWrapper->SetCompressed(HttpResponse->GetContent().GetData(), HttpResponse->GetContentLength()) )
+			if ( ImageWrapper.IsValid() && ImageWrapper->SetCompressed(HttpResponse->GetContent().GetData(), HttpResponse->GetContent().Num()) )
 			{
 				TArray64<uint8> RawData;
 				const ERGBFormat InFormat = ERGBFormat::BGRA;
@@ -89,7 +89,7 @@ void UAsyncTaskDownloadImage::HandleImageRequest(FHttpRequestPtr HttpRequest, FH
 									TextureResource->WriteRawToTexture_RenderThread(RawData);
 								});
 						}
-						OnSuccess.Broadcast(Texture);						
+						OnSuccess.Broadcast(Texture);
 						return;
 					}
 				}
