[25.20] URichTextBlock:
- Made ClearAllDefaultStyleOverrides BlueprintCallable.

diff --git a/Engine/Source/Runtime/UMG/Public/Components/RichTextBlock.h b/Engine/Source/Runtime/UMG/Public/Components/RichTextBlock.h
--- a/Engine/Source/Runtime/UMG/Public/Components/RichTextBlock.h
+++ b/Engine/Source/Runtime/UMG/Public/Components/RichTextBlock.h
@@ -154,7 +154,7 @@ public:
 	void SetDefaultMaterial(UMaterialInterface* InMaterial);
 
 	/** Remove all overrides made to the default text style and return to the style specified in the style set data table */
-	UFUNCTION()
+	UFUNCTION(BlueprintCallable, Category = Appearance)
 	void ClearAllDefaultStyleOverrides();
 
 	/**
