UMG: Add callback to add tag data to the widget generated class asset.

diff --git a/Engine/Source/Runtime/UMG/Private/WidgetBlueprintGeneratedClass.cpp b/Engine/Source/Runtime/UMG/Private/WidgetBlueprintGeneratedClass.cpp
--- a/Engine/Source/Runtime/UMG/Private/WidgetBlueprintGeneratedClass.cpp
+++ b/Engine/Source/Runtime/UMG/Private/WidgetBlueprintGeneratedClass.cpp
@@ -103,6 +103,8 @@ FAutoConsoleCommand GDumpTemplateSizesCommand(
 int32 TemplatePreviewInEditor = 0;
 static FAutoConsoleVariableRef CVarTemplatePreviewInEditor(TEXT("Widget.TemplatePreviewInEditor"), TemplatePreviewInEditor, TEXT("Should a dynamic template be generated at runtime for the editor for widgets?  Useful for debugging templates."), ECVF_Default);
 
+FWidgetBlueprintGeneratedClassDelegates::FGetAssetTags FWidgetBlueprintGeneratedClassDelegates::GetAssetTags;
+
 #endif
 
 #if WITH_EDITORONLY_DATA
@@ -525,5 +527,14 @@ void UWidgetBlueprintGeneratedClass::GetExtensions(TArray<UWidgetBlueprintGenera
 	}
 }
 
+#if WITH_EDITOR
+void UWidgetBlueprintGeneratedClass::GetAssetRegistryTags(TArray<FAssetRegistryTag>& OutTags) const
+{
+	Super::GetAssetRegistryTags(OutTags);
+
+	FWidgetBlueprintGeneratedClassDelegates::GetAssetTags.Broadcast(this, OutTags);
+}
+#endif
+
 #undef LOCTEXT_NAMESPACE
 
diff --git a/Engine/Source/Runtime/UMG/Public/Blueprint/WidgetBlueprintGeneratedClass.h b/Engine/Source/Runtime/UMG/Public/Blueprint/WidgetBlueprintGeneratedClass.h
--- a/Engine/Source/Runtime/UMG/Public/Blueprint/WidgetBlueprintGeneratedClass.h
+++ b/Engine/Source/Runtime/UMG/Public/Blueprint/WidgetBlueprintGeneratedClass.h
@@ -13,6 +13,7 @@
 class UWidget;
 class UUserWidget;
 class UWidgetAnimation;
+class UWidgetBlueprintGeneratedClass;
 class UWidgetBlueprintGeneratedClassExtension;
 class UWidgetTree;
 
@@ -50,6 +51,18 @@ struct FDelegateRuntimeBinding
 };
 
 
+#if WITH_EDITOR
+class UMG_API FWidgetBlueprintGeneratedClassDelegates
+{
+public:
+	// delegate for generating widget asset registry tags.
+	DECLARE_MULTICAST_DELEGATE_TwoParams(FGetAssetTags, const UWidgetBlueprintGeneratedClass*, TArray<UObject::FAssetRegistryTag>&);
+
+	// called by UWidgetBlueprintGeneratedClass::GetAssetRegistryTags()
+	static FGetAssetTags GetAssetTags;
+};
+#endif
+
 /**
  * The widget blueprint generated class allows us to create blueprint-able widgets for UMG at runtime.
  * All WBPGC's are of UUserWidget classes, and they perform special post initialization using this class
@@ -134,6 +147,9 @@ public:
 	virtual void Serialize(FArchive& Ar) override;
 	virtual void PostLoad() override;
 	virtual bool NeedsLoadForServer() const override;
+#if WITH_EDITOR
+	virtual void GetAssetRegistryTags(TArray<FAssetRegistryTag>& OutTags) const override;
+#endif
 	//~ End UObject interface
 
 	virtual void PurgeClass(bool bRecompilingOnLoad) override;
