Fix calculation for total entry size in tileview when entry spacing is not included in the entry size.

diff --git a/Engine/Source/Runtime/UMG/Public/Components/TileView.h b/Engine/Source/Runtime/UMG/Public/Components/TileView.h
--- a/Engine/Source/Runtime/UMG/Public/Components/TileView.h
+++ b/Engine/Source/Runtime/UMG/Public/Components/TileView.h
@@ -63,8 +63,8 @@ protected:
 		{
 			// This tells the underlying ListView to expect entries with the final width/height equal to EntryWidth/EntryHeight + spacing
 			// It allows the entry widget to always occupy the entire EntryWidth/EntryHeight.
-			Args.EntryWidth = GetHorizontalEntrySpacing();
-			Args.EntryHeight = GetVerticalEntrySpacing();
+			Args.EntryWidth = GetTotalEntryWidth();
+			Args.EntryHeight = GetTotalEntryHeight();
 		}
 		else
 		{
