Sequencer: Fix a crash in FWidgetMaterialAccessor when the widget is replaced and no longer valid

diff --git a/Engine/Source/Runtime/UMG/Private/Animation/MovieSceneWidgetMaterialSystem.cpp b/Engine/Source/Runtime/UMG/Private/Animation/MovieSceneWidgetMaterialSystem.cpp
--- a/Engine/Source/Runtime/UMG/Private/Animation/MovieSceneWidgetMaterialSystem.cpp
+++ b/Engine/Source/Runtime/UMG/Private/Animation/MovieSceneWidgetMaterialSystem.cpp
@@ -24,7 +24,7 @@ namespace UE::MovieScene
 {
 
 FWidgetMaterialAccessor::FWidgetMaterialAccessor(const FWidgetMaterialKey& InKey)
-	: Widget(CastChecked<UWidget>(InKey.Object.ResolveObjectPtr()))
+	: Widget(CastChecked<UWidget>(InKey.Object.ResolveObjectPtr(), ECastCheckedType::NullAllowed))
 {
 	// Only assign the material handle if the widget itself is still valid
 	// otherwise accessors constructed with the material handle will consider it valid
