Sequencer: Added custom accessor for UCanvasPanelSlot::LayoutData.Offsets

diff --git a/Engine/Source/Runtime/UMG/Private/Animation/MovieSceneUMGComponentTypes.cpp b/Engine/Source/Runtime/UMG/Private/Animation/MovieSceneUMGComponentTypes.cpp
--- a/Engine/Source/Runtime/UMG/Private/Animation/MovieSceneUMGComponentTypes.cpp
+++ b/Engine/Source/Runtime/UMG/Private/Animation/MovieSceneUMGComponentTypes.cpp
@@ -10,6 +10,7 @@
 #include "Systems/MovieScenePiecewiseDoubleBlenderSystem.h"
 
 #include "Components/Widget.h"
+#include "Components/CanvasPanelSlot.h"
 
 namespace UE
 {
@@ -84,6 +85,19 @@ static void SetRenderTransform(UObject* Object, const FIntermediateWidgetTransfo
 	CastChecked<UWidget>(Object)->SetRenderTransform(Transform);
 }
 
+FIntermediateMargin GetLayoutDataOffsets(const UObject* Object)
+{
+	FIntermediateMargin Margin;
+	ConvertOperationalProperty(CastChecked<UCanvasPanelSlot>(Object)->GetOffsets(), Margin);
+	return Margin;
+}
+void SetLayoutDataOffsets(UObject* Object, const FIntermediateMargin& InOffsets)
+{
+	FMargin Margin;
+	ConvertOperationalProperty(InOffsets, Margin);
+	CastChecked<UCanvasPanelSlot>(Object)->SetOffsets(Margin);
+}
+
 
 FMovieSceneUMGComponentTypes::FMovieSceneUMGComponentTypes()
 {
@@ -133,6 +147,7 @@ FMovieSceneUMGComponentTypes::FMovieSceneUMGComponentTypes()
 	FMovieSceneTracksComponentTypes::Get()->Accessors.Float.Add(UWidget::StaticClass(), "RenderOpacity", &GetRenderOpacity, &SetRenderOpacity);
 
 	CustomWidgetTransformAccessors.Add(UWidget::StaticClass(), "RenderTransform", &GetRenderTransform, &SetRenderTransform);
+	CustomMarginAccessors.Add(UCanvasPanelSlot::StaticClass(), "LayoutData.Offsets", &GetLayoutDataOffsets, &SetLayoutDataOffsets);
 
 	BuiltInComponents->PropertyRegistry.DefineCompositeProperty(Margin, TEXT("Apply FMargin Properties"))
 	.AddComposite(BuiltInComponents->DoubleResult[0], &FIntermediateMargin::Left)
@@ -140,6 +155,7 @@ FMovieSceneUMGComponentTypes::FMovieSceneUMGComponentTypes()
 	.AddComposite(BuiltInComponents->DoubleResult[2], &FIntermediateMargin::Right)
 	.AddComposite(BuiltInComponents->DoubleResult[3], &FIntermediateMargin::Bottom)
 	.SetBlenderSystem<UMovieScenePiecewiseDoubleBlenderSystem>()
+	.SetCustomAccessors(&CustomMarginAccessors)
 	.Commit();
 
 	BuiltInComponents->PropertyRegistry.DefineCompositeProperty(WidgetTransform, TEXT("Call UUserWidget::SetRenderTransform"))
diff --git a/Engine/Source/Runtime/UMG/Public/Animation/MovieSceneUMGComponentTypes.h b/Engine/Source/Runtime/UMG/Public/Animation/MovieSceneUMGComponentTypes.h
--- a/Engine/Source/Runtime/UMG/Public/Animation/MovieSceneUMGComponentTypes.h
+++ b/Engine/Source/Runtime/UMG/Public/Animation/MovieSceneUMGComponentTypes.h
@@ -64,6 +64,7 @@ struct UMG_API FMovieSceneUMGComponentTypes
 	TComponentTypeID<FWidgetMaterialPath> WidgetMaterialPath;
 	TComponentTypeID<FWidgetMaterialHandle> WidgetMaterialHandle;
 
+	TCustomPropertyRegistration<FMarginTraits, 1> CustomMarginAccessors;
 	TCustomPropertyRegistration<FWidgetTransformPropertyTraits, 1> CustomWidgetTransformAccessors;
 
 	static void Destroy();
