Fix for USpinBox widget's GetValue

diff --git a/Engine/Source/Runtime/UMG/Private/Components/SpinBox.cpp b/Engine/Source/Runtime/UMG/Private/Components/SpinBox.cpp
--- a/Engine/Source/Runtime/UMG/Private/Components/SpinBox.cpp
+++ b/Engine/Source/Runtime/UMG/Private/Components/SpinBox.cpp
@@ -464,6 +464,13 @@ void USpinBox::HandleOnValueChanged(float InValue)
 {
 	if ( !IsDesignTime() )
 	{
+		PRAGMA_DISABLE_DEPRECATION_WARNINGS
+		if (Value != InValue)
+		{
+			Value = InValue;
+			BroadcastFieldValueChanged(FFieldNotificationClassDescriptor::Value);
+		}
+		PRAGMA_ENABLE_DEPRECATION_WARNINGS
 		OnValueChanged.Broadcast(InValue);
 	}
 }
@@ -472,6 +479,13 @@ void USpinBox::HandleOnValueCommitted(float InValue, ETextCommit::Type CommitMet
 {
 	if ( !IsDesignTime() )
 	{
+		PRAGMA_DISABLE_DEPRECATION_WARNINGS
+		if (Value != InValue)
+		{
+			Value = InValue;
+			BroadcastFieldValueChanged(FFieldNotificationClassDescriptor::Value);
+		}
+		PRAGMA_ENABLE_DEPRECATION_WARNINGS
 		OnValueCommitted.Broadcast(InValue, CommitMethod);
 	}
 }
