[Engine]

diff --git a/Engine/Source/Runtime/UMG/Public/Blueprint/GameViewportSubsystem.h b/Engine/Source/Runtime/UMG/Public/Blueprint/GameViewportSubsystem.h
--- a/Engine/Source/Runtime/UMG/Public/Blueprint/GameViewportSubsystem.h
+++ b/Engine/Source/Runtime/UMG/Public/Blueprint/GameViewportSubsystem.h
@@ -18,7 +18,7 @@ class SConstraintCanvas;
  * The default value fills the entire screen / player region.
  */
 USTRUCT(BlueprintType)
-struct UMG_API FGameViewportWidgetSlot
+struct FGameViewportWidgetSlot
 {
 	GENERATED_BODY()
 
@@ -40,47 +40,47 @@ struct UMG_API FGameViewportWidgetSlot
 /**
  * 
  */
-UCLASS()
-class UMG_API UGameViewportSubsystem : public UEngineSubsystem
+UCLASS(MinimalAPI)
+class UGameViewportSubsystem : public UEngineSubsystem
 {
 	GENERATED_BODY()
 
 public:
 	//~ Begin Subsystem
-	virtual void Initialize(FSubsystemCollectionBase& Collection) override;
-	virtual void Deinitialize() override;
+	UMG_API virtual void Initialize(FSubsystemCollectionBase& Collection) override;
+	UMG_API virtual void Deinitialize() override;
 	//~ End Subsystem
 
-	static UGameViewportSubsystem* Get();
-	static UGameViewportSubsystem* Get(UWorld* World);
+	static UMG_API UGameViewportSubsystem* Get();
+	static UMG_API UGameViewportSubsystem* Get(UWorld* World);
 	
 public:
 	/* @return true if the widget was added to the viewport using AddWidget or AddWidgetForPlayer. */
 	UFUNCTION(BlueprintPure, BlueprintCosmetic, Category = "User Interface")
-	bool IsWidgetAdded(const UWidget* Widget) const;
+	UMG_API bool IsWidgetAdded(const UWidget* Widget) const;
 
 	/** Adds it to the game's viewport. */
 	UFUNCTION(BlueprintCallable, BlueprintCosmetic, Category = "User Interface")
-	bool AddWidget(UWidget* Widget, FGameViewportWidgetSlot Slot);
+	UMG_API bool AddWidget(UWidget* Widget, FGameViewportWidgetSlot Slot);
 
 	/**
 	 * Adds the widget to the game's viewport in the section dedicated to the player. This is valuable in a split screen
 	 * game where you need to only show a widget over a player's portion of the viewport.
 	 */
 	UFUNCTION(BlueprintCallable, BlueprintCosmetic, Category = "User Interface")
-	bool AddWidgetForPlayer(UWidget* Widget, ULocalPlayer* Player, FGameViewportWidgetSlot Slot);
+	UMG_API bool AddWidgetForPlayer(UWidget* Widget, ULocalPlayer* Player, FGameViewportWidgetSlot Slot);
 
 	/** Removes the widget from the viewport. */
 	UFUNCTION(BlueprintCallable, BlueprintCosmetic, Category = "User Interface")
-	void RemoveWidget(UWidget* Widget);
+	UMG_API void RemoveWidget(UWidget* Widget);
 
 	/* The slot info from previously added widget or info that is store for later. */
 	UFUNCTION(BlueprintCallable, BlueprintCosmetic, Category = "User Interface")
-	FGameViewportWidgetSlot GetWidgetSlot(const UWidget* Widget) const;
+	UMG_API FGameViewportWidgetSlot GetWidgetSlot(const UWidget* Widget) const;
 
 	/* Update the slot info of a previously added widget or Store the slot info for later use. */
 	UFUNCTION(BlueprintCallable, BlueprintCosmetic, Category = "User Interface")
-	void SetWidgetSlot(UWidget* Widget, FGameViewportWidgetSlot Slot);
+	UMG_API void SetWidgetSlot(UWidget* Widget, FGameViewportWidgetSlot Slot);
 
 	/** */
 	DECLARE_EVENT_TwoParams(UGameViewportSubsystem, FWidgetAddedEvent, UWidget*, ULocalPlayer*);
@@ -99,13 +99,13 @@ public:
 	 * by DPI, it ends up in the expected position.
 	 */
 	UFUNCTION(BlueprintCallable, Category = "GameViewportWidgetSlot")
-	static FGameViewportWidgetSlot SetWidgetSlotPosition(FGameViewportWidgetSlot Slot, const UWidget* Widget, FVector2D Position, bool bRemoveDPIScale);
+	static UMG_API FGameViewportWidgetSlot SetWidgetSlotPosition(FGameViewportWidgetSlot Slot, const UWidget* Widget, FVector2D Position, bool bRemoveDPIScale);
 
 	/**
 	 * Helper function to set the desired size in the viewport for the Slot.
 	 */
 	UFUNCTION(BlueprintCallable, Category = "GameViewportWidgetSlot")
-	static FGameViewportWidgetSlot SetWidgetSlotDesiredSize(FGameViewportWidgetSlot Slot, FVector2D Size);
+	static UMG_API FGameViewportWidgetSlot SetWidgetSlotDesiredSize(FGameViewportWidgetSlot Slot, FVector2D Size);
 
 private:
 	struct FSlotInfo
diff --git a/Engine/Source/Runtime/UMG/Public/Blueprint/IUserListEntry.h b/Engine/Source/Runtime/UMG/Public/Blueprint/IUserListEntry.h
--- a/Engine/Source/Runtime/UMG/Public/Blueprint/IUserListEntry.h
+++ b/Engine/Source/Runtime/UMG/Public/Blueprint/IUserListEntry.h
@@ -17,25 +17,25 @@ class IObjectTableRow;
  *
  * Note: To be usable as an entry for UListView, UTileView, or UTreeView, implement IUserObjectListEntry instead
  */
-UINTERFACE(BlueprintType)
-class UMG_API UUserListEntry : public UInterface
+UINTERFACE(BlueprintType, MinimalAPI)
+class UUserListEntry : public UInterface
 {
 	GENERATED_UINTERFACE_BODY()
 };
 
-class UMG_API IUserListEntry : public IInterface
+class IUserListEntry : public IInterface
 {
 	GENERATED_IINTERFACE_BODY()
 
 public:
 	/** Returns true if the item represented by this entry is currently selected in the owning list view. */
-	bool IsListItemSelected() const;
+	UMG_API bool IsListItemSelected() const;
 	
 	/** Returns true if the item represented by this entry is currently expanded and showing its children. Tree view entries only. */
-	bool IsListItemExpanded() const;
+	UMG_API bool IsListItemExpanded() const;
 	
 	/** Returns the list view that contains this entry. */
-	UListViewBase* GetOwningListView() const;
+	UMG_API UListViewBase* GetOwningListView() const;
 
 	/**
 	 * Advanced native-only option for specific rows to preclude themselves from any kind of selection.
@@ -46,35 +46,35 @@ public:
 
 public:
 	/** Functionality largely for "internal" use/plumbing - see SObjectTableRow's usage. You shouldn't have any cause to call these directly. */
-	static void ReleaseEntry(UUserWidget& ListEntryWidget);
-	static void UpdateItemSelection(UUserWidget& ListEntryWidget, bool bIsSelected);
-	static void UpdateItemExpansion(UUserWidget& ListEntryWidget, bool bIsExpanded);
+	static UMG_API void ReleaseEntry(UUserWidget& ListEntryWidget);
+	static UMG_API void UpdateItemSelection(UUserWidget& ListEntryWidget, bool bIsSelected);
+	static UMG_API void UpdateItemExpansion(UUserWidget& ListEntryWidget, bool bIsExpanded);
 
 protected:
 	/** 
 	 * These follow the same pattern as the NativeOn[X] methods in UUserWidget - super calls are expected in order to route the event to BP.
 	 * See the BP events below for descriptions.
 	 */
-	virtual void NativeOnItemSelectionChanged(bool bIsSelected);
-	virtual void NativeOnItemExpansionChanged(bool bIsExpanded);
-	virtual void NativeOnEntryReleased();
+	UMG_API virtual void NativeOnItemSelectionChanged(bool bIsSelected);
+	UMG_API virtual void NativeOnItemExpansionChanged(bool bIsExpanded);
+	UMG_API virtual void NativeOnEntryReleased();
 
 	/** Called when the selection state of the item represented by this entry changes. */
 	UFUNCTION(BlueprintImplementableEvent, Category = UserListEntry, meta = (DisplayName = "On Item Selection Changed"))
-	void BP_OnItemSelectionChanged(bool bIsSelected);
+	UMG_API void BP_OnItemSelectionChanged(bool bIsSelected);
 
 	/** Called when the expansion state of the item represented by this entry changes. Tree view entries only. */
 	UFUNCTION(BlueprintImplementableEvent, Category = UserListEntry, meta = (DisplayName = "On Item Expansion Changed"))
-	void BP_OnItemExpansionChanged(bool bIsExpanded);
+	UMG_API void BP_OnItemExpansionChanged(bool bIsExpanded);
 
 	/** Called when this entry is released from the owning table and no longer represents any list item */
 	UFUNCTION(BlueprintImplementableEvent, Category = UserListEntry, meta = (DisplayName = "On Entry Released"))
-	void BP_OnEntryReleased();
+	UMG_API void BP_OnEntryReleased();
 };
 
 /** Static library to supply "for free" functionality to widgets that implement IUserListEntry */
-UCLASS()
-class UMG_API UUserListEntryLibrary : public UBlueprintFunctionLibrary
+UCLASS(MinimalAPI)
+class UUserListEntryLibrary : public UBlueprintFunctionLibrary
 {
 	GENERATED_BODY()
 
@@ -84,19 +84,19 @@ public:
 	 * @param UserListEntry Note: Visually not transmitted, but this defaults to "self". No need to hook up if calling internally.
 	 */
 	UFUNCTION(BlueprintPure, Category = UserListEntry, meta = (DefaultToSelf = UserListEntry))
-	static bool IsListItemSelected(TScriptInterface<IUserListEntry> UserListEntry);
+	static UMG_API bool IsListItemSelected(TScriptInterface<IUserListEntry> UserListEntry);
 
 	/** 
 	 * Returns true if the item represented by this entry is currently expanded and showing its children. Tree view entries only.
 	 * @param UserListEntry Note: Visually not transmitted, but this defaults to "self". No need to hook up if calling internally.
 	 */
 	UFUNCTION(BlueprintPure, Category = UserListEntry, meta = (DefaultToSelf = UserListEntry))
-	static bool IsListItemExpanded(TScriptInterface<IUserListEntry> UserListEntry);
+	static UMG_API bool IsListItemExpanded(TScriptInterface<IUserListEntry> UserListEntry);
 
 	/** 
 	 * Returns the list view that contains this entry.
 	 * @param UserListEntry Note: Visually not transmitted, but this defaults to "self". No need to hook up if calling internally.
 	 */
 	UFUNCTION(BlueprintPure, Category = UserListEntry, meta = (DefaultToSelf = UserListEntry))
-	static UListViewBase* GetOwningListView(TScriptInterface<IUserListEntry> UserListEntry);
-};
\ No newline at end of file
+	static UMG_API UListViewBase* GetOwningListView(TScriptInterface<IUserListEntry> UserListEntry);
+};
diff --git a/Engine/Source/Runtime/UMG/Public/Blueprint/IUserObjectListEntry.h b/Engine/Source/Runtime/UMG/Public/Blueprint/IUserObjectListEntry.h
--- a/Engine/Source/Runtime/UMG/Public/Blueprint/IUserObjectListEntry.h
+++ b/Engine/Source/Runtime/UMG/Public/Blueprint/IUserObjectListEntry.h
@@ -9,13 +9,13 @@
  * Required interface for any UUserWidget class to be usable as entry widget in a stock UMG list view - ListView, TileView, and TreeView
  * Provides a change event and getter for the object item the entry is assigned to represent by the owning list view (in addition to functionality from IUserListEntry)
  */
-UINTERFACE()
-class UMG_API UUserObjectListEntry : public UUserListEntry
+UINTERFACE(MinimalAPI)
+class UUserObjectListEntry : public UUserListEntry
 {
 	GENERATED_UINTERFACE_BODY()
 };
 
-class UMG_API IUserObjectListEntry : public IUserListEntry
+class IUserObjectListEntry : public IUserListEntry
 {
 	GENERATED_IINTERFACE_BODY()
 
@@ -30,22 +30,22 @@ public:
 
 protected:
 	/** Follows the same pattern as the NativeOn[X] methods in UUserWidget - super calls are expected in order to route the event to BP. */
-	virtual void NativeOnListItemObjectSet(UObject* ListItemObject);
+	UMG_API virtual void NativeOnListItemObjectSet(UObject* ListItemObject);
 	
 	/** Called when this entry is assigned a new item object to represent by the owning list view */
 	UFUNCTION(BlueprintImplementableEvent, Category = ObjectListEntry)
-	void OnListItemObjectSet(UObject* ListItemObject);
+	UMG_API void OnListItemObjectSet(UObject* ListItemObject);
 
 private:
-	UObject* GetListItemObjectInternal() const;
+	UMG_API UObject* GetListItemObjectInternal() const;
 	
 	template <typename> friend class SObjectTableRow;
-	static void SetListItemObject(UUserWidget& ListEntryWidget, UObject* ListItemObject);
+	static UMG_API void SetListItemObject(UUserWidget& ListEntryWidget, UObject* ListItemObject);
 };
 
 /** Static library to supply "for free" functionality to widgets that implement IUserListEntry */
-UCLASS()
-class UMG_API UUserObjectListEntryLibrary : public UBlueprintFunctionLibrary
+UCLASS(MinimalAPI)
+class UUserObjectListEntryLibrary : public UBlueprintFunctionLibrary
 {
 	GENERATED_BODY()
 
@@ -55,5 +55,5 @@ public:
 	 * @param UserObjectListEntry Note: Visually not transmitted, but this defaults to "self". No need to hook up if calling internally.
 	 */
 	UFUNCTION(BlueprintPure, Category = UserObjectListEntry, meta = (DefaultToSelf = UserObjectListEntry))
-	static UObject* GetListItemObject(TScriptInterface<IUserObjectListEntry> UserObjectListEntry);
-};
\ No newline at end of file
+	static UMG_API UObject* GetListItemObject(TScriptInterface<IUserObjectListEntry> UserObjectListEntry);
+};
diff --git a/Engine/Source/Runtime/UMG/Public/Blueprint/SlateBlueprintLibrary.h b/Engine/Source/Runtime/UMG/Public/Blueprint/SlateBlueprintLibrary.h
--- a/Engine/Source/Runtime/UMG/Public/Blueprint/SlateBlueprintLibrary.h
+++ b/Engine/Source/Runtime/UMG/Public/Blueprint/SlateBlueprintLibrary.h
@@ -10,8 +10,8 @@
 #include "Kismet/BlueprintFunctionLibrary.h"
 #include "SlateBlueprintLibrary.generated.h"
 
-UCLASS(meta=(ScriptName="SlateLibrary"))
-class UMG_API USlateBlueprintLibrary : public UBlueprintFunctionLibrary
+UCLASS(meta=(ScriptName="SlateLibrary"), MinimalAPI)
+class USlateBlueprintLibrary : public UBlueprintFunctionLibrary
 {
 	GENERATED_UCLASS_BODY()
 
@@ -23,7 +23,7 @@ public:
 	 * @return true if the provided location in absolute coordinates is within the bounds of this geometry.
 	 */
 	 UFUNCTION(BlueprintPure, Category="User Interface|Geometry")
-	static bool IsUnderLocation(const FGeometry& Geometry, const FVector2D& AbsoluteCoordinate);
+	static UMG_API bool IsUnderLocation(const FGeometry& Geometry, const FVector2D& AbsoluteCoordinate);
 
 	/**
 	 * Absolute coordinates could be either desktop or window space depending on what space the root of the widget hierarchy is in.
@@ -31,7 +31,7 @@ public:
 	 * @return Transforms AbsoluteCoordinate into the local space of this Geometry.
 	 */
 	UFUNCTION(BlueprintPure, Category="User Interface|Geometry")
-	static FVector2D AbsoluteToLocal(const FGeometry& Geometry, FVector2D AbsoluteCoordinate);
+	static UMG_API FVector2D AbsoluteToLocal(const FGeometry& Geometry, FVector2D AbsoluteCoordinate);
 
 	/**
 	 * Translates local coordinates into absolute coordinates
@@ -41,39 +41,39 @@ public:
 	 * @return  Absolute coordinates
 	 */
 	UFUNCTION(BlueprintPure, Category="User Interface|Geometry")
-	static FVector2D LocalToAbsolute(const FGeometry& Geometry, FVector2D LocalCoordinate);
+	static UMG_API FVector2D LocalToAbsolute(const FGeometry& Geometry, FVector2D LocalCoordinate);
 
 	/** Returns the local top/left of the geometry in local space. */
 	UFUNCTION(BlueprintPure, Category = "User Interface|Geometry")
-	static FVector2D GetLocalTopLeft(const FGeometry& Geometry);
+	static UMG_API FVector2D GetLocalTopLeft(const FGeometry& Geometry);
 
 	/** Returns the size of the geometry in local space. */
 	UFUNCTION(BlueprintPure, Category="User Interface|Geometry")
-	static FVector2D GetLocalSize(const FGeometry& Geometry);
+	static UMG_API FVector2D GetLocalSize(const FGeometry& Geometry);
 
 	/** Returns the size of the geometry in absolute space. */
 	UFUNCTION(BlueprintPure, Category="User Interface|Geometry")
-	static FVector2D GetAbsoluteSize(const FGeometry& Geometry);
+	static UMG_API FVector2D GetAbsoluteSize(const FGeometry& Geometry);
 
 	/**  */
 	UFUNCTION(BlueprintPure, Category = "User Interface|Geometry")
-	static float TransformScalarAbsoluteToLocal(const FGeometry& Geometry, float AbsoluteScalar);
+	static UMG_API float TransformScalarAbsoluteToLocal(const FGeometry& Geometry, float AbsoluteScalar);
 
 	/**  */
 	UFUNCTION(BlueprintPure, Category = "User Interface|Geometry")
-	static float TransformScalarLocalToAbsolute(const FGeometry& Geometry, float LocalScalar);
+	static UMG_API float TransformScalarLocalToAbsolute(const FGeometry& Geometry, float LocalScalar);
 
 	/**  */
 	UFUNCTION(BlueprintPure, Category = "User Interface|Geometry")
-	static FVector2D TransformVectorAbsoluteToLocal(const FGeometry& Geometry, FVector2D AbsoluteVector);
+	static UMG_API FVector2D TransformVectorAbsoluteToLocal(const FGeometry& Geometry, FVector2D AbsoluteVector);
 
 	/**  */
 	UFUNCTION(BlueprintPure, Category = "User Interface|Geometry")
-	static FVector2D TransformVectorLocalToAbsolute(const FGeometry& Geometry, FVector2D LocalVector);
+	static UMG_API FVector2D TransformVectorLocalToAbsolute(const FGeometry& Geometry, FVector2D LocalVector);
 
 	/** Returns whether brushes A and B are identical. */
 	UFUNCTION(BlueprintPure, meta = (DisplayName = "Equal (SlateBrush)", CompactNodeTitle = "=="), Category = "SlateBrush")
-	static bool EqualEqual_SlateBrush(const FSlateBrush& A, const FSlateBrush& B);
+	static UMG_API bool EqualEqual_SlateBrush(const FSlateBrush& A, const FSlateBrush& B);
 
 	/**
 	 * Translates local coordinate of the geometry provided into local viewport coordinates.
@@ -85,7 +85,7 @@ public:
 	 * what you would use.
 	 */
 	UFUNCTION(BlueprintPure, Category="User Interface|Geometry", meta=( WorldContext="WorldContextObject" ))
-	static void LocalToViewport(UObject* WorldContextObject, const FGeometry& Geometry, FVector2D LocalCoordinate, FVector2D& PixelPosition, FVector2D& ViewportPosition);
+	static UMG_API void LocalToViewport(UObject* WorldContextObject, const FGeometry& Geometry, FVector2D LocalCoordinate, FVector2D& PixelPosition, FVector2D& ViewportPosition);
 
 	/**
 	 * Translates absolute coordinate in desktop space of the geometry provided into local viewport coordinates.
@@ -97,25 +97,25 @@ public:
 	 * what you would use.
 	 */
 	UFUNCTION(BlueprintPure, Category="User Interface|Geometry", meta=( WorldContext="WorldContextObject" ))
-	static void AbsoluteToViewport(UObject* WorldContextObject, FVector2D AbsoluteDesktopCoordinate, FVector2D& PixelPosition, FVector2D& ViewportPosition);
+	static UMG_API void AbsoluteToViewport(UObject* WorldContextObject, FVector2D AbsoluteDesktopCoordinate, FVector2D& PixelPosition, FVector2D& ViewportPosition);
 
 	/**
 	 * Translates a screen position in pixels into the local space of a widget with the given geometry. 
 	 * If bIncludeWindowPosition is true, then this method will also remove the game window's position (useful when in windowed mode).
 	 */
 	UFUNCTION(BlueprintPure, Category="User Interface|Geometry", meta=( WorldContext="WorldContextObject", DisplayName="ScreenToLocal" ))
-	static void ScreenToWidgetLocal(UObject* WorldContextObject, const FGeometry& Geometry, FVector2D ScreenPosition, FVector2D& LocalCoordinate, bool bIncludeWindowPosition = false);
+	static UMG_API void ScreenToWidgetLocal(UObject* WorldContextObject, const FGeometry& Geometry, FVector2D ScreenPosition, FVector2D& LocalCoordinate, bool bIncludeWindowPosition = false);
 
 	/**
 	 * Translates a screen position in pixels into absolute application coordinates.
 	 * If bIncludeWindowPosition is true, then this method will also remove the game window's position (useful when in windowed mode).
 	 */
 	UFUNCTION(BlueprintPure, Category="User Interface|Geometry", meta=( WorldContext="WorldContextObject", DisplayName="ScreenToAbsolute" ))
-	static void ScreenToWidgetAbsolute(UObject* WorldContextObject, FVector2D ScreenPosition, FVector2D& AbsoluteCoordinate, bool bIncludeWindowPosition = false);
+	static UMG_API void ScreenToWidgetAbsolute(UObject* WorldContextObject, FVector2D ScreenPosition, FVector2D& AbsoluteCoordinate, bool bIncludeWindowPosition = false);
 
 	/**
 	 * Translates a screen position in pixels into the local space of the viewport widget.
 	 */
 	UFUNCTION(BlueprintPure, Category="User Interface|Geometry", meta=( WorldContext="WorldContextObject" ))
-	static void ScreenToViewport(UObject* WorldContextObject, FVector2D ScreenPosition, FVector2D& ViewportPosition);
+	static UMG_API void ScreenToViewport(UObject* WorldContextObject, FVector2D ScreenPosition, FVector2D& ViewportPosition);
 };
diff --git a/Engine/Source/Runtime/UMG/Public/Blueprint/UserWidget.h b/Engine/Source/Runtime/UMG/Public/Blueprint/UserWidget.h
--- a/Engine/Source/Runtime/UMG/Public/Blueprint/UserWidget.h
+++ b/Engine/Source/Runtime/UMG/Public/Blueprint/UserWidget.h
@@ -109,14 +109,14 @@ public:
  * allow script code to override OnPaint behavior.
  */
 USTRUCT(BlueprintType)
-struct UMG_API FPaintContext
+struct FPaintContext
 {
 	GENERATED_USTRUCT_BODY()
 
 public:
 
 	/** Don't ever use this constructor.  Needed for code generation. */
-	FPaintContext();
+	UMG_API FPaintContext();
 
 	FPaintContext(const FGeometry& InAllottedGeometry, const FSlateRect& InMyCullingRect, FSlateWindowElementList& InOutDrawElements, const int32 InLayerId, const FWidgetStyle& InWidgetStyle, const bool bInParentEnabled)
 		: AllottedGeometry(InAllottedGeometry)
@@ -151,7 +151,7 @@ public:
 };
 
 USTRUCT()
-struct UMG_API FNamedSlotBinding
+struct FNamedSlotBinding
 {
 	GENERATED_USTRUCT_BODY()
 
@@ -215,25 +215,25 @@ DECLARE_DYNAMIC_MULTICAST_DELEGATE_OneParam(FOnVisibilityChangedEvent, ESlateVis
 /**
  * A widget that enables UI extensibility through WidgetBlueprint.
  */
-UCLASS(Abstract, editinlinenew, BlueprintType, Blueprintable, meta=( DontUseGenericSpawnObject="True", DisableNativeTick) )
-class UMG_API UUserWidget : public UWidget, public INamedSlotInterface
+UCLASS(Abstract, editinlinenew, BlueprintType, Blueprintable, meta=( DontUseGenericSpawnObject="True", DisableNativeTick) , MinimalAPI)
+class UUserWidget : public UWidget, public INamedSlotInterface
 {
 	GENERATED_BODY()
 
 	friend class SObjectWidget;
 public:
-	UUserWidget(const FObjectInitializer& ObjectInitializer);
+	UMG_API UUserWidget(const FObjectInitializer& ObjectInitializer);
 
 	//~ Begin UObject interface
-	virtual class UWorld* GetWorld() const override;
-	virtual void PostDuplicate(bool bDuplicateForPIE) override;
-	virtual void BeginDestroy() override;
-	virtual void PostLoad() override;
+	UMG_API virtual class UWorld* GetWorld() const override;
+	UMG_API virtual void PostDuplicate(bool bDuplicateForPIE) override;
+	UMG_API virtual void BeginDestroy() override;
+	UMG_API virtual void PostLoad() override;
 	//~ End UObject Interface
 
-	void DuplicateAndInitializeFromWidgetTree(UWidgetTree* InWidgetTree, const TMap<FName, UWidget*>& NamedSlotContentToMerge);
+	UMG_API void DuplicateAndInitializeFromWidgetTree(UWidgetTree* InWidgetTree, const TMap<FName, UWidget*>& NamedSlotContentToMerge);
 
-	virtual bool Initialize();
+	UMG_API virtual bool Initialize();
 
 	EWidgetTickFrequency GetDesiredTickFrequency() const { return TickFrequency; }
 
@@ -243,29 +243,29 @@ public:
 	 * The child UserWidget will have the same WidgetTree as the parent UserWidget.
 	 * This function returns the parent UserWidget's BlueprintClass.
 	 */
-	UWidgetBlueprintGeneratedClass* GetWidgetTreeOwningClass() const;
+	UMG_API UWidgetBlueprintGeneratedClass* GetWidgetTreeOwningClass() const;
 
-	void UpdateCanTick();
+	UMG_API void UpdateCanTick();
 
 protected:
 	/** The function is implemented only in nativized widgets (automatically converted from BP to c++) */
 	virtual void InitializeNativeClassData() {}
 
-	void InitializeNamedSlots();
+	UMG_API void InitializeNamedSlots();
 
 public:
 	//~ Begin UVisual interface
-	virtual void ReleaseSlateResources(bool bReleaseChildren) override;
+	UMG_API virtual void ReleaseSlateResources(bool bReleaseChildren) override;
 	//~ End UVisual Interface
 
 	//~ Begin UWidget Interface
-	virtual void SynchronizeProperties() override;
+	UMG_API virtual void SynchronizeProperties() override;
 	//~ End UWidget Interface
 
 	//~ Begin UNamedSlotInterface Begin
-	virtual void GetSlotNames(TArray<FName>& SlotNames) const override;
-	virtual UWidget* GetContentForSlot(FName SlotName) const override;
-	virtual void SetContentForSlot(FName SlotName, UWidget* Content) override;
+	UMG_API virtual void GetSlotNames(TArray<FName>& SlotNames) const override;
+	UMG_API virtual UWidget* GetContentForSlot(FName SlotName) const override;
+	UMG_API virtual void SetContentForSlot(FName SlotName, UWidget* Content) override;
 	//~ UNamedSlotInterface End
 
 	/**
@@ -275,7 +275,7 @@ public:
 	 * @param ZOrder The higher the number, the more on top this widget will be.
 	 */
 	UFUNCTION(BlueprintCallable, BlueprintCosmetic, Category="User Interface|Viewport", meta=( AdvancedDisplay = "ZOrder" ))
-	void AddToViewport(int32 ZOrder = 0);
+	UMG_API void AddToViewport(int32 ZOrder = 0);
 
 	/**
 	 * Adds the widget to the game's viewport in a section dedicated to the player.  This is valuable in a split screen
@@ -284,14 +284,14 @@ public:
 	 * @param ZOrder The higher the number, the more on top this widget will be.
 	 */
 	UFUNCTION(BlueprintCallable, BlueprintCosmetic, Category="User Interface|Viewport", meta=( AdvancedDisplay = "ZOrder" ))
-	bool AddToPlayerScreen(int32 ZOrder = 0);
+	UMG_API bool AddToPlayerScreen(int32 ZOrder = 0);
 
 	/**
 	 * Removes the widget from the viewport.
 	 */
 	UE_DEPRECATED(5.1, "RemoveFromViewport is deprecated. Use RemoveFromParent instead.")
 	UFUNCTION(BlueprintCallable, BlueprintCosmetic, Category="User Interface|Viewport", meta=( DeprecatedFunction, DeprecationMessage="Use RemoveFromParent instead" ))
-	void RemoveFromViewport();
+	UMG_API void RemoveFromViewport();
 
 	/**
 	 * Sets the widgets position in the viewport.
@@ -301,47 +301,47 @@ public:
 	 * by DPI, it ends up in the expected position.
 	 */
 	UFUNCTION(BlueprintCallable, BlueprintCosmetic, Category="User Interface|Viewport")
-	void SetPositionInViewport(FVector2D Position, bool bRemoveDPIScale = true);
+	UMG_API void SetPositionInViewport(FVector2D Position, bool bRemoveDPIScale = true);
 
 	/*  */
 	UFUNCTION(BlueprintCallable, BlueprintCosmetic, Category="User Interface|Viewport")
-	void SetDesiredSizeInViewport(FVector2D Size);
+	UMG_API void SetDesiredSizeInViewport(FVector2D Size);
 
 	/*  */
 	UFUNCTION(BlueprintCallable, BlueprintCosmetic, Category="User Interface|Viewport")
-	void SetAnchorsInViewport(FAnchors Anchors);
+	UMG_API void SetAnchorsInViewport(FAnchors Anchors);
 
 	/*  */
 	UFUNCTION(BlueprintCallable, BlueprintCosmetic, Category="User Interface|Viewport")
-	void SetAlignmentInViewport(FVector2D Alignment);
+	UMG_API void SetAlignmentInViewport(FVector2D Alignment);
 
 	/*  */
 	UFUNCTION(BlueprintCallable, BlueprintCosmetic, Category = "User Interface|Viewport")
-	FAnchors GetAnchorsInViewport() const;
+	UMG_API FAnchors GetAnchorsInViewport() const;
 
 	/*  */
 	UFUNCTION(BlueprintCallable, BlueprintCosmetic, Category = "User Interface|Viewport")
-	FVector2D GetAlignmentInViewport() const;
+	UMG_API FVector2D GetAlignmentInViewport() const;
 
 	/*  */
 	UE_DEPRECATED(5.1, "GetIsVisible is deprecated. Please use IsInViewport instead.")
 	UFUNCTION(BlueprintPure, BlueprintCosmetic, Category="Appearance", meta=( DeprecatedFunction, DeprecationMessage="Use IsInViewport instead" ))
-	bool GetIsVisible() const;
+	UMG_API bool GetIsVisible() const;
 
 	/** Sets the visibility of the widget. */
-	virtual void SetVisibility(ESlateVisibility InVisibility) override;
+	UMG_API virtual void SetVisibility(ESlateVisibility InVisibility) override;
 
 	/** Sets the player context associated with this UI. */
-	void SetPlayerContext(const FLocalPlayerContext& InPlayerContext);
+	UMG_API void SetPlayerContext(const FLocalPlayerContext& InPlayerContext);
 
 	/** Gets the player context associated with this UI. */
-	const FLocalPlayerContext& GetPlayerContext() const;
+	UMG_API const FLocalPlayerContext& GetPlayerContext() const;
 
 	/**
 	 * Gets the local player associated with this UI.
 	 * @return The owning local player.
 	 */
-	virtual ULocalPlayer* GetOwningLocalPlayer() const override;
+	UMG_API virtual ULocalPlayer* GetOwningLocalPlayer() const override;
 	
 	/**
 	 * Gets the local player associated with this UI cast to the template type.
@@ -357,13 +357,13 @@ public:
 	 * Sets the player associated with this UI via LocalPlayer reference.
 	 * @param LocalPlayer The local player you want to be the conceptual owner of this UI.
 	 */
-	void SetOwningLocalPlayer(ULocalPlayer* LocalPlayer);
+	UMG_API void SetOwningLocalPlayer(ULocalPlayer* LocalPlayer);
 
 	/**
 	 * Gets the player controller associated with this UI.
 	 * @return The player controller that owns the UI.
 	 */
-	virtual APlayerController* GetOwningPlayer() const override;
+	UMG_API virtual APlayerController* GetOwningPlayer() const override;
 	
 	/**
 	 * Gets the player controller associated with this UI cast to the template type.
@@ -380,14 +380,14 @@ public:
 	 * @param LocalPlayerController The PlayerController of the local player you want to be the conceptual owner of this UI.
 	 */
 	UFUNCTION(BlueprintCallable, BlueprintCosmetic, Category="Player")
-	void SetOwningPlayer(APlayerController* LocalPlayerController);
+	UMG_API void SetOwningPlayer(APlayerController* LocalPlayerController);
 
 	/**
 	 * Gets the player pawn associated with this UI.
 	 * @return Gets the owning player pawn that's owned by the player controller assigned to this widget.
 	 */
 	UFUNCTION(BlueprintCallable, BlueprintCosmetic, Category="Player")
-	class APawn* GetOwningPlayerPawn() const;
+	UMG_API class APawn* GetOwningPlayerPawn() const;
 	
 	/**
 	 * Gets the player pawn associated with this UI cast to the template type.
@@ -422,7 +422,7 @@ public:
 	 * @return Gets the owning player camera manager that's owned by the player controller assigned to this widget.
 	 */
 	UFUNCTION(BlueprintCallable, BlueprintCosmetic, Category = "Player")
-	class APlayerCameraManager* GetOwningPlayerCameraManager() const;
+	UMG_API class APlayerCameraManager* GetOwningPlayerCameraManager() const;
 
 	/**
 	 * Gets the player camera manager associated with this UI cast to the template type.
@@ -441,7 +441,7 @@ public:
 	 * If you have one-time things to establish up-front (like binding callbacks to events on BindWidget properties), do so here.
 	 */
 	UFUNCTION(BlueprintImplementableEvent, BlueprintCosmetic, Category="User Interface")
-	void OnInitialized();
+	UMG_API void OnInitialized();
 
 	/**
 	 * Called by both the game and the editor.  Allows users to run initial setup for their widgets to better preview
@@ -456,7 +456,7 @@ public:
 	 * PreConstruct evaluation in the Widget Designer settings in the Editor Preferences.
 	 */
 	UFUNCTION(BlueprintImplementableEvent, BlueprintCosmetic, Category="User Interface")
-	void PreConstruct(bool IsDesignTime);
+	UMG_API void PreConstruct(bool IsDesignTime);
 
 	/**
 	 * Called after the underlying slate widget is constructed.  Depending on how the slate object is used
@@ -464,14 +464,14 @@ public:
 	 * If you need a true called-once-when-created event, use OnInitialized.
 	 */
 	UFUNCTION(BlueprintImplementableEvent, BlueprintCosmetic, Category="User Interface", meta=( Keywords="Begin Play" ))
-	void Construct();
+	UMG_API void Construct();
 
 	/**
 	 * Called when a widget is no longer referenced causing the slate resource to destroyed.  Just like
 	 * Construct this event can be called multiple times.
 	 */
 	UFUNCTION(BlueprintImplementableEvent, BlueprintCosmetic, Category="User Interface", meta=( Keywords="End Play, Destroy" ))
-	void Destruct();
+	UMG_API void Destruct();
 
 	/**
 	 * Ticks this widget.  Override in derived classes, but always call the parent implementation.
@@ -480,19 +480,19 @@ public:
 	 * @param  InDeltaTime  Real time passed since last tick
 	 */
 	UFUNCTION(BlueprintImplementableEvent, BlueprintCosmetic, Category="User Interface")
-	void Tick(FGeometry MyGeometry, float InDeltaTime);
+	UMG_API void Tick(FGeometry MyGeometry, float InDeltaTime);
 
 	/**
 	 * 
 	 */
 	UFUNCTION(BlueprintImplementableEvent, BlueprintCosmetic, Category="User Interface | Painting")
-	void OnPaint(UPARAM(ref) FPaintContext& Context) const;
+	UMG_API void OnPaint(UPARAM(ref) FPaintContext& Context) const;
 
 	/**
 	 * Gets a value indicating if the widget is interactive.
 	 */
 	UFUNCTION(BlueprintImplementableEvent, BlueprintCosmetic, Category="User Interface | Interaction")
-	bool IsInteractable() const;
+	UMG_API bool IsInteractable() const;
 
 	/**
 	 * Called when keyboard focus is given to this widget.  This event does not bubble.
@@ -502,7 +502,7 @@ public:
 	 * @return  Returns whether the event was handled, along with other possible actions
 	 */
 	UFUNCTION(BlueprintImplementableEvent, BlueprintCosmetic, Category="Input")
-	FEventReply OnFocusReceived(FGeometry MyGeometry, FFocusEvent InFocusEvent);
+	UMG_API FEventReply OnFocusReceived(FGeometry MyGeometry, FFocusEvent InFocusEvent);
 
 	/**
 	 * Called when this widget loses focus.  This event does not bubble.
@@ -510,7 +510,7 @@ public:
 	 * @param  InFocusEvent  FocusEvent
 	 */
 	UFUNCTION(BlueprintImplementableEvent, BlueprintCosmetic, Category="Input")
-	void OnFocusLost(FFocusEvent InFocusEvent);
+	UMG_API void OnFocusLost(FFocusEvent InFocusEvent);
 
 	/**
 	 * If focus is gained on on this widget or on a child widget and this widget is added
@@ -519,7 +519,7 @@ public:
 	 * @param  InFocusEvent  FocusEvent
 	 */
 	UFUNCTION(BlueprintImplementableEvent, BlueprintCosmetic, Category="Input")
-	void OnAddedToFocusPath(FFocusEvent InFocusEvent);
+	UMG_API void OnAddedToFocusPath(FFocusEvent InFocusEvent);
 
 	/**
 	 * If focus is lost on on this widget or on a child widget and this widget is
@@ -528,7 +528,7 @@ public:
 	 * @param  InFocusEvent  FocusEvent
 	 */
 	UFUNCTION(BlueprintImplementableEvent, BlueprintCosmetic, Category="Input")
-	void OnRemovedFromFocusPath(FFocusEvent InFocusEvent);
+	UMG_API void OnRemovedFromFocusPath(FFocusEvent InFocusEvent);
 
 	/**
 	 * Called after a character is entered while this widget has focus
@@ -538,7 +538,7 @@ public:
 	 * @return  Returns whether the event was handled, along with other possible actions
 	 */
 	UFUNCTION(BlueprintImplementableEvent, BlueprintCosmetic, Category="Input")
-	FEventReply OnKeyChar(FGeometry MyGeometry, FCharacterEvent InCharacterEvent);
+	UMG_API FEventReply OnKeyChar(FGeometry MyGeometry, FCharacterEvent InCharacterEvent);
 
 	/**
 	 * Called after a key (keyboard, controller, ...) is pressed when this widget or a child of this widget has focus
@@ -552,7 +552,7 @@ public:
 	 * @return  Returns whether the event was handled, along with other possible actions
 	 */
 	UFUNCTION(BlueprintImplementableEvent, Category="Input")
-	FEventReply OnPreviewKeyDown(FGeometry MyGeometry, FKeyEvent InKeyEvent);
+	UMG_API FEventReply OnPreviewKeyDown(FGeometry MyGeometry, FKeyEvent InKeyEvent);
 
 	/**
 	 * Called after a key (keyboard, controller, ...) is pressed when this widget has focus (this event bubbles if not handled)
@@ -562,7 +562,7 @@ public:
 	 * @return  Returns whether the event was handled, along with other possible actions
 	 */
 	UFUNCTION(BlueprintImplementableEvent, BlueprintCosmetic, Category="Input")
-	FEventReply OnKeyDown(FGeometry MyGeometry, FKeyEvent InKeyEvent);
+	UMG_API FEventReply OnKeyDown(FGeometry MyGeometry, FKeyEvent InKeyEvent);
 
 	/**
 	 * Called after a key (keyboard, controller, ...) is released when this widget has focus
@@ -572,7 +572,7 @@ public:
 	 * @return  Returns whether the event was handled, along with other possible actions
 	 */
 	UFUNCTION(BlueprintImplementableEvent, BlueprintCosmetic, Category="Input")
-	FEventReply OnKeyUp(FGeometry MyGeometry, FKeyEvent InKeyEvent);
+	UMG_API FEventReply OnKeyUp(FGeometry MyGeometry, FKeyEvent InKeyEvent);
 
 	/**
 	* Called when an analog value changes on a button that supports analog
@@ -582,7 +582,7 @@ public:
 	* @return  Returns whether the event was handled, along with other possible actions
 	*/
 	UFUNCTION(BlueprintImplementableEvent, Category = "Input")
-	FEventReply OnAnalogValueChanged(FGeometry MyGeometry, FAnalogInputEvent InAnalogInputEvent);
+	UMG_API FEventReply OnAnalogValueChanged(FGeometry MyGeometry, FAnalogInputEvent InAnalogInputEvent);
 
 	/**
 	 * The system calls this method to notify the widget that a mouse button was pressed within it. This event is bubbled.
@@ -592,7 +592,7 @@ public:
 	 * @return Whether the event was handled along with possible requests for the system to take action.
 	 */
 	UFUNCTION(BlueprintImplementableEvent, BlueprintCosmetic, Category="Mouse")
-	FEventReply OnMouseButtonDown(FGeometry MyGeometry, const FPointerEvent& MouseEvent);
+	UMG_API FEventReply OnMouseButtonDown(FGeometry MyGeometry, const FPointerEvent& MouseEvent);
 
 	/**
 	 * Just like OnMouseButtonDown, but tunnels instead of bubbling.
@@ -606,7 +606,7 @@ public:
 	 * @return Whether the event was handled along with possible requests for the system to take action.
 	 */
 	UFUNCTION(BlueprintImplementableEvent, BlueprintCosmetic, Category="Mouse")
-	FEventReply OnPreviewMouseButtonDown(FGeometry MyGeometry, const FPointerEvent& MouseEvent);
+	UMG_API FEventReply OnPreviewMouseButtonDown(FGeometry MyGeometry, const FPointerEvent& MouseEvent);
 
 	/**
 	 * The system calls this method to notify the widget that a mouse button was release within it. This event is bubbled.
@@ -616,7 +616,7 @@ public:
 	 * @return Whether the event was handled along with possible requests for the system to take action.
 	 */
 	UFUNCTION(BlueprintImplementableEvent, BlueprintCosmetic, Category="Mouse")
-	FEventReply OnMouseButtonUp(FGeometry MyGeometry, const FPointerEvent& MouseEvent);
+	UMG_API FEventReply OnMouseButtonUp(FGeometry MyGeometry, const FPointerEvent& MouseEvent);
 
 	/**
 	 * The system calls this method to notify the widget that a mouse moved within it. This event is bubbled.
@@ -626,7 +626,7 @@ public:
 	 * @return Whether the event was handled along with possible requests for the system to take action.
 	 */
 	UFUNCTION(BlueprintImplementableEvent, BlueprintCosmetic, Category="Mouse")
-	FEventReply OnMouseMove(FGeometry MyGeometry, const FPointerEvent& MouseEvent);
+	UMG_API FEventReply OnMouseMove(FGeometry MyGeometry, const FPointerEvent& MouseEvent);
 
 	/**
 	 * The system will use this event to notify a widget that the cursor has entered it. This event is NOT bubbled.
@@ -635,7 +635,7 @@ public:
 	 * @param MouseEvent Information about the input event
 	 */
 	UFUNCTION(BlueprintImplementableEvent, BlueprintCosmetic, Category="Mouse")
-	void OnMouseEnter(FGeometry MyGeometry, const FPointerEvent& MouseEvent);
+	UMG_API void OnMouseEnter(FGeometry MyGeometry, const FPointerEvent& MouseEvent);
 
 	/**
 	 * The system will use this event to notify a widget that the cursor has left it. This event is NOT bubbled.
@@ -643,7 +643,7 @@ public:
 	 * @param MouseEvent Information about the input event
 	 */
 	UFUNCTION(BlueprintImplementableEvent, BlueprintCosmetic, Category="Mouse")
-	void OnMouseLeave(const FPointerEvent& MouseEvent);
+	UMG_API void OnMouseLeave(const FPointerEvent& MouseEvent);
 
 	/**
 	 * Called when the mouse wheel is spun. This event is bubbled.
@@ -652,7 +652,7 @@ public:
 	 * @return  Returns whether the event was handled, along with other possible actions
 	 */
 	UFUNCTION(BlueprintImplementableEvent, BlueprintCosmetic, Category="Mouse")
-	FEventReply OnMouseWheel(FGeometry MyGeometry, const FPointerEvent& MouseEvent);
+	UMG_API FEventReply OnMouseWheel(FGeometry MyGeometry, const FPointerEvent& MouseEvent);
 
 	/**
 	 * Called when a mouse button is double clicked.  Override this in derived classes.
@@ -662,7 +662,7 @@ public:
 	 * @return  Returns whether the event was handled, along with other possible actions
 	 */
 	UFUNCTION(BlueprintImplementableEvent, BlueprintCosmetic, Category="Mouse")
-	FEventReply OnMouseButtonDoubleClick(FGeometry InMyGeometry, const FPointerEvent& InMouseEvent);
+	UMG_API FEventReply OnMouseButtonDoubleClick(FGeometry InMyGeometry, const FPointerEvent& InMouseEvent);
 
 	// TODO
 	//UFUNCTION(BlueprintImplementableEvent, Category="Mouse")
@@ -679,7 +679,7 @@ public:
 	 * @param  Operation     The drag operation that was detected.
 	 */
 	UFUNCTION(BlueprintImplementableEvent, BlueprintCosmetic, Category="Drag and Drop")
-	void OnDragDetected(FGeometry MyGeometry, const FPointerEvent& PointerEvent, UDragDropOperation*& Operation);
+	UMG_API void OnDragDetected(FGeometry MyGeometry, const FPointerEvent& PointerEvent, UDragDropOperation*& Operation);
 
 	/**
 	 * Called when the user cancels the drag operation, typically when they simply release the mouse button after
@@ -689,7 +689,7 @@ public:
 	 * @param  Operation     The drag operation that was canceled.
 	 */
 	UFUNCTION(BlueprintImplementableEvent, BlueprintCosmetic, Category="Drag and Drop")
-	void OnDragCancelled(const FPointerEvent& PointerEvent, UDragDropOperation* Operation);
+	UMG_API void OnDragCancelled(const FPointerEvent& PointerEvent, UDragDropOperation* Operation);
 	
 	/**
 	 * Called during drag and drop when the drag enters the widget.
@@ -699,7 +699,7 @@ public:
 	 * @param Operation      The drag operation that entered the widget.
 	 */
 	UFUNCTION(BlueprintImplementableEvent, BlueprintCosmetic, Category="Drag and Drop")
-	void OnDragEnter(FGeometry MyGeometry, FPointerEvent PointerEvent, UDragDropOperation* Operation);
+	UMG_API void OnDragEnter(FGeometry MyGeometry, FPointerEvent PointerEvent, UDragDropOperation* Operation);
 
 	/**
 	 * Called during drag and drop when the drag leaves the widget.
@@ -708,7 +708,7 @@ public:
 	 * @param Operation      The drag operation that entered the widget.
 	 */
 	UFUNCTION(BlueprintImplementableEvent, BlueprintCosmetic, Category="Drag and Drop")
-	void OnDragLeave(FPointerEvent PointerEvent, UDragDropOperation* Operation);
+	UMG_API void OnDragLeave(FPointerEvent PointerEvent, UDragDropOperation* Operation);
 
 	/**
 	 * Called during drag and drop when the the mouse is being dragged over a widget.
@@ -720,7 +720,7 @@ public:
 	 * @return 'true' to indicate that you handled the drag over operation.  Returning 'false' will cause the operation to continue to bubble up.
 	 */
 	UFUNCTION(BlueprintImplementableEvent, BlueprintCosmetic, Category="Drag and Drop")
-	bool OnDragOver(FGeometry MyGeometry, FPointerEvent PointerEvent, UDragDropOperation* Operation);
+	UMG_API bool OnDragOver(FGeometry MyGeometry, FPointerEvent PointerEvent, UDragDropOperation* Operation);
 
 	/**
 	 * Called when the user is dropping something onto a widget.  Ends the drag and drop operation, even if no widget handles this.
@@ -732,7 +732,7 @@ public:
 	 * @return 'true' to indicate you handled the drop operation.
 	 */
 	UFUNCTION(BlueprintImplementableEvent, BlueprintCosmetic, Category="Drag and Drop")
-	bool OnDrop(FGeometry MyGeometry, FPointerEvent PointerEvent, UDragDropOperation* Operation);
+	UMG_API bool OnDrop(FGeometry MyGeometry, FPointerEvent PointerEvent, UDragDropOperation* Operation);
 
 	/**
 	 * Called when the user performs a gesture on trackpad. This event is bubbled.
@@ -742,7 +742,7 @@ public:
 	 * @return  Returns whether the event was handled, along with other possible actions
 	 */
 	UFUNCTION(BlueprintImplementableEvent, BlueprintCosmetic, Category="Touch Input")
-	FEventReply OnTouchGesture(FGeometry MyGeometry, const FPointerEvent& GestureEvent);
+	UMG_API FEventReply OnTouchGesture(FGeometry MyGeometry, const FPointerEvent& GestureEvent);
 
 	/**
 	 * Called when a touchpad touch is started (finger down)
@@ -751,7 +751,7 @@ public:
 	 * @param InTouchEvent	The touch event generated
 	 */
 	UFUNCTION(BlueprintImplementableEvent, BlueprintCosmetic, Category="Touch Input")
-	FEventReply OnTouchStarted(FGeometry MyGeometry, const FPointerEvent& InTouchEvent);
+	UMG_API FEventReply OnTouchStarted(FGeometry MyGeometry, const FPointerEvent& InTouchEvent);
 	
 	/**
 	 * Called when a touchpad touch is moved  (finger moved)
@@ -760,7 +760,7 @@ public:
 	 * @param InTouchEvent	The touch event generated
 	 */
 	UFUNCTION(BlueprintImplementableEvent, BlueprintCosmetic, Category="Touch Input")
-	FEventReply OnTouchMoved(FGeometry MyGeometry, const FPointerEvent& InTouchEvent);
+	UMG_API FEventReply OnTouchMoved(FGeometry MyGeometry, const FPointerEvent& InTouchEvent);
 
 	/**
 	 * Called when a touchpad touch is ended (finger lifted)
@@ -769,7 +769,7 @@ public:
 	 * @param InTouchEvent	The touch event generated
 	 */
 	UFUNCTION(BlueprintImplementableEvent, BlueprintCosmetic, Category="Touch Input")
-	FEventReply OnTouchEnded(FGeometry MyGeometry, const FPointerEvent& InTouchEvent);
+	UMG_API FEventReply OnTouchEnded(FGeometry MyGeometry, const FPointerEvent& InTouchEvent);
 	
 	/**
 	 * Called when motion is detected (controller or device)
@@ -779,25 +779,25 @@ public:
 	 * @param MotionEvent	The motion event generated
 	 */
 	UFUNCTION(BlueprintImplementableEvent, BlueprintCosmetic, Category="Touch Input")
-	FEventReply OnMotionDetected(FGeometry MyGeometry, FMotionEvent InMotionEvent);
+	UMG_API FEventReply OnMotionDetected(FGeometry MyGeometry, FMotionEvent InMotionEvent);
 
 	/**
 	 * Called when mouse capture is lost if it was owned by this widget.
 	 */
 	UFUNCTION(BlueprintImplementableEvent, BlueprintCosmetic, Category="Touch Input")
-	void OnMouseCaptureLost();
+	UMG_API void OnMouseCaptureLost();
 
 	/**
 	 * Cancels any pending Delays or timer callbacks for this widget.
 	 */
 	UFUNCTION(BlueprintCallable, Category = "Delay")
-	void CancelLatentActions();
+	UMG_API void CancelLatentActions();
 
 	/**
 	* Cancels any pending Delays or timer callbacks for this widget, and stops all active animations on the widget.
 	*/
 	UFUNCTION(BlueprintCallable, Category = "Delay")
-	void StopAnimationsAndLatentActions();
+	UMG_API void StopAnimationsAndLatentActions();
 
 	/**
 	* Called when a touchpad force has changed (user pressed down harder or let up)
@@ -806,7 +806,7 @@ public:
 	* @param InTouchEvent	The touch event generated
 	*/
 	UFUNCTION(BlueprintImplementableEvent, BlueprintCosmetic, Category = "Touch Input")
-	FEventReply OnTouchForceChanged(FGeometry MyGeometry, const FPointerEvent& InTouchEvent);
+	UMG_API FEventReply OnTouchForceChanged(FGeometry MyGeometry, const FPointerEvent& InTouchEvent);
 
 public:
 
@@ -816,7 +816,7 @@ public:
 	 * @param Delegate the delegate to call when the animation's state changes
 	 */
 	UFUNCTION(BlueprintCallable, Category=Animation)
-	void BindToAnimationStarted(UWidgetAnimation* Animation, FWidgetAnimationDynamicEvent Delegate);
+	UMG_API void BindToAnimationStarted(UWidgetAnimation* Animation, FWidgetAnimationDynamicEvent Delegate);
 
 	/**
 	 * Unbind an animation started delegate.
@@ -824,10 +824,10 @@ public:
 	 * @param Delegate the delegate to call when the animation's state changes
 	 */
 	UFUNCTION(BlueprintCallable, Category = Animation)
-	void UnbindFromAnimationStarted(UWidgetAnimation* Animation, FWidgetAnimationDynamicEvent Delegate);
+	UMG_API void UnbindFromAnimationStarted(UWidgetAnimation* Animation, FWidgetAnimationDynamicEvent Delegate);
 
 	UFUNCTION(BlueprintCallable, Category = Animation)
-	void UnbindAllFromAnimationStarted(UWidgetAnimation* Animation);
+	UMG_API void UnbindAllFromAnimationStarted(UWidgetAnimation* Animation);
 
 	/**
 	 * Bind an animation finished delegate.
@@ -835,7 +835,7 @@ public:
 	 * @param Delegate the delegate to call when the animation's state changes
 	 */
 	UFUNCTION(BlueprintCallable, Category = Animation)
-	void BindToAnimationFinished(UWidgetAnimation* Animation, FWidgetAnimationDynamicEvent Delegate);
+	UMG_API void BindToAnimationFinished(UWidgetAnimation* Animation, FWidgetAnimationDynamicEvent Delegate);
 
 	/**
 	 * Unbind an animation finished delegate.
@@ -843,10 +843,10 @@ public:
 	 * @param Delegate the delegate to call when the animation's state changes
 	 */
 	UFUNCTION(BlueprintCallable, Category = Animation)
-	void UnbindFromAnimationFinished(UWidgetAnimation* Animation, FWidgetAnimationDynamicEvent Delegate);
+	UMG_API void UnbindFromAnimationFinished(UWidgetAnimation* Animation, FWidgetAnimationDynamicEvent Delegate);
 
 	UFUNCTION(BlueprintCallable, Category = Animation)
-	void UnbindAllFromAnimationFinished(UWidgetAnimation* Animation);
+	UMG_API void UnbindAllFromAnimationFinished(UWidgetAnimation* Animation);
 
 	/**
 	 * Allows binding to a specific animation's event.
@@ -856,7 +856,7 @@ public:
 	 * @param UserTag Scopes the delegate to only be called when the animation completes with a specific tag set on it when it was played.
 	 */
 	UFUNCTION(BlueprintCallable, Category = Animation)
-	void BindToAnimationEvent(UWidgetAnimation* Animation, FWidgetAnimationDynamicEvent Delegate, EWidgetAnimationEvent AnimationEvent, FName UserTag = NAME_None);
+	UMG_API void BindToAnimationEvent(UWidgetAnimation* Animation, FWidgetAnimationDynamicEvent Delegate, EWidgetAnimationEvent AnimationEvent, FName UserTag = NAME_None);
 
 	/** Is this widget an editor utility widget. */
 	virtual bool IsEditorUtility() const { return false; }
@@ -869,9 +869,9 @@ protected:
 	 * @param Animation the animation that started
 	 */
 	UFUNCTION( BlueprintNativeEvent, BlueprintCosmetic, Category = "Animation" )
-	void OnAnimationStarted( const UWidgetAnimation* Animation );
+	UMG_API void OnAnimationStarted( const UWidgetAnimation* Animation );
 
-	virtual void OnAnimationStarted_Implementation(const UWidgetAnimation* Animation);
+	UMG_API virtual void OnAnimationStarted_Implementation(const UWidgetAnimation* Animation);
 
 	/**
 	 * Called when an animation has either played all the way through or is stopped
@@ -879,20 +879,20 @@ protected:
 	 * @param Animation The animation that has finished playing
 	 */
 	UFUNCTION( BlueprintNativeEvent, BlueprintCosmetic, Category = "Animation" )
-	void OnAnimationFinished( const UWidgetAnimation* Animation );
+	UMG_API void OnAnimationFinished( const UWidgetAnimation* Animation );
 
-	virtual void OnAnimationFinished_Implementation(const UWidgetAnimation* Animation);
+	UMG_API virtual void OnAnimationFinished_Implementation(const UWidgetAnimation* Animation);
 
 	/** Broadcast any events based on a state transition for the sequence player, started, finished...etc. */
-	void BroadcastAnimationStateChange(const UUMGSequencePlayer& Player, EWidgetAnimationEvent AnimationEvent);
+	UMG_API void BroadcastAnimationStateChange(const UUMGSequencePlayer& Player, EWidgetAnimationEvent AnimationEvent);
 
 protected:
 
 	/** Called when a sequence player is finished playing an animation */
-	virtual void OnAnimationStartedPlaying(UUMGSequencePlayer& Player);
+	UMG_API virtual void OnAnimationStartedPlaying(UUMGSequencePlayer& Player);
 
 	/** Called when a sequence player is finished playing an animation */
-	virtual void OnAnimationFinishedPlaying(UUMGSequencePlayer& Player);
+	UMG_API virtual void OnAnimationFinishedPlaying(UUMGSequencePlayer& Player);
 
 public:
 	UE_DEPRECATED(5.2, "Direct access to ColorAndOpacity is deprecated. Please use the getter or setter.")
@@ -946,12 +946,12 @@ public:
 	 * @param InColorAndOpacity	The tint to apply to all child widgets.
 	 */
 	UFUNCTION(BlueprintCallable, BlueprintCosmetic, Category="Appearance")
-	void SetColorAndOpacity(FLinearColor InColorAndOpacity);
+	UMG_API void SetColorAndOpacity(FLinearColor InColorAndOpacity);
 
 	/**
 	 * Gets the tint of the widget.
 	 */
-	const FLinearColor& GetColorAndOpacity() const;
+	UMG_API const FLinearColor& GetColorAndOpacity() const;
 
 	/**
 	 * Sets the foreground color of the widget, this is inherited by sub widgets.  Any color property 
@@ -960,41 +960,41 @@ public:
 	 * @param InForegroundColor	The foreground color.
 	 */
 	UFUNCTION(BlueprintCallable, BlueprintCosmetic, Category="Appearance")
-	void SetForegroundColor(FSlateColor InForegroundColor);
+	UMG_API void SetForegroundColor(FSlateColor InForegroundColor);
 
 	/**
 	 * Gets the foreground color of the widget, this is inherited by sub widgets.  Any color property
 	 * that is marked as inherit uses this color.
 	 */
-	const FSlateColor& GetForegroundColor() const;
+	UMG_API const FSlateColor& GetForegroundColor() const;
 
 	/**
 	 * Sets the padding for the user widget, putting a larger gap between the widget border and it's root widget.
 	 */
 	UFUNCTION(BlueprintCallable, BlueprintCosmetic, Category="Appearance")
-	void SetPadding(FMargin InPadding);
+	UMG_API void SetPadding(FMargin InPadding);
 
 	/**
 	 * Gets the padding for the user widget.
 	 */
-	FMargin GetPadding() const;
+	UMG_API FMargin GetPadding() const;
 
 	/**
 	 * Gets the priority of the input action.
 	 */
-	int32 GetInputActionPriority() const;
+	UMG_API int32 GetInputActionPriority() const;
 
 	/**
 	 * Returns whether the input action is blocking.
 	 */
-	bool IsInputActionBlocking() const;
+	UMG_API bool IsInputActionBlocking() const;
 
 	/**
 	 * Sets whether this widget to accept focus when clicked, or when navigated to.
 	 */
-	bool IsFocusable() const;
+	UMG_API bool IsFocusable() const;
 
-	void SetIsFocusable(bool InIsFocusable);
+	UMG_API void SetIsFocusable(bool InIsFocusable);
 
 	/**
 	 * Plays an animation in this widget a specified number of times
@@ -1007,7 +1007,7 @@ public:
 	 * @param bRestoreState Restores widgets to their pre-animated state when the animation stops
 	 */
 	UFUNCTION(BlueprintCallable, BlueprintCosmetic, Category = "User Interface|Animation")
-	UUMGSequencePlayer* PlayAnimation(UWidgetAnimation* InAnimation, float StartAtTime = 0.0f, int32 NumLoopsToPlay = 1, EUMGSequencePlayMode::Type PlayMode = EUMGSequencePlayMode::Forward, float PlaybackSpeed = 1.0f, bool bRestoreState = false);
+	UMG_API UUMGSequencePlayer* PlayAnimation(UWidgetAnimation* InAnimation, float StartAtTime = 0.0f, int32 NumLoopsToPlay = 1, EUMGSequencePlayMode::Type PlayMode = EUMGSequencePlayMode::Forward, float PlaybackSpeed = 1.0f, bool bRestoreState = false);
 
 	/**
 	 * Plays an animation in this widget a specified number of times stopping at a specified time
@@ -1021,7 +1021,7 @@ public:
 	 * @param bRestoreState Restores widgets to their pre-animated state when the animation stops
 	 */
 	UFUNCTION(BlueprintCallable, BlueprintCosmetic, Category = "User Interface|Animation")
-	UUMGSequencePlayer* PlayAnimationTimeRange(UWidgetAnimation* InAnimation, float StartAtTime = 0.0f, float EndAtTime = 0.0f, int32 NumLoopsToPlay = 1, EUMGSequencePlayMode::Type PlayMode = EUMGSequencePlayMode::Forward, float PlaybackSpeed = 1.0f, bool bRestoreState = false);
+	UMG_API UUMGSequencePlayer* PlayAnimationTimeRange(UWidgetAnimation* InAnimation, float StartAtTime = 0.0f, float EndAtTime = 0.0f, int32 NumLoopsToPlay = 1, EUMGSequencePlayMode::Type PlayMode = EUMGSequencePlayMode::Forward, float PlaybackSpeed = 1.0f, bool bRestoreState = false);
 
 	/**
 	 * Plays an animation on this widget relative to it's current state forward.  You should use this version in situations where
@@ -1034,7 +1034,7 @@ public:
 	 * @param bRestoreState Restores widgets to their pre-animated state when the animation stops
 	 */
 	UFUNCTION(BlueprintCallable, BlueprintCosmetic, Category = "User Interface|Animation")
-	UUMGSequencePlayer* PlayAnimationForward(UWidgetAnimation* InAnimation, float PlaybackSpeed = 1.0f, bool bRestoreState = false);
+	UMG_API UUMGSequencePlayer* PlayAnimationForward(UWidgetAnimation* InAnimation, float PlaybackSpeed = 1.0f, bool bRestoreState = false);
 
 	/**
 	 * Plays an animation on this widget relative to it's current state in reverse.  You should use this version in situations where
@@ -1047,7 +1047,7 @@ public:
 	 * @param bRestoreState Restores widgets to their pre-animated state when the animation stops
 	 */
 	UFUNCTION(BlueprintCallable, BlueprintCosmetic, Category = "User Interface|Animation")
-	UUMGSequencePlayer* PlayAnimationReverse(UWidgetAnimation* InAnimation, float PlaybackSpeed = 1.0f, bool bRestoreState = false);
+	UMG_API UUMGSequencePlayer* PlayAnimationReverse(UWidgetAnimation* InAnimation, float PlaybackSpeed = 1.0f, bool bRestoreState = false);
 
 	/**
 	 * Stops an already running animation in this widget
@@ -1055,7 +1055,7 @@ public:
 	 * @param The name of the animation to stop
 	 */
 	UFUNCTION(BlueprintCallable, Category="User Interface|Animation")
-	void StopAnimation(const UWidgetAnimation* InAnimation);
+	UMG_API void StopAnimation(const UWidgetAnimation* InAnimation);
 
 	/**
 	 * Stop All actively running animations.
@@ -1063,7 +1063,7 @@ public:
 	 * @param The name of the animation to stop
 	 */
 	UFUNCTION(BlueprintCallable, Category="User Interface|Animation")
-	void StopAllAnimations();
+	UMG_API void StopAllAnimations();
 
 	/**
 	 * Pauses an already running animation in this widget
@@ -1072,7 +1072,7 @@ public:
 	 * @return the time point the animation was at when it was paused, relative to its start position.  Use this as the StartAtTime when you trigger PlayAnimation.
 	 */
 	UFUNCTION(BlueprintCallable, BlueprintCosmetic, Category="User Interface|Animation")
-	float PauseAnimation(const UWidgetAnimation* InAnimation);
+	UMG_API float PauseAnimation(const UWidgetAnimation* InAnimation);
 
 	/**
 	 * Gets the current time of the animation in this widget
@@ -1081,7 +1081,7 @@ public:
 	 * @return the current time of the animation.
 	 */
 	UFUNCTION(BlueprintCallable, BlueprintCosmetic, Category = "User Interface|Animation")
-	float GetAnimationCurrentTime(const UWidgetAnimation* InAnimation) const;
+	UMG_API float GetAnimationCurrentTime(const UWidgetAnimation* InAnimation) const;
 
 	/**
 	 * Sets the current time of the animation in this widget. Does not change state.
@@ -1090,7 +1090,7 @@ public:
 	 * @param The current time of the animation.
 	 */
 	UFUNCTION(BlueprintCallable, BlueprintCosmetic, Category = "User Interface|Animation")
-	void SetAnimationCurrentTime(const UWidgetAnimation* InAnimation, float InTime);
+	UMG_API void SetAnimationCurrentTime(const UWidgetAnimation* InAnimation, float InTime);
 
 	/**
 	 * Gets whether an animation is currently playing on this widget.
@@ -1099,13 +1099,13 @@ public:
 	 * @return True if the animation is currently playing
 	 */
 	UFUNCTION(BlueprintCallable, BlueprintCosmetic, Category="User Interface|Animation")
-	bool IsAnimationPlaying(const UWidgetAnimation* InAnimation) const;
+	UMG_API bool IsAnimationPlaying(const UWidgetAnimation* InAnimation) const;
 
 	/**
 	 * @return True if any animation is currently playing
 	 */
 	UFUNCTION(BlueprintCallable, BlueprintCosmetic, Category="User Interface|Animation")
-	bool IsAnyAnimationPlaying() const;
+	UMG_API bool IsAnyAnimationPlaying() const;
 
 	/**
 	* Changes the number of loops to play given a playing animation
@@ -1114,7 +1114,7 @@ public:
 	* @param NumLoopsToPlay The number of loops to play. (0 to loop indefinitely)
 	*/
 	UFUNCTION(BlueprintCallable, BlueprintCosmetic, Category = "User Interface|Animation")
-	void SetNumLoopsToPlay(const UWidgetAnimation* InAnimation, int32 NumLoopsToPlay);
+	UMG_API void SetNumLoopsToPlay(const UWidgetAnimation* InAnimation, int32 NumLoopsToPlay);
 
 	/**
 	* Changes the playback rate of a playing animation
@@ -1123,7 +1123,7 @@ public:
 	* @param PlaybackRate Playback rate multiplier (1 is default)
 	*/
 	UFUNCTION(BlueprintCallable, BlueprintCosmetic, Category = "User Interface|Animation")
-	void SetPlaybackSpeed(const UWidgetAnimation* InAnimation, float PlaybackSpeed = 1.0f);
+	UMG_API void SetPlaybackSpeed(const UWidgetAnimation* InAnimation, float PlaybackSpeed = 1.0f);
 
 	/**
 	* If an animation is playing, this function will reverse the playback.
@@ -1131,7 +1131,7 @@ public:
 	* @param InAnimation The playing animation that we want to reverse
 	*/
 	UFUNCTION(BlueprintCallable, BlueprintCosmetic, Category = "User Interface|Animation")
-	void ReverseAnimation(const UWidgetAnimation* InAnimation);
+	UMG_API void ReverseAnimation(const UWidgetAnimation* InAnimation);
 
 	/**
 	 * returns true if the animation is currently playing forward, false otherwise.
@@ -1139,13 +1139,13 @@ public:
 	 * @param InAnimation The playing animation that we want to know about
 	 */
 	UFUNCTION(BlueprintCallable, BlueprintCosmetic, Category = "User Interface|Animation")
-	bool IsAnimationPlayingForward(const UWidgetAnimation* InAnimation);
+	UMG_API bool IsAnimationPlayingForward(const UWidgetAnimation* InAnimation);
 
 	/**
 	 * Flushes all animations on all widgets to guarantee that any queued updates are processed before this call returns
 	 */
 	UFUNCTION(BlueprintCallable, BlueprintCosmetic, Category = "User Interface|Animation")
-	void FlushAnimations();
+	UMG_API void FlushAnimations();
 
 	/** Find the first extension of the requested type. */
 	template<typename ExtensionType>
@@ -1156,11 +1156,11 @@ public:
 
 	/** Find the first extension of the requested type. */
 	UFUNCTION(BlueprintCallable, Category = "User Interface|Extension", Meta = (DeterminesOutputType = "ExtensionType"))
-	UUserWidgetExtension* GetExtension(TSubclassOf<UUserWidgetExtension> ExtensionType) const;
+	UMG_API UUserWidgetExtension* GetExtension(TSubclassOf<UUserWidgetExtension> ExtensionType) const;
 
 	/** Find the extensions of the requested type. */
 	UFUNCTION(BlueprintCallable, Category = "User Interface|Extension", Meta = (DeterminesOutputType = "ExtensionType"))
-	TArray<UUserWidgetExtension*> GetExtensions(TSubclassOf<UUserWidgetExtension> ExtensionType) const;
+	UMG_API TArray<UUserWidgetExtension*> GetExtensions(TSubclassOf<UUserWidgetExtension> ExtensionType) const;
 
 	/** Add the extension of the requested type. */
 	template<typename ExtensionType>
@@ -1171,11 +1171,11 @@ public:
 
 	/** Add the extension of the requested type. */
 	UFUNCTION(BlueprintCallable, Category = "User Interface|Extension", Meta = (DeterminesOutputType = "InExtensionType"))
-	UUserWidgetExtension* AddExtension(TSubclassOf<UUserWidgetExtension> InExtensionType);
+	UMG_API UUserWidgetExtension* AddExtension(TSubclassOf<UUserWidgetExtension> InExtensionType);
 
 	/** Remove the extension. */
 	UFUNCTION(BlueprintCallable, Category = "User Interface|Extension")
-	void RemoveExtension(UUserWidgetExtension* InExtension);
+	UMG_API void RemoveExtension(UUserWidgetExtension* InExtension);
 
 	/** Remove all extensions of the requested type. */
 	template<typename ExtensionType>
@@ -1186,7 +1186,7 @@ public:
 
 	/** Remove all extensions of the requested type. */
 	UFUNCTION(BlueprintCallable, Category = "User Interface|Extension")
-	void RemoveExtensions(TSubclassOf<UUserWidgetExtension> InExtensionType);
+	UMG_API void RemoveExtensions(TSubclassOf<UUserWidgetExtension> InExtensionType);
 
 	/**
 	 * Plays a sound through the UI
@@ -1194,7 +1194,7 @@ public:
 	 * @param The sound to play
 	 */
 	UFUNCTION(BlueprintCallable, BlueprintCosmetic, Category="Sound", meta=( DeprecatedFunction, DeprecationMessage="Use the UGameplayStatics::PlaySound2D instead." ))
-	void PlaySound(class USoundBase* SoundToPlay);
+	UMG_API void PlaySound(class USoundBase* SoundToPlay);
 
 	/** 
 	 * Sets the child Widget that should receive focus when this UserWidget gets focus using it's name. 
@@ -1202,7 +1202,7 @@ public:
 	 * @param WidgetName Name of the Widget to forward the focus to when this widget receives focus.
 	 * @return True if the Widget is set properly. Will return false if we can't find a child widget with the specified name.
 	 */
-	bool SetDesiredFocusWidget(FName WidgetName);
+	UMG_API bool SetDesiredFocusWidget(FName WidgetName);
 
 	/** 
 	 * Sets the child Widget that should receive focus when this UserWidget gets focus. 
@@ -1210,29 +1210,29 @@ public:
 	 * @param Widget Widget to forward the focus to when this widget receives focus
 	 * @return True if the Widget is set properly. Will return false if it's not a child of this UserWidget.
 	 */
-	bool SetDesiredFocusWidget(UWidget* Widget);
+	UMG_API bool SetDesiredFocusWidget(UWidget* Widget);
 
 	/** @returns The Name of the Widget that should receive focus when this UserWidget gets focus. */
-	FName GetDesiredFocusWidgetName() const;
+	UMG_API FName GetDesiredFocusWidgetName() const;
 
 	/** @returns The Widget that should receive focus when this UserWidget gets focus. */
-	UWidget* GetDesiredFocusWidget() const;
+	UMG_API UWidget* GetDesiredFocusWidget() const;
 
 	/** @returns The UObject wrapper for a given SWidget */
-	UWidget* GetWidgetHandle(TSharedRef<SWidget> InWidget);
+	UMG_API UWidget* GetWidgetHandle(TSharedRef<SWidget> InWidget);
 
 	/** @returns The root UObject widget wrapper */
-	UWidget* GetRootWidget() const;
+	UMG_API UWidget* GetRootWidget() const;
 
 	/** @returns The slate widget corresponding to a given name */
-	TSharedPtr<SWidget> GetSlateWidgetFromName(const FName& Name) const;
+	UMG_API TSharedPtr<SWidget> GetSlateWidgetFromName(const FName& Name) const;
 
 	/** @returns The uobject widget corresponding to a given name */
-	UWidget* GetWidgetFromName(const FName& Name) const;
+	UMG_API UWidget* GetWidgetFromName(const FName& Name) const;
 
 	//~ Begin UObject Interface
-	virtual bool IsAsset() const;
-	virtual void PreSave(FObjectPreSaveContext ObjectSaveContext) override;
+	UMG_API virtual bool IsAsset() const;
+	UMG_API virtual void PreSave(FObjectPreSaveContext ObjectSaveContext) override;
 	//~ End UObject Interface
 
 	/** Are we currently playing any animations? */
@@ -1241,25 +1241,25 @@ public:
 
 #if WITH_EDITOR
 	//~ Begin UWidget Interface
-	virtual const FText GetPaletteCategory() override;
+	UMG_API virtual const FText GetPaletteCategory() override;
 	//~ End UWidget Interface
 
-	virtual void SetDesignerFlags(EWidgetDesignFlags NewFlags) override;
-	virtual void OnDesignerChanged(const FDesignerChangedEventArgs& EventArgs) override;
-	virtual void PostEditChangeProperty(FPropertyChangedEvent& PropertyChangedEvent) override;
+	UMG_API virtual void SetDesignerFlags(EWidgetDesignFlags NewFlags) override;
+	UMG_API virtual void OnDesignerChanged(const FDesignerChangedEventArgs& EventArgs) override;
+	UMG_API virtual void PostEditChangeProperty(FPropertyChangedEvent& PropertyChangedEvent) override;
 
 	/** Update the binding for this namedslot if the name is not found but GUID is matched. */
-	void UpdateBindingForSlot(FName SlotName);
+	UMG_API void UpdateBindingForSlot(FName SlotName);
 
 	/** Add the GUID of each Namedslot widget to its corresponding binding, if any. */
-	void AssignGUIDToBindings();
+	UMG_API void AssignGUIDToBindings();
 
 	/**
 	 * Final step of Widget Blueprint compilation. Allows widgets to perform custom validation and trigger compiler outputs as needed.
 	 * @see ValidateCompiledDefaults
 	 * @see ValidateCompiledWidgetTree
 	 */
-	void ValidateBlueprint(const UWidgetTree& BlueprintWidgetTree, class IWidgetCompilerLog& CompileLog) const;
+	UMG_API void ValidateBlueprint(const UWidgetTree& BlueprintWidgetTree, class IWidgetCompilerLog& CompileLog) const;
 
 	/**
 	 * Override to perform any custom inspections of the default widget tree at the end of compilation.
@@ -1272,16 +1272,16 @@ public:
 	virtual void ValidateCompiledWidgetTree(const UWidgetTree& BlueprintWidgetTree, class IWidgetCompilerLog& CompileLog) const {};
 #endif
 
-	static UUserWidget* CreateWidgetInstance(UWidget& OwningWidget, TSubclassOf<UUserWidget> UserWidgetClass, FName WidgetName);
-	static UUserWidget* CreateWidgetInstance(UWidgetTree& OwningWidgetTree, TSubclassOf<UUserWidget> UserWidgetClass, FName WidgetName);
-	static UUserWidget* CreateWidgetInstance(APlayerController& OwnerPC, TSubclassOf<UUserWidget> UserWidgetClass, FName WidgetName);
-	static UUserWidget* CreateWidgetInstance(UGameInstance& GameInstance, TSubclassOf<UUserWidget> UserWidgetClass, FName WidgetName);
-	static UUserWidget* CreateWidgetInstance(UWorld& World, TSubclassOf<UUserWidget> UserWidgetClass, FName WidgetName);
+	static UMG_API UUserWidget* CreateWidgetInstance(UWidget& OwningWidget, TSubclassOf<UUserWidget> UserWidgetClass, FName WidgetName);
+	static UMG_API UUserWidget* CreateWidgetInstance(UWidgetTree& OwningWidgetTree, TSubclassOf<UUserWidget> UserWidgetClass, FName WidgetName);
+	static UMG_API UUserWidget* CreateWidgetInstance(APlayerController& OwnerPC, TSubclassOf<UUserWidget> UserWidgetClass, FName WidgetName);
+	static UMG_API UUserWidget* CreateWidgetInstance(UGameInstance& GameInstance, TSubclassOf<UUserWidget> UserWidgetClass, FName WidgetName);
+	static UMG_API UUserWidget* CreateWidgetInstance(UWorld& World, TSubclassOf<UUserWidget> UserWidgetClass, FName WidgetName);
 
 private:
-	static UUserWidget* CreateInstanceInternal(UObject* Outer, TSubclassOf<UUserWidget> UserWidgetClass, FName WidgetName, UWorld* World, ULocalPlayer* LocalPlayer);
+	static UMG_API UUserWidget* CreateInstanceInternal(UObject* Outer, TSubclassOf<UUserWidget> UserWidgetClass, FName WidgetName, UWorld* World, ULocalPlayer* LocalPlayer);
 
-	void ClearStoppedSequencePlayers();
+	UMG_API void ClearStoppedSequencePlayers();
 
 public:
 
@@ -1354,105 +1354,105 @@ private:
 	uint8 bStoppingAllAnimations : 1;
 
 protected:
-	virtual TSharedRef<SWidget> RebuildWidget() override;
-	virtual void OnWidgetRebuilt() override;
+	UMG_API virtual TSharedRef<SWidget> RebuildWidget() override;
+	UMG_API virtual void OnWidgetRebuilt() override;
 
 	UE_DEPRECATED(5.1, "GetFullScreenOffset is deprecated. Use the GameViewportSubsystem.")
-	FMargin GetFullScreenOffset() const;
+	UMG_API FMargin GetFullScreenOffset() const;
 
-	virtual void NativeOnInitialized();
-	virtual void NativePreConstruct();
-	virtual void NativeConstruct();
-	virtual void NativeDestruct();
-	virtual void NativeTick(const FGeometry& MyGeometry, float InDeltaTime);
+	UMG_API virtual void NativeOnInitialized();
+	UMG_API virtual void NativePreConstruct();
+	UMG_API virtual void NativeConstruct();
+	UMG_API virtual void NativeDestruct();
+	UMG_API virtual void NativeTick(const FGeometry& MyGeometry, float InDeltaTime);
 
 	/**
 	 * Native implemented paint function for the Widget
 	 * Returns the maximum LayerID painted on
 	 */
-	virtual int32 NativePaint(const FPaintArgs& Args, const FGeometry& AllottedGeometry, const FSlateRect& MyCullingRect, FSlateWindowElementList& OutDrawElements, int32 LayerId, const FWidgetStyle& InWidgetStyle, bool bParentEnabled ) const;
+	UMG_API virtual int32 NativePaint(const FPaintArgs& Args, const FGeometry& AllottedGeometry, const FSlateRect& MyCullingRect, FSlateWindowElementList& OutDrawElements, int32 LayerId, const FWidgetStyle& InWidgetStyle, bool bParentEnabled ) const;
 
 	FORCEINLINE FVector2D GetMinimumDesiredSize() const { return MinimumDesiredSize; }
-	void SetMinimumDesiredSize(FVector2D InMinimumDesiredSize);
+	UMG_API void SetMinimumDesiredSize(FVector2D InMinimumDesiredSize);
 
-	virtual bool NativeIsInteractable() const;
-	virtual bool NativeSupportsKeyboardFocus() const;
+	UMG_API virtual bool NativeIsInteractable() const;
+	UMG_API virtual bool NativeSupportsKeyboardFocus() const;
 	virtual bool NativeSupportsCustomNavigation() const { return false; }
 
-	virtual FReply NativeOnFocusReceived( const FGeometry& InGeometry, const FFocusEvent& InFocusEvent );
-	virtual void NativeOnFocusLost( const FFocusEvent& InFocusEvent );
-	virtual void NativeOnFocusChanging(const FWeakWidgetPath& PreviousFocusPath, const FWidgetPath& NewWidgetPath, const FFocusEvent& InFocusEvent);
-	virtual void NativeOnAddedToFocusPath(const FFocusEvent& InFocusEvent);
-	virtual void NativeOnRemovedFromFocusPath(const FFocusEvent& InFocusEvent);
-	virtual FNavigationReply NativeOnNavigation(const FGeometry& MyGeometry, const FNavigationEvent& InNavigationEvent, const FNavigationReply& InDefaultReply);
-	virtual FReply NativeOnKeyChar( const FGeometry& InGeometry, const FCharacterEvent& InCharEvent );
-	virtual FReply NativeOnPreviewKeyDown( const FGeometry& InGeometry, const FKeyEvent& InKeyEvent );
-	virtual FReply NativeOnKeyDown( const FGeometry& InGeometry, const FKeyEvent& InKeyEvent );
-	virtual FReply NativeOnKeyUp( const FGeometry& InGeometry, const FKeyEvent& InKeyEvent );
-	virtual FReply NativeOnAnalogValueChanged( const FGeometry& InGeometry, const FAnalogInputEvent& InAnalogEvent );
-	virtual FReply NativeOnMouseButtonDown( const FGeometry& InGeometry, const FPointerEvent& InMouseEvent );
-	virtual FReply NativeOnPreviewMouseButtonDown( const FGeometry& InGeometry, const FPointerEvent& InMouseEvent );
-	virtual FReply NativeOnMouseButtonUp( const FGeometry& InGeometry, const FPointerEvent& InMouseEvent );
-	virtual FReply NativeOnMouseMove( const FGeometry& InGeometry, const FPointerEvent& InMouseEvent );
-	virtual void NativeOnMouseEnter( const FGeometry& InGeometry, const FPointerEvent& InMouseEvent );
-	virtual void NativeOnMouseLeave( const FPointerEvent& InMouseEvent );
-	virtual FReply NativeOnMouseWheel( const FGeometry& InGeometry, const FPointerEvent& InMouseEvent );
-	virtual FReply NativeOnMouseButtonDoubleClick( const FGeometry& InGeometry, const FPointerEvent& InMouseEvent );
-	virtual void NativeOnDragDetected( const FGeometry& InGeometry, const FPointerEvent& InMouseEvent, UDragDropOperation*& OutOperation );
-	virtual void NativeOnDragEnter( const FGeometry& InGeometry, const FDragDropEvent& InDragDropEvent, UDragDropOperation* InOperation );
-	virtual void NativeOnDragLeave( const FDragDropEvent& InDragDropEvent, UDragDropOperation* InOperation );
-	virtual bool NativeOnDragOver( const FGeometry& InGeometry, const FDragDropEvent& InDragDropEvent, UDragDropOperation* InOperation );
-	virtual bool NativeOnDrop( const FGeometry& InGeometry, const FDragDropEvent& InDragDropEvent, UDragDropOperation* InOperation );
-	virtual void NativeOnDragCancelled( const FDragDropEvent& InDragDropEvent, UDragDropOperation* InOperation );
-	virtual FReply NativeOnTouchGesture( const FGeometry& InGeometry, const FPointerEvent& InGestureEvent );
-	virtual FReply NativeOnTouchStarted( const FGeometry& InGeometry, const FPointerEvent& InGestureEvent );
-	virtual FReply NativeOnTouchMoved( const FGeometry& InGeometry, const FPointerEvent& InGestureEvent );
-	virtual FReply NativeOnTouchEnded( const FGeometry& InGeometry, const FPointerEvent& InGestureEvent );
-	virtual FReply NativeOnMotionDetected( const FGeometry& InGeometry, const FMotionEvent& InMotionEvent );
-	virtual FReply NativeOnTouchForceChanged(const FGeometry& MyGeometry, const FPointerEvent& TouchEvent);
-	virtual FCursorReply NativeOnCursorQuery( const FGeometry& InGeometry, const FPointerEvent& InCursorEvent );
-	virtual FNavigationReply NativeOnNavigation(const FGeometry& InGeometry, const FNavigationEvent& InNavigationEvent);
-	virtual void NativeOnMouseCaptureLost(const FCaptureLostEvent& CaptureLostEvent);
+	UMG_API virtual FReply NativeOnFocusReceived( const FGeometry& InGeometry, const FFocusEvent& InFocusEvent );
+	UMG_API virtual void NativeOnFocusLost( const FFocusEvent& InFocusEvent );
+	UMG_API virtual void NativeOnFocusChanging(const FWeakWidgetPath& PreviousFocusPath, const FWidgetPath& NewWidgetPath, const FFocusEvent& InFocusEvent);
+	UMG_API virtual void NativeOnAddedToFocusPath(const FFocusEvent& InFocusEvent);
+	UMG_API virtual void NativeOnRemovedFromFocusPath(const FFocusEvent& InFocusEvent);
+	UMG_API virtual FNavigationReply NativeOnNavigation(const FGeometry& MyGeometry, const FNavigationEvent& InNavigationEvent, const FNavigationReply& InDefaultReply);
+	UMG_API virtual FReply NativeOnKeyChar( const FGeometry& InGeometry, const FCharacterEvent& InCharEvent );
+	UMG_API virtual FReply NativeOnPreviewKeyDown( const FGeometry& InGeometry, const FKeyEvent& InKeyEvent );
+	UMG_API virtual FReply NativeOnKeyDown( const FGeometry& InGeometry, const FKeyEvent& InKeyEvent );
+	UMG_API virtual FReply NativeOnKeyUp( const FGeometry& InGeometry, const FKeyEvent& InKeyEvent );
+	UMG_API virtual FReply NativeOnAnalogValueChanged( const FGeometry& InGeometry, const FAnalogInputEvent& InAnalogEvent );
+	UMG_API virtual FReply NativeOnMouseButtonDown( const FGeometry& InGeometry, const FPointerEvent& InMouseEvent );
+	UMG_API virtual FReply NativeOnPreviewMouseButtonDown( const FGeometry& InGeometry, const FPointerEvent& InMouseEvent );
+	UMG_API virtual FReply NativeOnMouseButtonUp( const FGeometry& InGeometry, const FPointerEvent& InMouseEvent );
+	UMG_API virtual FReply NativeOnMouseMove( const FGeometry& InGeometry, const FPointerEvent& InMouseEvent );
+	UMG_API virtual void NativeOnMouseEnter( const FGeometry& InGeometry, const FPointerEvent& InMouseEvent );
+	UMG_API virtual void NativeOnMouseLeave( const FPointerEvent& InMouseEvent );
+	UMG_API virtual FReply NativeOnMouseWheel( const FGeometry& InGeometry, const FPointerEvent& InMouseEvent );
+	UMG_API virtual FReply NativeOnMouseButtonDoubleClick( const FGeometry& InGeometry, const FPointerEvent& InMouseEvent );
+	UMG_API virtual void NativeOnDragDetected( const FGeometry& InGeometry, const FPointerEvent& InMouseEvent, UDragDropOperation*& OutOperation );
+	UMG_API virtual void NativeOnDragEnter( const FGeometry& InGeometry, const FDragDropEvent& InDragDropEvent, UDragDropOperation* InOperation );
+	UMG_API virtual void NativeOnDragLeave( const FDragDropEvent& InDragDropEvent, UDragDropOperation* InOperation );
+	UMG_API virtual bool NativeOnDragOver( const FGeometry& InGeometry, const FDragDropEvent& InDragDropEvent, UDragDropOperation* InOperation );
+	UMG_API virtual bool NativeOnDrop( const FGeometry& InGeometry, const FDragDropEvent& InDragDropEvent, UDragDropOperation* InOperation );
+	UMG_API virtual void NativeOnDragCancelled( const FDragDropEvent& InDragDropEvent, UDragDropOperation* InOperation );
+	UMG_API virtual FReply NativeOnTouchGesture( const FGeometry& InGeometry, const FPointerEvent& InGestureEvent );
+	UMG_API virtual FReply NativeOnTouchStarted( const FGeometry& InGeometry, const FPointerEvent& InGestureEvent );
+	UMG_API virtual FReply NativeOnTouchMoved( const FGeometry& InGeometry, const FPointerEvent& InGestureEvent );
+	UMG_API virtual FReply NativeOnTouchEnded( const FGeometry& InGeometry, const FPointerEvent& InGestureEvent );
+	UMG_API virtual FReply NativeOnMotionDetected( const FGeometry& InGeometry, const FMotionEvent& InMotionEvent );
+	UMG_API virtual FReply NativeOnTouchForceChanged(const FGeometry& MyGeometry, const FPointerEvent& TouchEvent);
+	UMG_API virtual FCursorReply NativeOnCursorQuery( const FGeometry& InGeometry, const FPointerEvent& InCursorEvent );
+	UMG_API virtual FNavigationReply NativeOnNavigation(const FGeometry& InGeometry, const FNavigationEvent& InNavigationEvent);
+	UMG_API virtual void NativeOnMouseCaptureLost(const FCaptureLostEvent& CaptureLostEvent);
 
 protected:
 
 	/**
 	 * Ticks the active sequences and latent actions that have been scheduled for this Widget.
 	 */
-	void TickActionsAndAnimation(float InDeltaTime);
-	void PostTickActionsAndAnimation(float InDeltaTime);
+	UMG_API void TickActionsAndAnimation(float InDeltaTime);
+	UMG_API void PostTickActionsAndAnimation(float InDeltaTime);
 
-	void RemoveObsoleteBindings(const TArray<FName>& NamedSlots);
+	UMG_API void RemoveObsoleteBindings(const TArray<FName>& NamedSlots);
 
-	UUMGSequencePlayer* GetSequencePlayer(const UWidgetAnimation* InAnimation) const;
-	UUMGSequencePlayer* GetOrAddSequencePlayer(UWidgetAnimation* InAnimation);
+	UMG_API UUMGSequencePlayer* GetSequencePlayer(const UWidgetAnimation* InAnimation) const;
+	UMG_API UUMGSequencePlayer* GetOrAddSequencePlayer(UWidgetAnimation* InAnimation);
 
-	void ConditionalTearDownAnimations();
+	UMG_API void ConditionalTearDownAnimations();
 
-	void TearDownAnimations();
+	UMG_API void TearDownAnimations();
 
-	void DisableAnimations();
+	UMG_API void DisableAnimations();
 
-	void Invalidate(EInvalidateWidgetReason InvalidateReason);
+	UMG_API void Invalidate(EInvalidateWidgetReason InvalidateReason);
 	
 	/**
 	 * Listens for a particular Player Input Action by name.  This requires that those actions are being executed, and
 	 * that we're not currently in UI-Only Input Mode.
 	 */
 	UFUNCTION( BlueprintCallable, Category = "Input", meta = ( BlueprintProtected = "true" ) )
-	void ListenForInputAction( FName ActionName, TEnumAsByte< EInputEvent > EventType, bool bConsume, FOnInputAction Callback );
+	UMG_API void ListenForInputAction( FName ActionName, TEnumAsByte< EInputEvent > EventType, bool bConsume, FOnInputAction Callback );
 
 	/**
 	 * Removes the binding for a particular action's callback.
 	 */
 	UFUNCTION( BlueprintCallable, Category = "Input", meta = ( BlueprintProtected = "true" ) )
-	void StopListeningForInputAction( FName ActionName, TEnumAsByte< EInputEvent > EventType );
+	UMG_API void StopListeningForInputAction( FName ActionName, TEnumAsByte< EInputEvent > EventType );
 
 	/**
 	 * Stops listening to all input actions, and unregisters the input component with the player controller.
 	 */
 	UFUNCTION( BlueprintCallable, Category = "Input", meta = ( BlueprintProtected = "true" ) )
-	void StopListeningForAllInputActions();
+	UMG_API void StopListeningForAllInputActions();
 
 	/**
 	 * ListenForInputAction will automatically Register an Input Component with the player input system.
@@ -1460,7 +1460,7 @@ protected:
 	 * UnregisterInputComponent to pause, and RegisterInputComponent to resume listening.
 	 */
 	UFUNCTION(BlueprintCallable, Category = "Input", meta = ( BlueprintProtected = "true" ))
-	void RegisterInputComponent();
+	UMG_API void RegisterInputComponent();
 
 	/**
 	 * StopListeningForAllInputActions will automatically Register an Input Component with the player input system.
@@ -1468,23 +1468,23 @@ protected:
 	 * UnregisterInputComponent to pause, and RegisterInputComponent to resume listening.
 	 */
 	UFUNCTION(BlueprintCallable, Category = "Input", meta = ( BlueprintProtected = "true" ))
-	void UnregisterInputComponent();
+	UMG_API void UnregisterInputComponent();
 
 	/**
 	 * Checks if the action has a registered callback with the input component.
 	 */
 	UFUNCTION( BlueprintCallable, Category = "Input", meta = ( BlueprintProtected = "true" ) )
-	bool IsListeningForInputAction( FName ActionName ) const;
+	UMG_API bool IsListeningForInputAction( FName ActionName ) const;
 
 	UFUNCTION( BlueprintCallable, Category = "Input", meta = ( BlueprintProtected = "true" ) )
-	void SetInputActionPriority( int32 NewPriority );
+	UMG_API void SetInputActionPriority( int32 NewPriority );
 
 	UFUNCTION( BlueprintCallable, Category = "Input", meta = ( BlueprintProtected = "true" ) )
-	void SetInputActionBlocking( bool bShouldBlock );
+	UMG_API void SetInputActionBlocking( bool bShouldBlock );
 
-	void OnInputAction( FOnInputAction Callback );
+	UMG_API void OnInputAction( FOnInputAction Callback );
 
-	virtual void InitializeInputComponent();
+	UMG_API virtual void InitializeInputComponent();
 
 private:
 	FVector2D MinimumDesiredSize;
@@ -1509,7 +1509,7 @@ protected:
 	TArray<FAnimationEventBinding> AnimationCallbacks;
 
 private:
-	static void OnLatentActionsChanged(UObject* ObjectWhichChanged, ELatentActionChangeType ChangeType);
+	static UMG_API void OnLatentActionsChanged(UObject* ObjectWhichChanged, ELatentActionChangeType ChangeType);
 
 	/** The player context that is associated with this UI.  Think of this as the owner of the UI. */
 	FLocalPlayerContext PlayerContext;
@@ -1517,8 +1517,8 @@ private:
 	/** Get World calls can be expensive for Widgets, we speed them up by caching the last found world until it goes away. */
 	mutable TWeakObjectPtr<UWorld> CachedWorld;
 
-	static bool bTemplateInitializing;
-	static uint32 bInitializingFromWidgetTree;
+	static UMG_API bool bTemplateInitializing;
+	static UMG_API uint32 bInitializingFromWidgetTree;
 
 protected:
 
diff --git a/Engine/Source/Runtime/UMG/Public/Blueprint/UserWidgetBlueprint.h b/Engine/Source/Runtime/UMG/Public/Blueprint/UserWidgetBlueprint.h
--- a/Engine/Source/Runtime/UMG/Public/Blueprint/UserWidgetBlueprint.h
+++ b/Engine/Source/Runtime/UMG/Public/Blueprint/UserWidgetBlueprint.h
@@ -5,8 +5,8 @@
 #include "Engine/Blueprint.h"
 #include "UserWidgetBlueprint.generated.h"
 
-UCLASS(Abstract)
-class UMG_API UUserWidgetBlueprint : public UBlueprint
+UCLASS(Abstract, MinimalAPI)
+class UUserWidgetBlueprint : public UBlueprint
 {
 	GENERATED_BODY()
 
diff --git a/Engine/Source/Runtime/UMG/Public/Blueprint/UserWidgetPool.h b/Engine/Source/Runtime/UMG/Public/Blueprint/UserWidgetPool.h
--- a/Engine/Source/Runtime/UMG/Public/Blueprint/UserWidgetPool.h
+++ b/Engine/Source/Runtime/UMG/Public/Blueprint/UserWidgetPool.h
@@ -23,30 +23,30 @@ class APlayerController;
  * @see UDynamicEntryBox
  */
 USTRUCT()
-struct UMG_API FUserWidgetPool
+struct FUserWidgetPool
 {
 	GENERATED_BODY();
 
 public:
 	FUserWidgetPool() = default;
-	FUserWidgetPool(UWidget& InOwningWidget);
-	~FUserWidgetPool();
+	UMG_API FUserWidgetPool(UWidget& InOwningWidget);
+	UMG_API ~FUserWidgetPool();
 
 	/** In the case that you don't have an owner widget, you should set a world to your pool, or it won't be able to construct widgets. */
-	void SetWorld(UWorld* OwningWorld);
+	UMG_API void SetWorld(UWorld* OwningWorld);
 
 	/**
 	 * In the case a pool has no owner widget, a default player controller can optionally be specified allowing greater control  
 	 * in split screen scenarios. The following priority rules are applied when determining the controller which should own created widgets:
 	 * Controller of owning widget (if set) > Default player controller (if set) > First local player controller from world.
 	 */
-	void SetDefaultPlayerController(APlayerController* InDefaultPlayerController);
+	UMG_API void SetDefaultPlayerController(APlayerController* InDefaultPlayerController);
 
 	/** Triggers RebuildWidget on all currently active UserWidget instances */
-	void RebuildWidgets();
+	UMG_API void RebuildWidgets();
 
 	/** Report any references to UObjects to the reference collector (only necessary if this is not already a UPROPERTY) */
-	void AddReferencedObjects(FReferenceCollector& Collector);
+	UMG_API void AddReferencedObjects(FReferenceCollector& Collector);
 
 	bool IsInitialized() const { return OwningWidget.IsValid() || OwningWorld.IsValid(); }
 	const TArray<UUserWidget*>& GetActiveWidgets() const { return ActiveWidgets; }
@@ -76,25 +76,25 @@ public:
 	}
 
 	/** Return a widget object to the pool, allowing it to be reused in the future */
-	void Release(UUserWidget* Widget, bool bReleaseSlate = false);
+	UMG_API void Release(UUserWidget* Widget, bool bReleaseSlate = false);
 
 	/** Return a widget object to the pool, allowing it to be reused in the future */
-	void Release(TArray<UUserWidget*> Widgets, bool bReleaseSlate = false);
+	UMG_API void Release(TArray<UUserWidget*> Widgets, bool bReleaseSlate = false);
 
 	/** Returns all active widget objects to the inactive pool and optionally destroys all cached underlying slate widgets. */
-	void ReleaseAll(bool bReleaseSlate = false);
+	UMG_API void ReleaseAll(bool bReleaseSlate = false);
 
 	/** Full reset of all created widget objects (and any cached underlying slate) */
-	void ResetPool();
+	UMG_API void ResetPool();
 
 	/** Reset of all cached underlying Slate widgets, only for inactive widgets in the pool. */
-	void ReleaseInactiveSlateResources();
+	UMG_API void ReleaseInactiveSlateResources();
 
 	UE_DEPRECATED(4.24, "Use ReleaseAllSlateResources instead.")
 	void ReleaseSlateResources() { ReleaseAllSlateResources(); }
 
 	/** Reset of all cached underlying Slate widgets, but not the active UUserWidget objects */
-	void ReleaseAllSlateResources();
+	UMG_API void ReleaseAllSlateResources();
 
 private:
 	template <typename UserWidgetT = UUserWidget>
@@ -161,4 +161,4 @@ private:
 	TWeakObjectPtr<UWorld> OwningWorld;
 	TWeakObjectPtr<APlayerController> DefaultPlayerController;
 	TMap<UUserWidget*, TSharedPtr<SWidget>> CachedSlateByWidgetObject;
-};
\ No newline at end of file
+};
diff --git a/Engine/Source/Runtime/UMG/Public/Blueprint/WidgetBlueprintGeneratedClass.h b/Engine/Source/Runtime/UMG/Public/Blueprint/WidgetBlueprintGeneratedClass.h
--- a/Engine/Source/Runtime/UMG/Public/Blueprint/WidgetBlueprintGeneratedClass.h
+++ b/Engine/Source/Runtime/UMG/Public/Blueprint/WidgetBlueprintGeneratedClass.h
@@ -52,14 +52,14 @@ struct FDelegateRuntimeBinding
 
 
 #if WITH_EDITOR
-class UMG_API FWidgetBlueprintGeneratedClassDelegates
+class FWidgetBlueprintGeneratedClassDelegates
 {
 public:
 	// delegate for generating widget asset registry tags.
 	DECLARE_MULTICAST_DELEGATE_TwoParams(FGetAssetTags, const UWidgetBlueprintGeneratedClass*, TArray<UObject::FAssetRegistryTag>&);
 
 	// called by UWidgetBlueprintGeneratedClass::GetAssetRegistryTags()
-	static FGetAssetTags GetAssetTags;
+	static UMG_API FGetAssetTags GetAssetTags;
 };
 #endif
 
@@ -69,14 +69,14 @@ public:
  * to give themselves many of the same capabilities as AActor blueprints, like dynamic delegate binding for
  * widgets.
  */
-UCLASS()
-class UMG_API UWidgetBlueprintGeneratedClass : public UBlueprintGeneratedClass
+UCLASS(MinimalAPI)
+class UWidgetBlueprintGeneratedClass : public UBlueprintGeneratedClass
 {
 	GENERATED_BODY()
 	friend class FWidgetBlueprintCompilerContext;
 
 public:
-	UWidgetBlueprintGeneratedClass();
+	UMG_API UWidgetBlueprintGeneratedClass();
 
 private:
 
@@ -136,32 +136,32 @@ public:
 	
 public:
 	UWidgetTree* GetWidgetTreeArchetype() const { return WidgetTree; }
-	void SetWidgetTreeArchetype(UWidgetTree* InWidgetTree);
+	UMG_API void SetWidgetTreeArchetype(UWidgetTree* InWidgetTree);
 
-	void GetNamedSlotArchetypeContent(TFunctionRef<void(FName /*SlotName*/, UWidget* /*Content*/)> Predicate) const;
+	UMG_API void GetNamedSlotArchetypeContent(TFunctionRef<void(FName /*SlotName*/, UWidget* /*Content*/)> Predicate) const;
 
 	// Walks up the hierarchy looking for a valid widget tree.
-	UWidgetBlueprintGeneratedClass* FindWidgetTreeOwningClass() const;
+	UMG_API UWidgetBlueprintGeneratedClass* FindWidgetTreeOwningClass() const;
 
 	//~ Begin UObject interface
-	virtual void Serialize(FArchive& Ar) override;
-	virtual void PostLoad() override;
-	virtual bool NeedsLoadForServer() const override;
+	UMG_API virtual void Serialize(FArchive& Ar) override;
+	UMG_API virtual void PostLoad() override;
+	UMG_API virtual bool NeedsLoadForServer() const override;
 #if WITH_EDITOR
-	virtual void GetAssetRegistryTags(TArray<FAssetRegistryTag>& OutTags) const override;
+	UMG_API virtual void GetAssetRegistryTags(TArray<FAssetRegistryTag>& OutTags) const override;
 #endif
 	//~ End UObject interface
 
-	virtual void PurgeClass(bool bRecompilingOnLoad) override;
+	UMG_API virtual void PurgeClass(bool bRecompilingOnLoad) override;
 
 	/**
 	 * This is the function that makes UMG work.  Once a user widget is constructed, it will post load
 	 * call into its generated class and ask to be initialized.  The class will perform all the delegate
 	 * binding and wiring necessary to have the user's widget perform as desired.
 	 */
-	void InitializeWidget(UUserWidget* UserWidget) const;
+	UMG_API void InitializeWidget(UUserWidget* UserWidget) const;
 
-	static void InitializeWidgetStatic(UUserWidget* UserWidget
+	static UMG_API void InitializeWidgetStatic(UUserWidget* UserWidget
 		, const UClass* InClass
 		, UWidgetTree* InWidgetTree
 		, const UClass* InWidgetTreeWidgetClass
@@ -171,7 +171,7 @@ public:
 	bool ClassRequiresNativeTick() const { return bClassRequiresNativeTick; }
 
 #if WITH_EDITOR
-	void SetClassRequiresNativeTick(bool InClassRequiresNativeTick);
+	UMG_API void SetClassRequiresNativeTick(bool InClassRequiresNativeTick);
 #endif
 
 	/** Find the first extension of the requested type. */
@@ -182,10 +182,10 @@ public:
 	}
 
 	/** Find the first extension of the requested type. */
-	UWidgetBlueprintGeneratedClassExtension* GetExtension(TSubclassOf<UWidgetBlueprintGeneratedClassExtension> InExtensionType, bool bIncludeSuper = true);
+	UMG_API UWidgetBlueprintGeneratedClassExtension* GetExtension(TSubclassOf<UWidgetBlueprintGeneratedClassExtension> InExtensionType, bool bIncludeSuper = true);
 
 	/** Find the extensions of the requested type. */
-	TArray<UWidgetBlueprintGeneratedClassExtension*> GetExtensions(TSubclassOf<UWidgetBlueprintGeneratedClassExtension> InExtensionType, bool bIncludeSuper = true);
+	UMG_API TArray<UWidgetBlueprintGeneratedClassExtension*> GetExtensions(TSubclassOf<UWidgetBlueprintGeneratedClassExtension> InExtensionType, bool bIncludeSuper = true);
 
 	template<typename Predicate>
 	void ForEachExtension(Predicate Pred, bool bIncludeSuper = true) const
@@ -205,8 +205,8 @@ public:
 	}
 
 private:
-	static void InitializeBindingsStatic(UUserWidget* UserWidget, const TArrayView<const FDelegateRuntimeBinding> InBindings, const TMap<FName, FObjectPropertyBase*>& InPropertyMap);
-	static void BindAnimationsStatic(UUserWidget* Instance, const TArrayView<UWidgetAnimation*> InAnimations, const TMap<FName, FObjectPropertyBase*>& InPropertyMap);
+	static UMG_API void InitializeBindingsStatic(UUserWidget* UserWidget, const TArrayView<const FDelegateRuntimeBinding> InBindings, const TMap<FName, FObjectPropertyBase*>& InPropertyMap);
+	static UMG_API void BindAnimationsStatic(UUserWidget* Instance, const TArrayView<UWidgetAnimation*> InAnimations, const TMap<FName, FObjectPropertyBase*>& InPropertyMap);
 
-	void GetExtensions(TArray<UWidgetBlueprintGeneratedClassExtension*>& OutExtensions, TSubclassOf<UWidgetBlueprintGeneratedClassExtension> InExtensionType, bool bIncludeSuper);
+	UMG_API void GetExtensions(TArray<UWidgetBlueprintGeneratedClassExtension*>& OutExtensions, TSubclassOf<UWidgetBlueprintGeneratedClassExtension> InExtensionType, bool bIncludeSuper);
 };
