[Engine]

diff --git a/Engine/Source/Runtime/UMG/Public/Blueprint/WidgetBlueprintLibrary.h b/Engine/Source/Runtime/UMG/Public/Blueprint/WidgetBlueprintLibrary.h
--- a/Engine/Source/Runtime/UMG/Public/Blueprint/WidgetBlueprintLibrary.h
+++ b/Engine/Source/Runtime/UMG/Public/Blueprint/WidgetBlueprintLibrary.h
@@ -26,8 +26,8 @@ class USlateBrushAsset;
 class UTexture2D;
 class UUserWidget;
 
-UCLASS(meta=(ScriptName="WidgetLibrary"))
-class UMG_API UWidgetBlueprintLibrary : public UBlueprintFunctionLibrary
+UCLASS(meta=(ScriptName="WidgetLibrary"), MinimalAPI)
+class UWidgetBlueprintLibrary : public UBlueprintFunctionLibrary
 {
 	GENERATED_UCLASS_BODY()
 
@@ -35,14 +35,14 @@ public:
 
 	/** Creates a widget */
 	UFUNCTION(BlueprintCallable, BlueprintCosmetic, meta=( WorldContext="WorldContextObject", DisplayName="Create Widget", BlueprintInternalUseOnly="true" ), Category="Widget")
-	static class UUserWidget* Create(UObject* WorldContextObject, TSubclassOf<class UUserWidget> WidgetType, APlayerController* OwningPlayer);
+	static UMG_API class UUserWidget* Create(UObject* WorldContextObject, TSubclassOf<class UUserWidget> WidgetType, APlayerController* OwningPlayer);
 
 	/**
 	 * Creates a new drag and drop operation that can be returned from a drag begin to inform the UI what i
 	 * being dragged and dropped and what it looks like.
 	 */
 	UFUNCTION(BlueprintCallable, Category="Widget|Drag and Drop", meta=( BlueprintInternalUseOnly="true" ))
-	static UDragDropOperation* CreateDragDropOperation(TSubclassOf<UDragDropOperation> OperationClass);
+	static UMG_API UDragDropOperation* CreateDragDropOperation(TSubclassOf<UDragDropOperation> OperationClass);
 
 	/**
 	 * Setup an input mode that allows only the UI to respond to user input.
@@ -50,14 +50,14 @@ public:
 	 * Note: This means that any bound Input Events in the widget will not be called!
 	 */
 	UFUNCTION(BlueprintCallable, BlueprintCosmetic, Category="Input", meta=(DisplayName="Set Input Mode UI Only"))
-	static void SetInputMode_UIOnlyEx(APlayerController* PlayerController, UWidget* InWidgetToFocus = nullptr, EMouseLockMode InMouseLockMode = EMouseLockMode::DoNotLock, const bool bFlushInput = false);
+	static UMG_API void SetInputMode_UIOnlyEx(APlayerController* PlayerController, UWidget* InWidgetToFocus = nullptr, EMouseLockMode InMouseLockMode = EMouseLockMode::DoNotLock, const bool bFlushInput = false);
 	/**
 	 * Setup an input mode that allows only the UI to respond to user input, and if the UI doesn't handle it player input / player controller gets a chance.
 	 * 
 	 * Note: This means that any bound Input events in the widget will be called.
 	 */
 	UFUNCTION(BlueprintCallable, BlueprintCosmetic, Category="Input", meta = (DisplayName = "Set Input Mode Game And UI"))
-	static void SetInputMode_GameAndUIEx(APlayerController* PlayerController, UWidget* InWidgetToFocus = nullptr, EMouseLockMode InMouseLockMode = EMouseLockMode::DoNotLock, bool bHideCursorDuringCapture = true, const bool bFlushInput = false);
+	static UMG_API void SetInputMode_GameAndUIEx(APlayerController* PlayerController, UWidget* InWidgetToFocus = nullptr, EMouseLockMode InMouseLockMode = EMouseLockMode::DoNotLock, bool bHideCursorDuringCapture = true, const bool bFlushInput = false);
 
 	/**
 	 * Setup an input mode that allows only player input / player controller to respond to user input.
@@ -65,15 +65,15 @@ public:
 	 * Note: Any bound Input Events in this widget will be called.
 	 */
 	UFUNCTION(BlueprintCallable, BlueprintCosmetic, Category="Input")
-	static void SetInputMode_GameOnly(APlayerController* PlayerController, const bool bFlushInput = false);
+	static UMG_API void SetInputMode_GameOnly(APlayerController* PlayerController, const bool bFlushInput = false);
 
 	/** */
 	UFUNCTION(BlueprintCallable, BlueprintCosmetic, Category="Focus")
-	static void SetFocusToGameViewport();
+	static UMG_API void SetFocusToGameViewport();
 
 	/** Draws a box */
 	UFUNCTION(BlueprintCallable, Category="Painting")
-	static void DrawBox(UPARAM(ref) FPaintContext& Context, FVector2D Position, FVector2D Size, USlateBrushAsset* Brush, FLinearColor Tint = FLinearColor::White);
+	static UMG_API void DrawBox(UPARAM(ref) FPaintContext& Context, FVector2D Position, FVector2D Size, USlateBrushAsset* Brush, FLinearColor Tint = FLinearColor::White);
 
 	/**
 	 * Draws a hermite spline.
@@ -86,7 +86,7 @@ public:
 	 * @param Thickness		How many pixels thick this spline should be.
 	 */
 	UFUNCTION(BlueprintCallable, meta = (AdvancedDisplay = "6"), Category = "Painting")
-	static void DrawSpline(UPARAM(ref) FPaintContext& Context, FVector2D Start, FVector2D StartDir, FVector2D End, FVector2D EndDir, FLinearColor Tint = FLinearColor::White, float Thickness = 1.0f);
+	static UMG_API void DrawSpline(UPARAM(ref) FPaintContext& Context, FVector2D Start, FVector2D StartDir, FVector2D End, FVector2D EndDir, FLinearColor Tint = FLinearColor::White, float Thickness = 1.0f);
 
 	/**
 	 * Draws a line.
@@ -98,7 +98,7 @@ public:
 	 * @param Thickness		How many pixels thick this line should be.
 	 */
 	UFUNCTION(BlueprintCallable, meta = (AdvancedDisplay = "4"), Category = "Painting")
-	static void DrawLine(UPARAM(ref) FPaintContext& Context, FVector2D PositionA, FVector2D PositionB, FLinearColor Tint = FLinearColor::White, bool bAntiAlias = true, float Thickness = 1.0f);
+	static UMG_API void DrawLine(UPARAM(ref) FPaintContext& Context, FVector2D PositionA, FVector2D PositionB, FLinearColor Tint = FLinearColor::White, bool bAntiAlias = true, float Thickness = 1.0f);
 
 	/**
 	 * Draws several line segments.
@@ -109,7 +109,7 @@ public:
 	 * @param Thickness		How many pixels thick this line should be.
 	 */
 	UFUNCTION(BlueprintCallable, meta = (AdvancedDisplay = "3"), Category = "Painting")
-	static void DrawLines(UPARAM(ref) FPaintContext& Context, const TArray<FVector2D>& Points, FLinearColor Tint = FLinearColor::White, bool bAntiAlias = true, float Thickness = 1.0f);
+	static UMG_API void DrawLines(UPARAM(ref) FPaintContext& Context, const TArray<FVector2D>& Points, FLinearColor Tint = FLinearColor::White, bool bAntiAlias = true, float Thickness = 1.0f);
 
 	/**
 	 * Draws text.
@@ -119,7 +119,7 @@ public:
 	 * @param Tint			Color to render the line.
 	 */
 	UFUNCTION(BlueprintCallable, Category="Painting", meta=( DeprecatedFunction, DeprecationMessage = "Use Draw Text instead", DisplayName="Draw String"))
-	static void DrawText(UPARAM(ref) FPaintContext& Context, const FString& InString, FVector2D Position, FLinearColor Tint = FLinearColor::White);
+	static UMG_API void DrawText(UPARAM(ref) FPaintContext& Context, const FString& InString, FVector2D Position, FLinearColor Tint = FLinearColor::White);
 
 	/**
 	 * Draws text.
@@ -129,50 +129,50 @@ public:
 	 * @param Tint			Color to render the line.
 	 */
 	UFUNCTION(BlueprintCallable, Category="Painting", meta=(DisplayName="Draw Text"))
-	static void DrawTextFormatted(UPARAM(ref) FPaintContext& Context, const FText& Text, FVector2D Position, UFont* Font, float FontSize = 16.0f, FName FontTypeFace = FName(TEXT("Regular")), FLinearColor Tint = FLinearColor::White);
+	static UMG_API void DrawTextFormatted(UPARAM(ref) FPaintContext& Context, const FText& Text, FVector2D Position, UFont* Font, float FontSize = 16.0f, FName FontTypeFace = FName(TEXT("Regular")), FLinearColor Tint = FLinearColor::White);
 
 	/**
 	 * The event reply to use when you choose to handle an event.  This will prevent the event 
 	 * from continuing to bubble up / down the widget hierarchy.
 	 */
 	UFUNCTION(BlueprintPure, Category="Widget|Event Reply")
-	static FEventReply Handled();
+	static UMG_API FEventReply Handled();
 
 	/** The event reply to use when you choose not to handle an event. */
 	UFUNCTION(BlueprintPure, Category="Widget|Event Reply")
-	static FEventReply Unhandled();
+	static UMG_API FEventReply Unhandled();
 
 	/**  */
 	UFUNCTION(BlueprintPure, meta=( DefaultToSelf="CapturingWidget" ), Category="Widget|Event Reply")
-	static FEventReply CaptureMouse(UPARAM(ref) FEventReply& Reply, UWidget* CapturingWidget);
+	static UMG_API FEventReply CaptureMouse(UPARAM(ref) FEventReply& Reply, UWidget* CapturingWidget);
 
 	/**  */
 	UFUNCTION(BlueprintPure, Category="Widget|Event Reply")
-	static FEventReply ReleaseMouseCapture(UPARAM(ref) FEventReply& Reply);
+	static UMG_API FEventReply ReleaseMouseCapture(UPARAM(ref) FEventReply& Reply);
 
 	UFUNCTION( BlueprintPure, Category = "Widget|Event Reply", meta = ( HidePin = "CapturingWidget", DefaultToSelf = "CapturingWidget" ) )
-	static FEventReply LockMouse( UPARAM( ref ) FEventReply& Reply, UWidget* CapturingWidget );
+	static UMG_API FEventReply LockMouse( UPARAM( ref ) FEventReply& Reply, UWidget* CapturingWidget );
 
 	UFUNCTION( BlueprintPure, Category = "Widget|Event Reply" )
-	static FEventReply UnlockMouse( UPARAM( ref ) FEventReply& Reply );
+	static UMG_API FEventReply UnlockMouse( UPARAM( ref ) FEventReply& Reply );
 
 	/**  */
 	UFUNCTION(BlueprintPure, meta= (HidePin="CapturingWidget", DefaultToSelf="CapturingWidget"), Category="Widget|Event Reply")
-	static FEventReply SetUserFocus(UPARAM(ref) FEventReply& Reply, UWidget* FocusWidget, bool bInAllUsers = false);
+	static UMG_API FEventReply SetUserFocus(UPARAM(ref) FEventReply& Reply, UWidget* FocusWidget, bool bInAllUsers = false);
 
 	UFUNCTION(BlueprintPure, meta = (DeprecatedFunction, DeprecationMessage = "Use SetUserFocus() instead"), Category = "Widget|Event Reply")
-	static FEventReply CaptureJoystick(UPARAM(ref) FEventReply& Reply, UWidget* CapturingWidget, bool bInAllJoysticks = false);
+	static UMG_API FEventReply CaptureJoystick(UPARAM(ref) FEventReply& Reply, UWidget* CapturingWidget, bool bInAllJoysticks = false);
 
 	/**  */
 	UFUNCTION(BlueprintPure, meta = (HidePin = "CapturingWidget", DefaultToSelf = "CapturingWidget"), Category = "Widget|Event Reply")
-	static FEventReply ClearUserFocus(UPARAM(ref) FEventReply& Reply, bool bInAllUsers = false);
+	static UMG_API FEventReply ClearUserFocus(UPARAM(ref) FEventReply& Reply, bool bInAllUsers = false);
 
 	UFUNCTION(BlueprintPure, meta = (DeprecatedFunction, DeprecationMessage = "Use ClearUserFocus() instead"), Category = "Widget|Event Reply")
-	static FEventReply ReleaseJoystickCapture(UPARAM(ref) FEventReply& Reply, bool bInAllJoysticks = false);
+	static UMG_API FEventReply ReleaseJoystickCapture(UPARAM(ref) FEventReply& Reply, bool bInAllJoysticks = false);
 
 	/**  */
 	UFUNCTION(BlueprintPure, Category="Widget|Event Reply")
-	static FEventReply SetMousePosition(UPARAM(ref) FEventReply& Reply, FVector2D NewMousePosition);
+	static UMG_API FEventReply SetMousePosition(UPARAM(ref) FEventReply& Reply, FVector2D NewMousePosition);
 
 	/**
 	 * Ask Slate to detect if a user starts dragging in this widget later.  Slate internally tracks the movement
@@ -182,7 +182,7 @@ public:
 	 * @param DragKey		       This button should be pressed to detect the drag
 	 */
 	UFUNCTION(BlueprintPure, meta=( HidePin="WidgetDetectingDrag", DefaultToSelf="WidgetDetectingDrag" ), Category="Widget|Drag and Drop|Event Reply")
-	static FEventReply DetectDrag(UPARAM(ref) FEventReply& Reply, UWidget* WidgetDetectingDrag, FKey DragKey);
+	static UMG_API FEventReply DetectDrag(UPARAM(ref) FEventReply& Reply, UWidget* WidgetDetectingDrag, FKey DragKey);
 
 	/**
 	 * Given the pointer event, emit the DetectDrag reply if the provided key was pressed.
@@ -192,31 +192,31 @@ public:
 	 * @param DragKey		       This button should be pressed to detect the drag, won't emit the DetectDrag FEventReply unless this is pressed.
 	 */
 	UFUNCTION(BlueprintCallable, meta=( HidePin="WidgetDetectingDrag", DefaultToSelf="WidgetDetectingDrag" ), Category="Widget|Drag and Drop|Event Reply")
-	static FEventReply DetectDragIfPressed(const FPointerEvent& PointerEvent, UWidget* WidgetDetectingDrag, FKey DragKey);
+	static UMG_API FEventReply DetectDragIfPressed(const FPointerEvent& PointerEvent, UWidget* WidgetDetectingDrag, FKey DragKey);
 
 	/**
 	 * An event should return FReply::Handled().EndDragDrop() to request that the current drag/drop operation be terminated.
 	 */
 	UFUNCTION(BlueprintPure, Category="Widget|Drag and Drop|Event Reply")
-	static FEventReply EndDragDrop(UPARAM(ref) FEventReply& Reply);
+	static UMG_API FEventReply EndDragDrop(UPARAM(ref) FEventReply& Reply);
 
 	/**
 	 * Returns true if a drag/drop event is occurring that a widget can handle.
 	 */
 	UFUNCTION(BlueprintPure, BlueprintCosmetic, Category="Widget|Drag and Drop")
-	static bool IsDragDropping();
+	static UMG_API bool IsDragDropping();
 
 	/**
 	 * Returns the drag and drop operation that is currently occurring if any, otherwise nothing.
 	 */
 	UFUNCTION(BlueprintPure, BlueprintCosmetic, Category="Widget|Drag and Drop")
-	static UDragDropOperation* GetDragDroppingContent();
+	static UMG_API UDragDropOperation* GetDragDroppingContent();
 
 	/**
 	 * Cancels any current drag drop operation.
 	 */
 	UFUNCTION(BlueprintCallable, Category="Widget|Drag and Drop")
-	static void CancelDragDrop();
+	static UMG_API void CancelDragDrop();
 
 	/**
 	 * Creates a Slate Brush from a Slate Brush Asset
@@ -224,7 +224,7 @@ public:
 	 * @return A new slate brush using the asset's brush.
 	 */
 	UFUNCTION(BlueprintPure, Category="Widget|Brush")
-	static FSlateBrush MakeBrushFromAsset(USlateBrushAsset* BrushAsset);
+	static UMG_API FSlateBrush MakeBrushFromAsset(USlateBrushAsset* BrushAsset);
 
 	/** 
 	 * Creates a Slate Brush from a Texture2D
@@ -235,7 +235,7 @@ public:
 	 * @return A new slate brush using the texture.
 	 */
 	UFUNCTION(BlueprintPure, Category="Widget|Brush")
-	static FSlateBrush MakeBrushFromTexture(UTexture2D* Texture, int32 Width = 0, int32 Height = 0);
+	static UMG_API FSlateBrush MakeBrushFromTexture(UTexture2D* Texture, int32 Width = 0, int32 Height = 0);
 
 	/**
 	 * Creates a Slate Brush from a Material.  Materials don't have an implicit size, so providing a widget and height
@@ -244,37 +244,37 @@ public:
 	 * @return A new slate brush using the material.
 	 */
 	UFUNCTION(BlueprintPure, Category="Widget|Brush")
-	static FSlateBrush MakeBrushFromMaterial(UMaterialInterface* Material, int32 Width = 32, int32 Height = 32);
+	static UMG_API FSlateBrush MakeBrushFromMaterial(UMaterialInterface* Material, int32 Width = 32, int32 Height = 32);
 
 	/**
 	 * Gets the resource object on a brush.  This could be a UTexture2D or a UMaterialInterface.
 	 */
 	UFUNCTION(BlueprintPure, Category="Widget|Brush")
-	static UObject* GetBrushResource(const FSlateBrush& Brush);
+	static UMG_API UObject* GetBrushResource(const FSlateBrush& Brush);
 
 	/**
 	 * Gets the brush resource as a texture 2D.
 	 */
 	UFUNCTION(BlueprintPure, Category="Widget|Brush")
-	static UTexture2D* GetBrushResourceAsTexture2D(const FSlateBrush& Brush);
+	static UMG_API UTexture2D* GetBrushResourceAsTexture2D(const FSlateBrush& Brush);
 
 	/**
 	 * Gets the brush resource as a material.
 	 */
 	UFUNCTION(BlueprintPure, Category="Widget|Brush")
-	static UMaterialInterface* GetBrushResourceAsMaterial(const FSlateBrush& Brush);
+	static UMG_API UMaterialInterface* GetBrushResourceAsMaterial(const FSlateBrush& Brush);
 
 	/**
 	 * Sets the resource on a brush to be a UTexture2D.
 	 */
 	UFUNCTION(BlueprintCallable, Category="Widget|Brush")
-	static void SetBrushResourceToTexture(UPARAM(ref) FSlateBrush& Brush, UTexture2D* Texture);
+	static UMG_API void SetBrushResourceToTexture(UPARAM(ref) FSlateBrush& Brush, UTexture2D* Texture);
 
 	/**
 	 * Sets the resource on a brush to be a Material.
 	 */
 	UFUNCTION(BlueprintCallable, Category="Widget|Brush")
-	static void SetBrushResourceToMaterial(UPARAM(ref) FSlateBrush& Brush, UMaterialInterface* Material);
+	static UMG_API void SetBrushResourceToMaterial(UPARAM(ref) FSlateBrush& Brush, UMaterialInterface* Material);
 
 	/**
 	 * Creates a Slate Brush that wont draw anything, the "Null Brush".
@@ -282,7 +282,7 @@ public:
 	 * @return A new slate brush that wont draw anything.
 	 */
 	UFUNCTION(BlueprintPure, Category="Widget|Brush")
-	static FSlateBrush NoResourceBrush();
+	static UMG_API FSlateBrush NoResourceBrush();
 
 	/**
 	 * Gets the material that allows changes to parameters at runtime.  The brush must already have a material assigned to it, 
@@ -291,11 +291,11 @@ public:
 	 * @return A material that supports dynamic input from the game.
 	 */
 	UFUNCTION(BlueprintPure, Category="Widget|Brush")
-	static UMaterialInstanceDynamic* GetDynamicMaterial(UPARAM(ref) FSlateBrush& Brush);
+	static UMG_API UMaterialInstanceDynamic* GetDynamicMaterial(UPARAM(ref) FSlateBrush& Brush);
 
 	/** Closes any popup menu */
 	UFUNCTION(BlueprintCallable, BlueprintCosmetic, Category="Widget|Menu")
-	static void DismissAllMenus();
+	static UMG_API void DismissAllMenus();
 
 	/**
 	 * Find all widgets of a certain class.
@@ -304,7 +304,7 @@ public:
 	 * @param TopLevelOnly Only the widgets that are direct children of the viewport will be returned.
 	 */
 	UFUNCTION(BlueprintCallable, BlueprintCosmetic, Category="Widget", meta=( WorldContext="WorldContextObject", DeterminesOutputType = "WidgetClass", DynamicOutputParam = "FoundWidgets" ))
-	static void GetAllWidgetsOfClass(UObject* WorldContextObject, TArray<UUserWidget*>& FoundWidgets, TSubclassOf<UUserWidget> WidgetClass, bool TopLevelOnly = true);
+	static UMG_API void GetAllWidgetsOfClass(UObject* WorldContextObject, TArray<UUserWidget*>& FoundWidgets, TSubclassOf<UUserWidget> WidgetClass, bool TopLevelOnly = true);
 
 	/**
 	* Find all widgets in the world with the specified interface.
@@ -314,28 +314,28 @@ public:
 	* @param TopLevelOnly Only the widgets that are direct children of the viewport will be returned.
 	*/
 	UFUNCTION(BlueprintCallable, BlueprintCosmetic, Category="Widget", meta = (WorldContext = "WorldContextObject", DeterminesOutputType = "WidgetClass", DynamicOutputParam = "FoundWidgets"))
-	static void GetAllWidgetsWithInterface(UObject* WorldContextObject, TArray<UUserWidget*>& FoundWidgets, TSubclassOf<UInterface> Interface, bool TopLevelOnly);
+	static UMG_API void GetAllWidgetsWithInterface(UObject* WorldContextObject, TArray<UUserWidget*>& FoundWidgets, TSubclassOf<UInterface> Interface, bool TopLevelOnly);
 
 	UFUNCTION(BlueprintPure, Category="Widget", meta = (CompactNodeTitle = "->", BlueprintAutocast))
-	static FInputEvent GetInputEventFromKeyEvent(const FKeyEvent& Event);
+	static UMG_API FInputEvent GetInputEventFromKeyEvent(const FKeyEvent& Event);
 
 	UFUNCTION(BlueprintPure, Category="Widget", meta = (CompactNodeTitle = "->", BlueprintAutocast))
-	static FKeyEvent GetKeyEventFromAnalogInputEvent(const FAnalogInputEvent& Event);
+	static UMG_API FKeyEvent GetKeyEventFromAnalogInputEvent(const FAnalogInputEvent& Event);
 
 	UFUNCTION(BlueprintPure, Category="Widget", meta = ( CompactNodeTitle = "->", BlueprintAutocast ))
-	static FInputEvent GetInputEventFromCharacterEvent(const FCharacterEvent& Event);
+	static UMG_API FInputEvent GetInputEventFromCharacterEvent(const FCharacterEvent& Event);
 
 	UFUNCTION(BlueprintPure, Category="Widget", meta = ( CompactNodeTitle = "->", BlueprintAutocast ))
-	static FInputEvent GetInputEventFromPointerEvent(const FPointerEvent& Event);
+	static UMG_API FInputEvent GetInputEventFromPointerEvent(const FPointerEvent& Event);
 
 	UFUNCTION(BlueprintPure, Category="Widget", meta = ( CompactNodeTitle = "->", BlueprintAutocast ))
-	static FInputEvent GetInputEventFromNavigationEvent(const FNavigationEvent& Event);
+	static UMG_API FInputEvent GetInputEventFromNavigationEvent(const FNavigationEvent& Event);
 
 	/**
 	 * Gets the amount of padding that needs to be added when accounting for the safe zone on TVs.
 	 */
 	UFUNCTION(BlueprintCallable, Category="Widget|Safe Zone", meta=( WorldContext="WorldContextObject" ))
-	static void GetSafeZonePadding(UObject* WorldContextObject, FVector4& SafePadding, FVector2D& SafePaddingScale, FVector4& SpillOverPadding);
+	static UMG_API void GetSafeZonePadding(UObject* WorldContextObject, FVector4& SafePadding, FVector2D& SafePaddingScale, FVector4& SpillOverPadding);
 
 	/**
 	* Apply color deficiency correction settings to the game window 
@@ -345,25 +345,25 @@ public:
 	* @param ShowCorrectionWithDeficiency If you're correcting the color deficiency, you can use this to visualize what the correction looks like with the deficiency.
 	*/
 	UFUNCTION(BlueprintCallable, BlueprintCosmetic, Category = "Widget|Accessibility", meta = (AdvancedDisplay = "3"))
-	static void SetColorVisionDeficiencyType(EColorVisionDeficiency Type, float Severity, bool CorrectDeficiency, bool ShowCorrectionWithDeficiency);
+	static UMG_API void SetColorVisionDeficiencyType(EColorVisionDeficiency Type, float Severity, bool CorrectDeficiency, bool ShowCorrectionWithDeficiency);
 
 	/**
 	 * Loads or sets a hardware cursor from the content directory in the game.
 	 */
 	UFUNCTION(BlueprintCallable, Category="Widget|Hardware Cursor", meta=( WorldContext="WorldContextObject" ))
-	static bool SetHardwareCursor(UObject* WorldContextObject, EMouseCursor::Type CursorShape, FName CursorName, FVector2D HotSpot);
+	static UMG_API bool SetHardwareCursor(UObject* WorldContextObject, EMouseCursor::Type CursorShape, FName CursorName, FVector2D HotSpot);
 
 	UFUNCTION(BlueprintCallable, Category = "Widget|Window Title Bar")
-	static void SetWindowTitleBarState(UWidget* TitleBarContent, EWindowTitleBarMode Mode, bool bTitleBarDragEnabled, bool bWindowButtonsVisible, bool bTitleBarVisible);
+	static UMG_API void SetWindowTitleBarState(UWidget* TitleBarContent, EWindowTitleBarMode Mode, bool bTitleBarDragEnabled, bool bWindowButtonsVisible, bool bTitleBarVisible);
 
 	UFUNCTION(BlueprintCallable, Category = "Widget|Window Title Bar")
-	static void RestorePreviousWindowTitleBarState();
+	static UMG_API void RestorePreviousWindowTitleBarState();
 
 	DECLARE_DYNAMIC_DELEGATE(FOnGameWindowCloseButtonClickedDelegate);
 
 	UFUNCTION(BlueprintCallable, Category = "Widget|Window Title Bar")
-	static void SetWindowTitleBarOnCloseClickedDelegate(FOnGameWindowCloseButtonClickedDelegate Delegate);
+	static UMG_API void SetWindowTitleBarOnCloseClickedDelegate(FOnGameWindowCloseButtonClickedDelegate Delegate);
 
 	UFUNCTION(BlueprintCallable, Category = "Widget|Window Title Bar")
-	static void SetWindowTitleBarCloseButtonActive(bool bActive);
+	static UMG_API void SetWindowTitleBarCloseButtonActive(bool bActive);
 };
diff --git a/Engine/Source/Runtime/UMG/Public/Blueprint/WidgetChild.h b/Engine/Source/Runtime/UMG/Public/Blueprint/WidgetChild.h
--- a/Engine/Source/Runtime/UMG/Public/Blueprint/WidgetChild.h
+++ b/Engine/Source/Runtime/UMG/Public/Blueprint/WidgetChild.h
@@ -15,12 +15,12 @@ class UWidget;
  * Represent a Widget present in the Tree Widget of the UserWidget
  */
 USTRUCT()
-struct UMG_API FWidgetChild
+struct FWidgetChild
 {
 	GENERATED_BODY();
 public:
-	FWidgetChild();
-	FWidgetChild(const class UUserWidget* Outer, FName InChildName);
+	UMG_API FWidgetChild();
+	UMG_API FWidgetChild(const class UUserWidget* Outer, FName InChildName);
 
 	FName GetFName() const
 	{
@@ -33,7 +33,7 @@ public:
 	};
 
 	/** Resolves the Widget ptr using it's name. */
-	UWidget* Resolve(const class UWidgetTree* WidgetTree);
+	UMG_API UWidget* Resolve(const class UWidgetTree* WidgetTree);
 
 private:
 	/** This either the widget to focus, OR the name of the function to call. */
diff --git a/Engine/Source/Runtime/UMG/Public/Blueprint/WidgetLayoutLibrary.h b/Engine/Source/Runtime/UMG/Public/Blueprint/WidgetLayoutLibrary.h
--- a/Engine/Source/Runtime/UMG/Public/Blueprint/WidgetLayoutLibrary.h
+++ b/Engine/Source/Runtime/UMG/Public/Blueprint/WidgetLayoutLibrary.h
@@ -24,8 +24,8 @@ class UWrapBoxSlot;
 class UWidgetSwitcherSlot;
 class UWidget;
 
-UCLASS()
-class UMG_API UWidgetLayoutLibrary : public UBlueprintFunctionLibrary
+UCLASS(MinimalAPI)
+class UWidgetLayoutLibrary : public UBlueprintFunctionLibrary
 {
 	GENERATED_UCLASS_BODY()
 
@@ -41,7 +41,7 @@ public:
 	 * @return true if the position projects onto the screen.
 	 */
 	UFUNCTION(BlueprintPure, BlueprintCosmetic, Category="Viewport")
-	static bool ProjectWorldLocationToWidgetPosition(APlayerController* PlayerController, FVector WorldLocation, FVector2D& ScreenPosition, bool bPlayerViewportRelative);
+	static UMG_API bool ProjectWorldLocationToWidgetPosition(APlayerController* PlayerController, FVector WorldLocation, FVector2D& ScreenPosition, bool bPlayerViewportRelative);
 
 	/**
 	 * Convert a World Space 3D position into a 2D Widget Screen Space position, with distance from the camera the Z component.  This
@@ -49,24 +49,24 @@ public:
 	 *
 	 * @return true if the world coordinate was successfully projected to the screen.
 	 */
-	static bool ProjectWorldLocationToWidgetPositionWithDistance(APlayerController* PlayerController, FVector WorldLocation, FVector& ScreenPosition, bool bPlayerViewportRelative);
+	static UMG_API bool ProjectWorldLocationToWidgetPositionWithDistance(APlayerController* PlayerController, FVector WorldLocation, FVector& ScreenPosition, bool bPlayerViewportRelative);
 
 	/**
 	 * Gets the current DPI Scale being applied to the viewport and all the Widgets.
 	 */
 	UFUNCTION(BlueprintPure, BlueprintCosmetic, Category="Viewport", meta=( WorldContext="WorldContextObject" ))
-	static float GetViewportScale(const UObject* WorldContextObject);
+	static UMG_API float GetViewportScale(const UObject* WorldContextObject);
 
 	/**
 	 * Gets the current DPI Scale being applied to the viewport and all the Widgets.
 	 */
-	static float GetViewportScale(const UGameViewportClient* ViewportClient);
+	static UMG_API float GetViewportScale(const UGameViewportClient* ViewportClient);
 
 	/**
 	 * Gets the size of the game viewport.
 	 */
 	UFUNCTION(BlueprintPure, BlueprintCosmetic, Category="Viewport", meta=( WorldContext="WorldContextObject", Keywords = "screen size"))
-	static FVector2D GetViewportSize(UObject* WorldContextObject);
+	static UMG_API FVector2D GetViewportSize(UObject* WorldContextObject);
 
 	/**
 	 * Gets the geometry of the widget holding all widgets added to the "Viewport".  You
@@ -74,7 +74,7 @@ public:
 	 * widget.
 	 */
 	UFUNCTION(BlueprintCallable, Category="Viewport", meta=( WorldContext="WorldContextObject" ))
-	static FGeometry GetViewportWidgetGeometry(UObject* WorldContextObject);
+	static UMG_API FGeometry GetViewportWidgetGeometry(UObject* WorldContextObject);
 
 	/**
 	 * Gets the geometry of the widget holding all widgets added to the "Player Screen". You
@@ -82,19 +82,19 @@ public:
 	 * widget.
 	 */
 	UFUNCTION(BlueprintCallable, Category="Viewport")
-	static FGeometry GetPlayerScreenWidgetGeometry(APlayerController* PlayerController);
+	static UMG_API FGeometry GetPlayerScreenWidgetGeometry(APlayerController* PlayerController);
 
 	/**
 	 * Gets the platform's mouse cursor position.  This is the 'absolute' desktop location of the mouse.
 	 */
 	UFUNCTION(BlueprintCallable, Category="Viewport")
-	static FVector2D GetMousePositionOnPlatform();
+	static UMG_API FVector2D GetMousePositionOnPlatform();
 
 	/**
 	 * Gets the platform's mouse cursor position in the local space of the viewport widget.
 	 */
 	UFUNCTION(BlueprintCallable, Category="Viewport", meta=( WorldContext="WorldContextObject" ))
-	static FVector2D GetMousePositionOnViewport(UObject* WorldContextObject);
+	static UMG_API FVector2D GetMousePositionOnViewport(UObject* WorldContextObject);
 
 	/**
 	 * Gets the mouse position of the player controller, scaled by the DPI.  If you're trying to go from raw mouse screenspace coordinates
@@ -104,57 +104,57 @@ public:
 	 */
 	//UE_DEPRECATED(4.17, "Use GetMousePositionOnViewport() instead.  Optionally and for more options, you can use GetViewportWidgetGeometry and GetPlayerScreenWidgetGeometry are newly introduced to give you the geometry of the viewport and the player screen for widgets to help convert between spaces.")
 	UFUNCTION(BlueprintPure, BlueprintCosmetic, Category="Viewport")
-	static bool GetMousePositionScaledByDPI(APlayerController* Player, float& LocationX, float& LocationY);
-	static bool GetMousePositionScaledByDPI(APlayerController* Player, double& LocationX, double& LocationY);	// LWC_TODO: Temp stand in for native calls with FVector2D components.
+	static UMG_API bool GetMousePositionScaledByDPI(APlayerController* Player, float& LocationX, float& LocationY);
+	static UMG_API bool GetMousePositionScaledByDPI(APlayerController* Player, double& LocationX, double& LocationY);	// LWC_TODO: Temp stand in for native calls with FVector2D components.
 
 	/**
 	* Gets the slot object on the child widget as a Border Slot, allowing you to manipulate layout information.
 	* @param Widget The child widget of a border panel.
 	*/
 	UFUNCTION(BlueprintPure, Category = "Slot")
-	static class UBorderSlot* SlotAsBorderSlot(UWidget* Widget);
+	static UMG_API class UBorderSlot* SlotAsBorderSlot(UWidget* Widget);
 
 	/**
 	 * Gets the slot object on the child widget as a Canvas Slot, allowing you to manipulate layout information.
 	 * @param Widget The child widget of a canvas panel.
 	 */
 	UFUNCTION(BlueprintPure, Category="Slot")
-	static UCanvasPanelSlot* SlotAsCanvasSlot(UWidget* Widget);
+	static UMG_API UCanvasPanelSlot* SlotAsCanvasSlot(UWidget* Widget);
 
 	/**
 	 * Gets the slot object on the child widget as a Grid Slot, allowing you to manipulate layout information.
 	 * @param Widget The child widget of a grid panel.
 	 */
 	UFUNCTION(BlueprintPure, Category="Slot")
-	static UGridSlot* SlotAsGridSlot(UWidget* Widget);
+	static UMG_API UGridSlot* SlotAsGridSlot(UWidget* Widget);
 
 	/**
 	 * Gets the slot object on the child widget as a Horizontal Box Slot, allowing you to manipulate its information.
 	 * @param Widget The child widget of a Horizontal Box.
 	 */
 	UFUNCTION(BlueprintPure, Category="Slot")
-	static UHorizontalBoxSlot* SlotAsHorizontalBoxSlot(UWidget* Widget);
+	static UMG_API UHorizontalBoxSlot* SlotAsHorizontalBoxSlot(UWidget* Widget);
 
 	/**
 	 * Gets the slot object on the child widget as a Overlay Slot, allowing you to manipulate layout information.
 	 * @param Widget The child widget of a overlay panel.
 	 */
 	UFUNCTION(BlueprintPure, Category="Slot")
-	static UOverlaySlot* SlotAsOverlaySlot(UWidget* Widget);
+	static UMG_API UOverlaySlot* SlotAsOverlaySlot(UWidget* Widget);
 
 	/**
 	 * Gets the slot object on the child widget as a Uniform Grid Slot, allowing you to manipulate layout information.
 	 * @param Widget The child widget of a uniform grid panel.
 	 */
 	UFUNCTION(BlueprintPure, Category="Slot")
-	static UUniformGridSlot* SlotAsUniformGridSlot(UWidget* Widget);
+	static UMG_API UUniformGridSlot* SlotAsUniformGridSlot(UWidget* Widget);
 
 	/**
 	 * Gets the slot object on the child widget as a Vertical Box Slot, allowing you to manipulate its information.
 	 * @param Widget The child widget of a Vertical Box.
 	 */
 	UFUNCTION(BlueprintPure, Category = "Slot")
-	static UVerticalBoxSlot* SlotAsVerticalBoxSlot(UWidget* Widget);
+	static UMG_API UVerticalBoxSlot* SlotAsVerticalBoxSlot(UWidget* Widget);
 
 
 	/**
@@ -162,46 +162,46 @@ public:
 	* @param Widget The child widget of a Scroll Box.
 	*/
 	UFUNCTION(BlueprintPure, Category = "Slot")
-	static UScrollBoxSlot* SlotAsScrollBoxSlot(UWidget* Widget);
+	static UMG_API UScrollBoxSlot* SlotAsScrollBoxSlot(UWidget* Widget);
 
 	/**
 	 * Gets the slot object on the child widget as a Safe Box Slot, allowing you to manipulate its information.
 	 * @param Widget The child widget of a Safe Box.
 	 */
 	UFUNCTION(BlueprintPure, Category = "Slot")
-	static USafeZoneSlot* SlotAsSafeBoxSlot(UWidget* Widget);
+	static UMG_API USafeZoneSlot* SlotAsSafeBoxSlot(UWidget* Widget);
 
 	/**
 	 * Gets the slot object on the child widget as a Scale Box Slot, allowing you to manipulate its information.
 	 * @param Widget The child widget of a Scale Box.
 	 */
 	UFUNCTION(BlueprintPure, Category = "Slot")
-	static UScaleBoxSlot* SlotAsScaleBoxSlot(UWidget* Widget);
+	static UMG_API UScaleBoxSlot* SlotAsScaleBoxSlot(UWidget* Widget);
 
 	/**
 	 * Gets the slot object on the child widget as a Size Box Slot, allowing you to manipulate its information.
 	 * @param Widget The child widget of a Size Box.
 	 */
 	UFUNCTION(BlueprintPure, Category = "Slot")
-	static USizeBoxSlot* SlotAsSizeBoxSlot(UWidget* Widget);
+	static UMG_API USizeBoxSlot* SlotAsSizeBoxSlot(UWidget* Widget);
 
 	/**
 	 * Gets the slot object on the child widget as a Wrap Box Slot, allowing you to manipulate its information.
 	 * @param Widget The child widget of a Wrap Box.
 	 */
 	UFUNCTION(BlueprintPure, Category = "Slot")
-	static UWrapBoxSlot* SlotAsWrapBoxSlot(UWidget* Widget);
+	static UMG_API UWrapBoxSlot* SlotAsWrapBoxSlot(UWidget* Widget);
 
 	/**
 	 * Gets the slot object on the child widget as a Widget Switcher Slot, allowing you to manipulate its information.
 	 * @param Widget The child widget of a Widget Switcher Slot.
 	 */
 	UFUNCTION(BlueprintPure, Category = "Slot")
-	static UWidgetSwitcherSlot* SlotAsWidgetSwitcherSlot(UWidget* Widget);
+	static UMG_API UWidgetSwitcherSlot* SlotAsWidgetSwitcherSlot(UWidget* Widget);
 
 	/**
 	 * Removes all widgets from the viewport.
 	 */
 	UFUNCTION(BlueprintCallable, BlueprintCosmetic, Category="Viewport", meta=( WorldContext="WorldContextObject" ))
-	static void RemoveAllWidgets(UObject* WorldContextObject);
+	static UMG_API void RemoveAllWidgets(UObject* WorldContextObject);
 };
diff --git a/Engine/Source/Runtime/UMG/Public/Blueprint/WidgetNavigation.h b/Engine/Source/Runtime/UMG/Public/Blueprint/WidgetNavigation.h
--- a/Engine/Source/Runtime/UMG/Public/Blueprint/WidgetNavigation.h
+++ b/Engine/Source/Runtime/UMG/Public/Blueprint/WidgetNavigation.h
@@ -18,7 +18,7 @@ DECLARE_DYNAMIC_DELEGATE_RetVal_OneParam(UWidget*, FCustomWidgetNavigationDelega
  *
  */
 USTRUCT(BlueprintType)
-struct UMG_API FWidgetNavigationData
+struct FWidgetNavigationData
 {
 	GENERATED_USTRUCT_BODY()
 
@@ -36,18 +36,18 @@ public:
 	UPROPERTY()
 	FCustomWidgetNavigationDelegate CustomDelegate;
 
-	void Resolve(class UUserWidget* Outer, class UWidgetTree* WidgetTree);
+	UMG_API void Resolve(class UUserWidget* Outer, class UWidgetTree* WidgetTree);
 
 #if WITH_EDITOR
-	void TryToRenameBinding(FName OldName, FName NewName);
+	UMG_API void TryToRenameBinding(FName OldName, FName NewName);
 #endif
 };
 
 /**
  * 
  */
-UCLASS()
-class UMG_API UWidgetNavigation : public UObject
+UCLASS(MinimalAPI)
+class UWidgetNavigation : public UObject
 {
 	GENERATED_UCLASS_BODY()
 	
@@ -81,24 +81,24 @@ public:
 #if WITH_EDITOR
 
 	/**  */
-	FWidgetNavigationData& GetNavigationData(EUINavigation Nav);
+	UMG_API FWidgetNavigationData& GetNavigationData(EUINavigation Nav);
 
 	/**  */
-	EUINavigationRule GetNavigationRule(EUINavigation Nav);
+	UMG_API EUINavigationRule GetNavigationRule(EUINavigation Nav);
 
 	/** Try to rename any explicit or custom bindings from an old to a new name. This method can be overridden to customize resolving rules. */
-	virtual void TryToRenameBinding(FName OldName, FName NewName);
+	UMG_API virtual void TryToRenameBinding(FName OldName, FName NewName);
 
 #endif
 
 	/** Resolve widget names. This method can be overridden to customize resolving rules. */
-	virtual void ResolveRules(class UUserWidget* Outer, class UWidgetTree* WidgetTree);
+	UMG_API virtual void ResolveRules(class UUserWidget* Outer, class UWidgetTree* WidgetTree);
 
 	/** Updates a slate metadata object to match this configured navigation ruleset. */
-	void UpdateMetaData(TSharedRef<FNavigationMetaData> MetaData);
+	UMG_API void UpdateMetaData(TSharedRef<FNavigationMetaData> MetaData);
 
 	/** @return true if the configured navigation object is the same as an un-customized navigation rule set. */
-	bool IsDefaultNavigation() const;
+	UMG_API bool IsDefaultNavigation() const;
 
 private:
 
diff --git a/Engine/Source/Runtime/UMG/Public/Blueprint/WidgetTree.h b/Engine/Source/Runtime/UMG/Public/Blueprint/WidgetTree.h
--- a/Engine/Source/Runtime/UMG/Public/Blueprint/WidgetTree.h
+++ b/Engine/Source/Runtime/UMG/Public/Blueprint/WidgetTree.h
@@ -15,22 +15,22 @@
 #include "WidgetTree.generated.h"
 
 /** The widget tree manages the collection of widgets in a blueprint widget. */
-UCLASS()
-class UMG_API UWidgetTree : public UObject, public INamedSlotInterface
+UCLASS(MinimalAPI)
+class UWidgetTree : public UObject, public INamedSlotInterface
 {
 	GENERATED_UCLASS_BODY()
 
 public:
 
 	// Begin UObject
-	virtual UWorld* GetWorld() const override;
+	UMG_API virtual UWorld* GetWorld() const override;
 	// End UObject
 
 	/** Finds the widget in the tree by name. */
-	UWidget* FindWidget(const FName& Name) const;
+	UMG_API UWidget* FindWidget(const FName& Name) const;
 
 	/** Finds a widget in the tree using the native widget as the key. */
-	UWidget* FindWidget(TSharedRef<SWidget> InWidget) const;
+	UMG_API UWidget* FindWidget(TSharedRef<SWidget> InWidget) const;
 
 	/** Finds the widget in the tree by name and casts the return to the desired type. */
 	template <typename WidgetT>
@@ -40,51 +40,51 @@ public:
 	}
 
 	/** Removes the widget from the hierarchy and all sub widgets. */
-	bool RemoveWidget(UWidget* Widget);
+	UMG_API bool RemoveWidget(UWidget* Widget);
 
 	/** Gets the parent widget of a given widget, and potentially the child index. */
-	static class UPanelWidget* FindWidgetParent(UWidget* Widget, int32& OutChildIndex);
+	static UMG_API class UPanelWidget* FindWidgetParent(UWidget* Widget, int32& OutChildIndex);
 
 	/**
 	 * Searches recursively through the children of the given ParentWidget to find a child widget of the given name.
 	 * If successful, also gets the index the child ultimately occupies within the starting ParentWidget (INDEX_NONE otherwise)
 	 */
-	static UWidget* FindWidgetChild(UPanelWidget* ParentWidget, FName ChildWidgetName, int32& OutChildIndex);
+	static UMG_API UWidget* FindWidgetChild(UPanelWidget* ParentWidget, FName ChildWidgetName, int32& OutChildIndex);
 
 	/**
 	 * Determines the child index of the given ParentWidget that the given ChildWidget ultimately occupies, accounting for nesting
 	 * @return The child slot index within ParentWidget that ChildWidget ultimately occupies (INDEX_NONE if ChildWidget is not within ParentWidget at any level)
 	 */
-	static int32 FindChildIndex(const UPanelWidget* ParentWidget, const UWidget* ChildWidget);
+	static UMG_API int32 FindChildIndex(const UPanelWidget* ParentWidget, const UWidget* ChildWidget);
 
 	/** Gathers all the widgets in the tree recursively */
-	void GetAllWidgets(TArray<UWidget*>& Widgets) const;
+	UMG_API void GetAllWidgets(TArray<UWidget*>& Widgets) const;
 
 	/** Gathers descendant child widgets of a parent widget. */
-	static void GetChildWidgets(UWidget* Parent, TArray<UWidget*>& Widgets);
+	static UMG_API void GetChildWidgets(UWidget* Parent, TArray<UWidget*>& Widgets);
 
 	/** Attempts to move a constructed Widget to another tree. Returns true on a successful move. */
-	static bool TryMoveWidgetToNewTree(UWidget* Widget, UWidgetTree* DestinationTree);
+	static UMG_API bool TryMoveWidgetToNewTree(UWidget* Widget, UWidgetTree* DestinationTree);
 
 	/**
 	 * Iterates through all widgets including widgets contained in named slots, other than
 	 * investigating named slots, this code does not dive into foreign WidgetTrees, as would exist
 	 * inside another user widget.
 	 */
-	void ForEachWidget(TFunctionRef<void(UWidget*)> Predicate) const;
+	UMG_API void ForEachWidget(TFunctionRef<void(UWidget*)> Predicate) const;
 
 	/**
 	 * Iterates through all widgets including widgets contained in named slots, other than
 	 * investigating named slots.  This includes foreign widget trees inside of other UserWidgets.
 	 */
-	void ForEachWidgetAndDescendants(TFunctionRef<void(UWidget*)> Predicate) const;
+	UMG_API void ForEachWidgetAndDescendants(TFunctionRef<void(UWidget*)> Predicate) const;
 
 	/**
 	 * Iterates through all child widgets including widgets contained in named slots, other than
 	 * investigating named slots, this code does not dive into foreign WidgetTrees, as would exist
 	 * inside another user widget.
 	 */
-	static void ForWidgetAndChildren(UWidget* Widget, TFunctionRef<void(UWidget*)> Predicate);
+	static UMG_API void ForWidgetAndChildren(UWidget* Widget, TFunctionRef<void(UWidget*)> Predicate);
 
 	/** Constructs the widget, and adds it to the tree. */
 	template <typename WidgetT>
@@ -104,19 +104,19 @@ public:
 	//----------------------------------------------------------------------------------------
 	
 	/** Gets the names for slots that we can store widgets into. */
-	virtual void GetSlotNames(TArray<FName>& SlotNames) const override;
+	UMG_API virtual void GetSlotNames(TArray<FName>& SlotNames) const override;
 
 	/** Gets the widget for a given slot by name, will return nullptr if no widget is in the slot. */
-	virtual UWidget* GetContentForSlot(FName SlotName) const override;
+	UMG_API virtual UWidget* GetContentForSlot(FName SlotName) const override;
 
 	/** Sets the widget for a given slot by name. */
-	virtual void SetContentForSlot(FName SlotName, UWidget* Content) override;
+	UMG_API virtual void SetContentForSlot(FName SlotName, UWidget* Content) override;
 
 	//----------------------------------------------------------------------------------------
 
 	// UObject interface
-	virtual void PreSave(FObjectPreSaveContext ObjectSaveContext) override;
-	virtual void PostLoad() override;
+	UMG_API virtual void PreSave(FObjectPreSaveContext ObjectSaveContext) override;
+	UMG_API virtual void PostLoad() override;
 	// End of UObject interface
 
 public:
diff --git a/Engine/Source/Runtime/UMG/Public/Components/BackgroundBlur.h b/Engine/Source/Runtime/UMG/Public/Components/BackgroundBlur.h
--- a/Engine/Source/Runtime/UMG/Public/Components/BackgroundBlur.h
+++ b/Engine/Source/Runtime/UMG/Public/Components/BackgroundBlur.h
@@ -12,8 +12,8 @@
  * * Single Child
  * * Blur Effect
  */
-UCLASS()
-class UMG_API UBackgroundBlur : public UContentWidget
+UCLASS(MinimalAPI)
+class UBackgroundBlur : public UContentWidget
 {
 	GENERATED_UCLASS_BODY()
 
@@ -76,70 +76,70 @@ public:
 	FSlateBrush LowQualityFallbackBrush;
 
 public:
-	virtual void ReleaseSlateResources(bool bReleaseChildren) override;
+	UMG_API virtual void ReleaseSlateResources(bool bReleaseChildren) override;
 
 #if WITH_EDITOR
-	virtual void PostEditChangeProperty(struct FPropertyChangedEvent& PropertyChangedEvent) override;
-	virtual const FText GetPaletteCategory() override;
+	UMG_API virtual void PostEditChangeProperty(struct FPropertyChangedEvent& PropertyChangedEvent) override;
+	UMG_API virtual const FText GetPaletteCategory() override;
 #endif
 
 	UFUNCTION(BlueprintCallable, Category="Appearance")
-	void SetPadding(FMargin InPadding);
+	UMG_API void SetPadding(FMargin InPadding);
 
-	FMargin GetPadding() const;
+	UMG_API FMargin GetPadding() const;
 
 	UFUNCTION(BlueprintCallable, Category="Appearance")
-	void SetHorizontalAlignment(EHorizontalAlignment InHorizontalAlignment);
+	UMG_API void SetHorizontalAlignment(EHorizontalAlignment InHorizontalAlignment);
 
-	EHorizontalAlignment GetHorizontalAlignment() const;
+	UMG_API EHorizontalAlignment GetHorizontalAlignment() const;
 
 	UFUNCTION(BlueprintCallable, Category="Appearance")
-	void SetVerticalAlignment(EVerticalAlignment InVerticalAlignment);
+	UMG_API void SetVerticalAlignment(EVerticalAlignment InVerticalAlignment);
 
-	EVerticalAlignment GetVerticalAlignment() const;
+	UMG_API EVerticalAlignment GetVerticalAlignment() const;
 
 	UFUNCTION(BlueprintCallable, Category = "Appearance")
-	void SetApplyAlphaToBlur(bool bInApplyAlphaToBlur);
+	UMG_API void SetApplyAlphaToBlur(bool bInApplyAlphaToBlur);
 
-	bool GetApplyAlphaToBlur() const;
+	UMG_API bool GetApplyAlphaToBlur() const;
 
-	void SetOverrideAutoRadiusCalculation(bool InOverrideAutoRadiusCalculation);
+	UMG_API void SetOverrideAutoRadiusCalculation(bool InOverrideAutoRadiusCalculation);
 
-	bool GetOverrideAutoRadiusCalculation() const;
+	UMG_API bool GetOverrideAutoRadiusCalculation() const;
 
 	UFUNCTION(BlueprintCallable, Category = "Appearance")
-	void SetBlurRadius(int32 InBlurRadius);
+	UMG_API void SetBlurRadius(int32 InBlurRadius);
 
-	int32 GetBlurRadius() const;
+	UMG_API int32 GetBlurRadius() const;
 
 	UFUNCTION(BlueprintCallable, Category = "Appearance")
-	virtual void SetBlurStrength(float InStrength);
+	UMG_API virtual void SetBlurStrength(float InStrength);
 
-	float GetBlurStrength() const;
+	UMG_API float GetBlurStrength() const;
 
 	UFUNCTION(BlueprintCallable, Category = "Appearance")
-	virtual void SetCornerRadius(FVector4 InCornerRadius);
+	UMG_API virtual void SetCornerRadius(FVector4 InCornerRadius);
 	
-	FVector4 GetCornerRadius() const;
+	UMG_API FVector4 GetCornerRadius() const;
 
 	UFUNCTION(BlueprintCallable, Category = "Appearance")
-	void SetLowQualityFallbackBrush(const FSlateBrush& InBrush);
+	UMG_API void SetLowQualityFallbackBrush(const FSlateBrush& InBrush);
 
-	FSlateBrush GetLowQualityFallbackBrush() const;
+	UMG_API FSlateBrush GetLowQualityFallbackBrush() const;
 
 	/** UObject interface */
-	virtual void Serialize(FArchive& Ar) override;
-	virtual void PostLoad() override;
+	UMG_API virtual void Serialize(FArchive& Ar) override;
+	UMG_API virtual void PostLoad() override;
 
 protected:
 	/** UWidget interface */
-	virtual UClass* GetSlotClass() const override;
-	virtual TSharedRef<SWidget> RebuildWidget() override;
-	virtual void SynchronizeProperties() override;
+	UMG_API virtual UClass* GetSlotClass() const override;
+	UMG_API virtual TSharedRef<SWidget> RebuildWidget() override;
+	UMG_API virtual void SynchronizeProperties() override;
 
 	/** UPanelWidget interface */
-	virtual void OnSlotAdded(UPanelSlot* Slot) override;
-	virtual void OnSlotRemoved(UPanelSlot* Slot) override;
+	UMG_API virtual void OnSlotAdded(UPanelSlot* Slot) override;
+	UMG_API virtual void OnSlotRemoved(UPanelSlot* Slot) override;
 
 protected:
 	TSharedPtr<class SBackgroundBlur> MyBackgroundBlur;
diff --git a/Engine/Source/Runtime/UMG/Public/Components/BackgroundBlurSlot.h b/Engine/Source/Runtime/UMG/Public/Components/BackgroundBlurSlot.h
--- a/Engine/Source/Runtime/UMG/Public/Components/BackgroundBlurSlot.h
+++ b/Engine/Source/Runtime/UMG/Public/Components/BackgroundBlurSlot.h
@@ -15,27 +15,27 @@ class UBackgroundBlur;
 /**
  * The Slot for the UBackgroundBlurSlot, contains the widget displayed in a BackgroundBlur's single slot
  */
-UCLASS()
-class UMG_API UBackgroundBlurSlot : public UPanelSlot
+UCLASS(MinimalAPI)
+class UBackgroundBlurSlot : public UPanelSlot
 {
 	GENERATED_UCLASS_BODY()
 
 public:
 
-	FMargin GetPadding() const;
+	UMG_API FMargin GetPadding() const;
 
 	UFUNCTION(BlueprintCallable, Category="Layout|Background Blur Slot")
-	void SetPadding(FMargin InPadding);
+	UMG_API void SetPadding(FMargin InPadding);
 
-	EHorizontalAlignment GetHorizontalAlignment() const;
+	UMG_API EHorizontalAlignment GetHorizontalAlignment() const;
 
 	UFUNCTION(BlueprintCallable, Category="Layout|Background Blur Slot")
-	void SetHorizontalAlignment(EHorizontalAlignment InHorizontalAlignment);
+	UMG_API void SetHorizontalAlignment(EHorizontalAlignment InHorizontalAlignment);
 
-	EVerticalAlignment GetVerticalAlignment() const;
+	UMG_API EVerticalAlignment GetVerticalAlignment() const;
 
 	UFUNCTION(BlueprintCallable, Category="Layout|Background Blur Slot")
-	void SetVerticalAlignment(EVerticalAlignment InVerticalAlignment);
+	UMG_API void SetVerticalAlignment(EVerticalAlignment InVerticalAlignment);
 
 protected:
 	UE_DEPRECATED(5.1, "Direct access to Padding is deprecated. Please use the getter or setter.")
@@ -56,20 +56,20 @@ protected:
 public:
 
 	//~ Begin UPanelSlot interface
-	virtual void SynchronizeProperties() override;
+	UMG_API virtual void SynchronizeProperties() override;
 	//~ End of UPanelSlot interface
 
 	/** Builds the underlying slot for the slate BackgroundBlur. */
-	void BuildSlot(TSharedRef<SBackgroundBlur> InBackgroundBlur);
+	UMG_API void BuildSlot(TSharedRef<SBackgroundBlur> InBackgroundBlur);
 
-	virtual void ReleaseSlateResources(bool bReleaseChildren) override;
+	UMG_API virtual void ReleaseSlateResources(bool bReleaseChildren) override;
 
 public:
 
 #if WITH_EDITOR
 
 	//~ Begin UObject interface
-	virtual void PostEditChangeProperty(struct FPropertyChangedEvent& PropertyChangedEvent) override;
+	UMG_API virtual void PostEditChangeProperty(struct FPropertyChangedEvent& PropertyChangedEvent) override;
 	//~ End of UObject interface
 
 #endif
diff --git a/Engine/Source/Runtime/UMG/Public/Components/Border.h b/Engine/Source/Runtime/UMG/Public/Components/Border.h
--- a/Engine/Source/Runtime/UMG/Public/Components/Border.h
+++ b/Engine/Source/Runtime/UMG/Public/Components/Border.h
@@ -25,8 +25,8 @@ class UTexture2D;
  * * Single Child
  * * Image
  */
-UCLASS()
-class UMG_API UBorder : public UContentWidget
+UCLASS(MinimalAPI)
+class UBorder : public UContentWidget
 {
 	GENERATED_UCLASS_BODY()
 
@@ -107,64 +107,64 @@ public:
 
 public:
 	/**  */
-	FLinearColor GetContentColorAndOpacity() const;
+	UMG_API FLinearColor GetContentColorAndOpacity() const;
 
 	/**  */
 	UFUNCTION(BlueprintCallable, Category="Appearance")
-	void SetContentColorAndOpacity(FLinearColor InContentColorAndOpacity);
+	UMG_API void SetContentColorAndOpacity(FLinearColor InContentColorAndOpacity);
 
-	FMargin GetPadding() const;
+	UMG_API FMargin GetPadding() const;
 
 	UFUNCTION(BlueprintCallable, Category="Appearance")
-	void SetPadding(FMargin InPadding);
+	UMG_API void SetPadding(FMargin InPadding);
 
-	EHorizontalAlignment GetHorizontalAlignment() const;
+	UMG_API EHorizontalAlignment GetHorizontalAlignment() const;
 
 	UFUNCTION(BlueprintCallable, Category="Appearance")
-	void SetHorizontalAlignment(EHorizontalAlignment InHorizontalAlignment);
+	UMG_API void SetHorizontalAlignment(EHorizontalAlignment InHorizontalAlignment);
 
-	EVerticalAlignment GetVerticalAlignment() const;
+	UMG_API EVerticalAlignment GetVerticalAlignment() const;
 
 	UFUNCTION(BlueprintCallable, Category="Appearance")
-	void SetVerticalAlignment(EVerticalAlignment InVerticalAlignment);
+	UMG_API void SetVerticalAlignment(EVerticalAlignment InVerticalAlignment);
 
 	/**  */
-	FLinearColor GetBrushColor() const;
+	UMG_API FLinearColor GetBrushColor() const;
 
 	/**  */
 	UFUNCTION(BlueprintCallable, Category="Appearance")
-	void SetBrushColor(FLinearColor InBrushColor);
+	UMG_API void SetBrushColor(FLinearColor InBrushColor);
 
 	/**  */
 	UFUNCTION(BlueprintCallable, Category="Appearance")
-	void SetBrush(const FSlateBrush& InBrush);
+	UMG_API void SetBrush(const FSlateBrush& InBrush);
 
 	/**  */
 	UFUNCTION(BlueprintCallable, Category="Appearance")
-	void SetBrushFromAsset(USlateBrushAsset* Asset);
+	UMG_API void SetBrushFromAsset(USlateBrushAsset* Asset);
 
 	/**  */
 	UFUNCTION(BlueprintCallable, Category="Appearance")
-	void SetBrushFromTexture(UTexture2D* Texture);
+	UMG_API void SetBrushFromTexture(UTexture2D* Texture);
 
 	/**  */
 	UFUNCTION(BlueprintCallable, Category="Appearance")
-	void SetBrushFromMaterial(UMaterialInterface* Material);
+	UMG_API void SetBrushFromMaterial(UMaterialInterface* Material);
 
 	/**  */
-	bool GetShowEffectWhenDisabled() const;
+	UMG_API bool GetShowEffectWhenDisabled() const;
 
 	/**  */
 	UFUNCTION(BlueprintCallable, Category="Appearance")
-	void SetShowEffectWhenDisabled(bool bInShowEffectWhenDisabled);
+	UMG_API void SetShowEffectWhenDisabled(bool bInShowEffectWhenDisabled);
 
 	/**  */
 	UFUNCTION(BlueprintCallable, Category="Appearance")
-	UMaterialInstanceDynamic* GetDynamicMaterial();
+	UMG_API UMaterialInstanceDynamic* GetDynamicMaterial();
 
 public:
 	/** Gets the DesiredSizeScale of this border. */
-	FVector2D GetDesiredSizeScale() const;
+	UMG_API FVector2D GetDesiredSizeScale() const;
 
 	/**
 	 * Sets the DesiredSizeScale of this border.
@@ -172,36 +172,36 @@ public:
 	 * @param InScale    The X and Y multipliers for the desired size
 	 */
 	UFUNCTION(BlueprintCallable, Category = "Appearance")
-	void SetDesiredSizeScale(FVector2D InScale);
+	UMG_API void SetDesiredSizeScale(FVector2D InScale);
 
 public:
 
 	//~ Begin UWidget Interface
-	virtual void SynchronizeProperties() override;
+	UMG_API virtual void SynchronizeProperties() override;
 	//~ End UWidget Interface
 
 	//~ Begin UVisual Interface
-	virtual void ReleaseSlateResources(bool bReleaseChildren) override;
+	UMG_API virtual void ReleaseSlateResources(bool bReleaseChildren) override;
 	//~ End UVisual Interface
 
 	//~ Begin UObject Interface
-	virtual void PostLoad() override;
+	UMG_API virtual void PostLoad() override;
 	//~ End UObject Interface
 
 #if WITH_EDITOR
 	//~ Begin UObject Interface
-	virtual void PostEditChangeProperty(struct FPropertyChangedEvent& PropertyChangedEvent) override;
+	UMG_API virtual void PostEditChangeProperty(struct FPropertyChangedEvent& PropertyChangedEvent) override;
 	//~ End UObject Interface
 
-	virtual const FText GetPaletteCategory() override;
+	UMG_API virtual const FText GetPaletteCategory() override;
 #endif
 
 protected:
 
 	// UPanelWidget
-	virtual UClass* GetSlotClass() const override;
-	virtual void OnSlotAdded(UPanelSlot* Slot) override;
-	virtual void OnSlotRemoved(UPanelSlot* Slot) override;
+	UMG_API virtual UClass* GetSlotClass() const override;
+	UMG_API virtual void OnSlotAdded(UPanelSlot* Slot) override;
+	UMG_API virtual void OnSlotRemoved(UPanelSlot* Slot) override;
 	// End UPanelWidget
 
 protected:
@@ -209,16 +209,16 @@ protected:
 
 protected:
 	//~ Begin UWidget Interface
-	virtual TSharedRef<SWidget> RebuildWidget() override;
+	UMG_API virtual TSharedRef<SWidget> RebuildWidget() override;
 	//~ End UWidget Interface
 
-	FReply HandleMouseButtonDown(const FGeometry& Geometry, const FPointerEvent& MouseEvent);
-	FReply HandleMouseButtonUp(const FGeometry& Geometry, const FPointerEvent& MouseEvent);
-	FReply HandleMouseMove(const FGeometry& Geometry, const FPointerEvent& MouseEvent);
-	FReply HandleMouseDoubleClick(const FGeometry& Geometry, const FPointerEvent& MouseEvent);
+	UMG_API FReply HandleMouseButtonDown(const FGeometry& Geometry, const FPointerEvent& MouseEvent);
+	UMG_API FReply HandleMouseButtonUp(const FGeometry& Geometry, const FPointerEvent& MouseEvent);
+	UMG_API FReply HandleMouseMove(const FGeometry& Geometry, const FPointerEvent& MouseEvent);
+	UMG_API FReply HandleMouseDoubleClick(const FGeometry& Geometry, const FPointerEvent& MouseEvent);
 
 	/** Translates the bound brush data and assigns it to the cached brush used by this widget. */
-	const FSlateBrush* ConvertImage(TAttribute<FSlateBrush> InImageAsset) const;
+	UMG_API const FSlateBrush* ConvertImage(TAttribute<FSlateBrush> InImageAsset) const;
 
 	PROPERTY_BINDING_IMPLEMENTATION(FLinearColor, ContentColorAndOpacity)
 };
diff --git a/Engine/Source/Runtime/UMG/Public/Components/BorderSlot.h b/Engine/Source/Runtime/UMG/Public/Components/BorderSlot.h
--- a/Engine/Source/Runtime/UMG/Public/Components/BorderSlot.h
+++ b/Engine/Source/Runtime/UMG/Public/Components/BorderSlot.h
@@ -16,27 +16,27 @@ class UBorder;
 /**
  * The Slot for the UBorderSlot, contains the widget displayed in a border's single slot
  */
-UCLASS()
-class UMG_API UBorderSlot : public UPanelSlot
+UCLASS(MinimalAPI)
+class UBorderSlot : public UPanelSlot
 {
 	GENERATED_UCLASS_BODY()
 
 public:
 
-	FMargin GetPadding() const;
+	UMG_API FMargin GetPadding() const;
 
 	UFUNCTION(BlueprintCallable, Category="Layout|Border Slot")
-	void SetPadding(FMargin InPadding);
+	UMG_API void SetPadding(FMargin InPadding);
 
-	EHorizontalAlignment GetHorizontalAlignment() const;
+	UMG_API EHorizontalAlignment GetHorizontalAlignment() const;
 
 	UFUNCTION(BlueprintCallable, Category="Layout|Border Slot")
-	void SetHorizontalAlignment(EHorizontalAlignment InHorizontalAlignment);
+	UMG_API void SetHorizontalAlignment(EHorizontalAlignment InHorizontalAlignment);
 
-	EVerticalAlignment GetVerticalAlignment() const;
+	UMG_API EVerticalAlignment GetVerticalAlignment() const;
 
 	UFUNCTION(BlueprintCallable, Category="Layout|Border Slot")
-	void SetVerticalAlignment(EVerticalAlignment InVerticalAlignment);
+	UMG_API void SetVerticalAlignment(EVerticalAlignment InVerticalAlignment);
 
 protected:
 	UE_DEPRECATED(5.1, "Direct access to Padding is deprecated. Please use the getter or setter.")
@@ -57,20 +57,20 @@ protected:
 public:
 
 	//~ Begin UPanelSlot interface
-	virtual void SynchronizeProperties() override;
+	UMG_API virtual void SynchronizeProperties() override;
 	//~ End of UPanelSlot interface
 
 	/** Builds the underlying slot for the slate border. */
-	void BuildSlot(TSharedRef<SBorder> InBorder);
+	UMG_API void BuildSlot(TSharedRef<SBorder> InBorder);
 
-	virtual void ReleaseSlateResources(bool bReleaseChildren) override;
+	UMG_API virtual void ReleaseSlateResources(bool bReleaseChildren) override;
 
 public:
 
 #if WITH_EDITOR
 
 	//~ Begin UObject interface
-	virtual void PostEditChangeProperty(struct FPropertyChangedEvent& PropertyChangedEvent) override;
+	UMG_API virtual void PostEditChangeProperty(struct FPropertyChangedEvent& PropertyChangedEvent) override;
 	//~ End of UObject interface
 
 #endif
diff --git a/Engine/Source/Runtime/UMG/Public/Components/Button.h b/Engine/Source/Runtime/UMG/Public/Components/Button.h
--- a/Engine/Source/Runtime/UMG/Public/Components/Button.h
+++ b/Engine/Source/Runtime/UMG/Public/Components/Button.h
@@ -27,8 +27,8 @@ DECLARE_DYNAMIC_MULTICAST_DELEGATE(FOnButtonHoverEvent);
  * * Single Child
  * * Clickable
  */
-UCLASS()
-class UMG_API UButton : public UContentWidget
+UCLASS(MinimalAPI)
+class UButton : public UContentWidget
 {
 	GENERATED_UCLASS_BODY()
 
@@ -92,21 +92,21 @@ public:
 	
 	/** Sets the color multiplier for the button background */
 	UFUNCTION(BlueprintCallable, Category="Button|Appearance")
-	void SetStyle(const FButtonStyle& InStyle);
+	UMG_API void SetStyle(const FButtonStyle& InStyle);
 
-	const FButtonStyle& GetStyle() const;
+	UMG_API const FButtonStyle& GetStyle() const;
 
 	/** Sets the color multiplier for the button content */
 	UFUNCTION(BlueprintCallable, Category="Button|Appearance")
-	void SetColorAndOpacity(FLinearColor InColorAndOpacity);
+	UMG_API void SetColorAndOpacity(FLinearColor InColorAndOpacity);
 
-	FLinearColor GetColorAndOpacity() const;
+	UMG_API FLinearColor GetColorAndOpacity() const;
 
 	/** Sets the color multiplier for the button background */
 	UFUNCTION(BlueprintCallable, Category="Button|Appearance")
-	void SetBackgroundColor(FLinearColor InBackgroundColor);
+	UMG_API void SetBackgroundColor(FLinearColor InBackgroundColor);
 
-	FLinearColor GetBackgroundColor() const;
+	UMG_API FLinearColor GetBackgroundColor() const;
 
 	/**
 	 * Returns true if the user is actively pressing the button.  Do not use this for detecting 'Clicks', use the OnClicked event instead.
@@ -114,71 +114,71 @@ public:
 	 * @return true if the user is actively pressing the button otherwise false.
 	 */
 	UFUNCTION(BlueprintCallable, Category="Button")
-	bool IsPressed() const;
+	UMG_API bool IsPressed() const;
 
 	UFUNCTION(BlueprintCallable, Category="Button")
-	void SetClickMethod(EButtonClickMethod::Type InClickMethod);
+	UMG_API void SetClickMethod(EButtonClickMethod::Type InClickMethod);
 
-	EButtonClickMethod::Type GetClickMethod() const;
+	UMG_API EButtonClickMethod::Type GetClickMethod() const;
 
 	UFUNCTION(BlueprintCallable, Category="Button")
-	void SetTouchMethod(EButtonTouchMethod::Type InTouchMethod);
+	UMG_API void SetTouchMethod(EButtonTouchMethod::Type InTouchMethod);
 
-	EButtonTouchMethod::Type GetTouchMethod() const;
+	UMG_API EButtonTouchMethod::Type GetTouchMethod() const;
 
 	UFUNCTION(BlueprintCallable, Category="Button")
-	void SetPressMethod(EButtonPressMethod::Type InPressMethod);
+	UMG_API void SetPressMethod(EButtonPressMethod::Type InPressMethod);
 
-	EButtonPressMethod::Type GetPressMethod() const;
+	UMG_API EButtonPressMethod::Type GetPressMethod() const;
 
-	bool GetIsFocusable() const;
+	UMG_API bool GetIsFocusable() const;
 
 public:
 
 	//~ Begin UWidget Interface
-	virtual void SynchronizeProperties() override;
+	UMG_API virtual void SynchronizeProperties() override;
 	//~ End UWidget Interface
 
 	//~ Begin UVisual Interface
-	virtual void ReleaseSlateResources(bool bReleaseChildren) override;
+	UMG_API virtual void ReleaseSlateResources(bool bReleaseChildren) override;
 	//~ End UVisual Interface
 
 	//~ Begin UObject Interface
-	virtual void PostLoad() override;
+	UMG_API virtual void PostLoad() override;
 	//~ End UObject Interface
 
 #if WITH_EDITOR
-	virtual const FText GetPaletteCategory() override;
+	UMG_API virtual const FText GetPaletteCategory() override;
 #endif
 
 protected:
 
 	// UPanelWidget
-	virtual UClass* GetSlotClass() const override;
-	virtual void OnSlotAdded(UPanelSlot* Slot) override;
-	virtual void OnSlotRemoved(UPanelSlot* Slot) override;
+	UMG_API virtual UClass* GetSlotClass() const override;
+	UMG_API virtual void OnSlotAdded(UPanelSlot* Slot) override;
+	UMG_API virtual void OnSlotRemoved(UPanelSlot* Slot) override;
 	// End UPanelWidget
 
 protected:
 	/** Handle the actual click event from slate and forward it on */
-	FReply SlateHandleClicked();
-	void SlateHandlePressed();
-	void SlateHandleReleased();
-	void SlateHandleHovered();
-	void SlateHandleUnhovered();
+	UMG_API FReply SlateHandleClicked();
+	UMG_API void SlateHandlePressed();
+	UMG_API void SlateHandleReleased();
+	UMG_API void SlateHandleHovered();
+	UMG_API void SlateHandleUnhovered();
 
 	// Initialize IsFocusable in the constructor before the SWidget is constructed.
-	void InitIsFocusable(bool InIsFocusable);
+	UMG_API void InitIsFocusable(bool InIsFocusable);
 protected:
 	//~ Begin UWidget Interface
-	virtual TSharedRef<SWidget> RebuildWidget() override;
+	UMG_API virtual TSharedRef<SWidget> RebuildWidget() override;
 #if WITH_EDITOR
 	virtual TSharedRef<SWidget> RebuildDesignWidget(TSharedRef<SWidget> Content) override { return Content; }
 #endif
 	//~ End UWidget Interface
 
 #if WITH_ACCESSIBILITY
-	virtual TSharedPtr<SWidget> GetAccessibleWidget() const override;
+	UMG_API virtual TSharedPtr<SWidget> GetAccessibleWidget() const override;
 #endif
 
 protected:
diff --git a/Engine/Source/Runtime/UMG/Public/Components/ButtonSlot.h b/Engine/Source/Runtime/UMG/Public/Components/ButtonSlot.h
--- a/Engine/Source/Runtime/UMG/Public/Components/ButtonSlot.h
+++ b/Engine/Source/Runtime/UMG/Public/Components/ButtonSlot.h
@@ -13,8 +13,8 @@
 class SButton;
 
 /** The Slot for the UButtonSlot, contains the widget displayed in a button's single slot */
-UCLASS()
-class UMG_API UButtonSlot : public UPanelSlot
+UCLASS(MinimalAPI)
+class UButtonSlot : public UPanelSlot
 {
 	GENERATED_UCLASS_BODY()
 
@@ -37,29 +37,29 @@ public:
 
 public:
 
-	FMargin GetPadding() const;
+	UMG_API FMargin GetPadding() const;
 
 	UFUNCTION(BlueprintCallable, Category="Layout|Button Slot")
-	void SetPadding(FMargin InPadding);
+	UMG_API void SetPadding(FMargin InPadding);
 
-	EHorizontalAlignment GetHorizontalAlignment() const;
+	UMG_API EHorizontalAlignment GetHorizontalAlignment() const;
 
 	UFUNCTION(BlueprintCallable, Category="Layout|Button Slot")
-	void SetHorizontalAlignment(EHorizontalAlignment InHorizontalAlignment);
+	UMG_API void SetHorizontalAlignment(EHorizontalAlignment InHorizontalAlignment);
 
-	EVerticalAlignment GetVerticalAlignment() const;
+	UMG_API EVerticalAlignment GetVerticalAlignment() const;
 
 	UFUNCTION(BlueprintCallable, Category="Layout|Button Slot")
-	void SetVerticalAlignment(EVerticalAlignment InVerticalAlignment);
+	UMG_API void SetVerticalAlignment(EVerticalAlignment InVerticalAlignment);
 
 	// UPanelSlot interface
-	virtual void SynchronizeProperties() override;
+	UMG_API virtual void SynchronizeProperties() override;
 	// End of UPanelSlot interface
 
 	/** Builds the underlying slot for the slate button. */
-	void BuildSlot(TSharedRef<SButton> InButton);
+	UMG_API void BuildSlot(TSharedRef<SButton> InButton);
 
-	virtual void ReleaseSlateResources(bool bReleaseChildren) override;
+	UMG_API virtual void ReleaseSlateResources(bool bReleaseChildren) override;
 
 private:
 
diff --git a/Engine/Source/Runtime/UMG/Public/Components/CanvasPanel.h b/Engine/Source/Runtime/UMG/Public/Components/CanvasPanel.h
--- a/Engine/Source/Runtime/UMG/Public/Components/CanvasPanel.h
+++ b/Engine/Source/Runtime/UMG/Public/Components/CanvasPanel.h
@@ -21,8 +21,8 @@ class UCanvasPanelSlot;
  * * Absolute Layout
  * * Anchors
  */
-UCLASS(meta = (ShortTooltip = "A designer-friendly panel useful for laying out top-level widgets. Use sparingly."))
-class UMG_API UCanvasPanel : public UPanelWidget
+UCLASS(meta = (ShortTooltip = "A designer-friendly panel useful for laying out top-level widgets. Use sparingly."), MinimalAPI)
+class UCanvasPanel : public UPanelWidget
 {
 	GENERATED_UCLASS_BODY()
 
@@ -30,22 +30,22 @@ public:
 
 	/**  */
 	UFUNCTION(BlueprintCallable, Category="Canvas Panel")
-	UCanvasPanelSlot* AddChildToCanvas(UWidget* Content);
+	UMG_API UCanvasPanelSlot* AddChildToCanvas(UWidget* Content);
 
 	/** Gets the underlying native canvas widget if it has been constructed */
-	TSharedPtr<class SConstraintCanvas> GetCanvasWidget() const;
+	UMG_API TSharedPtr<class SConstraintCanvas> GetCanvasWidget() const;
 
 	/** Computes the geometry for a particular slot based on the current geometry of the canvas. */
-	bool GetGeometryForSlot(int32 SlotIndex, FGeometry& ArrangedGeometry) const;
+	UMG_API bool GetGeometryForSlot(int32 SlotIndex, FGeometry& ArrangedGeometry) const;
 
 	/** Computes the geometry for a particular slot based on the current geometry of the canvas. */
-	bool GetGeometryForSlot(UCanvasPanelSlot* Slot, FGeometry& ArrangedGeometry) const;
+	UMG_API bool GetGeometryForSlot(UCanvasPanelSlot* Slot, FGeometry& ArrangedGeometry) const;
 
-	void ReleaseSlateResources(bool bReleaseChildren) override;
+	UMG_API void ReleaseSlateResources(bool bReleaseChildren) override;
 
 #if WITH_EDITOR
 	// UWidget interface
-	virtual const FText GetPaletteCategory() override;
+	UMG_API virtual const FText GetPaletteCategory() override;
 	// End UWidget interface
 
 	// UWidget interface
@@ -59,9 +59,9 @@ public:
 protected:
 
 	// UPanelWidget
-	virtual UClass* GetSlotClass() const override;
-	virtual void OnSlotAdded(UPanelSlot* Slot) override;
-	virtual void OnSlotRemoved(UPanelSlot* Slot) override;
+	UMG_API virtual UClass* GetSlotClass() const override;
+	UMG_API virtual void OnSlotAdded(UPanelSlot* Slot) override;
+	UMG_API virtual void OnSlotRemoved(UPanelSlot* Slot) override;
 	// End UPanelWidget
 
 protected:
@@ -70,6 +70,6 @@ protected:
 
 protected:
 	// UWidget interface
-	virtual TSharedRef<SWidget> RebuildWidget() override;
+	UMG_API virtual TSharedRef<SWidget> RebuildWidget() override;
 	// End of UWidget interface
 };
diff --git a/Engine/Source/Runtime/UMG/Public/Components/CanvasPanelSlot.h b/Engine/Source/Runtime/UMG/Public/Components/CanvasPanelSlot.h
--- a/Engine/Source/Runtime/UMG/Public/Components/CanvasPanelSlot.h
+++ b/Engine/Source/Runtime/UMG/Public/Components/CanvasPanelSlot.h
@@ -62,8 +62,8 @@ public:
 	}
 };
 
-UCLASS()
-class UMG_API UCanvasPanelSlot : public UPanelSlot
+UCLASS(MinimalAPI)
+class UCanvasPanelSlot : public UPanelSlot
 {
 	GENERATED_UCLASS_BODY()
 
@@ -86,109 +86,109 @@ public:
 
 public:
 #if WITH_EDITOR
-	virtual bool NudgeByDesigner(const FVector2D& NudgeDirection, const TOptional<int32>& GridSnapSize) override;
-	virtual bool DragDropPreviewByDesigner(const FVector2D& LocalCursorPosition, const TOptional<int32>& XGridSnapSize, const TOptional<int32>& YGridSnapSize) override;
-	virtual void SynchronizeFromTemplate(const UPanelSlot* const TemplateSlot) override;
+	UMG_API virtual bool NudgeByDesigner(const FVector2D& NudgeDirection, const TOptional<int32>& GridSnapSize) override;
+	UMG_API virtual bool DragDropPreviewByDesigner(const FVector2D& LocalCursorPosition, const TOptional<int32>& XGridSnapSize, const TOptional<int32>& YGridSnapSize) override;
+	UMG_API virtual void SynchronizeFromTemplate(const UPanelSlot* const TemplateSlot) override;
 #endif //WITH_EDITOR
 
 	/** Sets the layout data of the slot */
 	UFUNCTION(BlueprintCallable, Category="Layout|Canvas Slot")
-	void SetLayout(const FAnchorData& InLayoutData);
+	UMG_API void SetLayout(const FAnchorData& InLayoutData);
 
 	/** Gets the layout data of the slot */
 	UFUNCTION(BlueprintCallable, Category="Layout|Canvas Slot")
-	FAnchorData GetLayout() const;
+	UMG_API FAnchorData GetLayout() const;
 
 	/** Sets the position of the slot */
 	UFUNCTION(BlueprintCallable, Category="Layout|Canvas Slot")
-	void SetPosition(FVector2D InPosition);
+	UMG_API void SetPosition(FVector2D InPosition);
 
 	/** Gets the position of the slot */
 	UFUNCTION(BlueprintCallable, Category="Layout|Canvas Slot")
-	FVector2D GetPosition() const;
+	UMG_API FVector2D GetPosition() const;
 
 	/** Sets the size of the slot */
 	UFUNCTION(BlueprintCallable, Category="Layout|Canvas Slot")
-	void SetSize(FVector2D InSize);
+	UMG_API void SetSize(FVector2D InSize);
 
 	/** Gets the size of the slot */
 	UFUNCTION(BlueprintCallable, Category="Layout|Canvas Slot")
-	FVector2D GetSize() const;
+	UMG_API FVector2D GetSize() const;
 
 	/** Sets the offset data of the slot, which could be position and size, or margins depending on the anchor points */
 	UFUNCTION(BlueprintCallable, Category="Layout|Canvas Slot")
-	void SetOffsets(FMargin InOffset);
+	UMG_API void SetOffsets(FMargin InOffset);
 
 	/** Gets the offset data of the slot, which could be position and size, or margins depending on the anchor points */
 	UFUNCTION(BlueprintCallable, Category="Layout|Canvas Slot")
-	FMargin GetOffsets() const;
+	UMG_API FMargin GetOffsets() const;
 	
 	/** Sets the anchors on the slot */
 	UFUNCTION(BlueprintCallable, Category="Layout|Canvas Slot")
-	void SetAnchors(FAnchors InAnchors);
+	UMG_API void SetAnchors(FAnchors InAnchors);
 
 	/** Gets the anchors on the slot */
 	UFUNCTION(BlueprintCallable, Category="Layout|Canvas Slot")
-	FAnchors GetAnchors() const;
+	UMG_API FAnchors GetAnchors() const;
 
 	/** Sets the alignment on the slot */
 	UFUNCTION(BlueprintCallable, Category="Layout|Canvas Slot")
-	void SetAlignment(FVector2D InAlignment);
+	UMG_API void SetAlignment(FVector2D InAlignment);
 
 	/** Gets the alignment on the slot */
 	UFUNCTION(BlueprintCallable, Category="Layout|Canvas Slot")
-	FVector2D GetAlignment() const;
+	UMG_API FVector2D GetAlignment() const;
 
 	/** Sets if the slot to be auto-sized */
 	UFUNCTION(BlueprintCallable, Category="Layout|Canvas Slot")
-	void SetAutoSize(bool InbAutoSize);
+	UMG_API void SetAutoSize(bool InbAutoSize);
 
 	/** Gets if the slot to be auto-sized */
 	UFUNCTION(BlueprintCallable, Category="Layout|Canvas Slot")
-	bool GetAutoSize() const;
+	UMG_API bool GetAutoSize() const;
 
 	/** Sets the z-order on the slot */
 	UFUNCTION(BlueprintCallable, Category="Layout|Canvas Slot")
-	void SetZOrder(int32 InZOrder);
+	UMG_API void SetZOrder(int32 InZOrder);
 
 	/** Gets the z-order on the slot */
 	UFUNCTION(BlueprintCallable, Category="Layout|Canvas Slot")
-	int32 GetZOrder() const;
+	UMG_API int32 GetZOrder() const;
 
 public:
 
 	/** Sets the anchors on the slot */
 	UFUNCTION()
-	void SetMinimum(FVector2D InMinimumAnchors);
+	UMG_API void SetMinimum(FVector2D InMinimumAnchors);
 
 	/** Sets the anchors on the slot */
 	UFUNCTION()
-	void SetMaximum(FVector2D InMaximumAnchors);
+	UMG_API void SetMaximum(FVector2D InMaximumAnchors);
 
 public:
 
-	void BuildSlot(TSharedRef<SConstraintCanvas> Canvas);
+	UMG_API void BuildSlot(TSharedRef<SConstraintCanvas> Canvas);
 
 	// UPanelSlot interface
-	virtual void SynchronizeProperties() override;
+	UMG_API virtual void SynchronizeProperties() override;
 	// End of UPanelSlot interface
 
-	virtual void ReleaseSlateResources(bool bReleaseChildren) override;
+	UMG_API virtual void ReleaseSlateResources(bool bReleaseChildren) override;
 
 #if WITH_EDITOR
 	// UObject interface
 	using Super::PreEditChange;
-	virtual void PreEditChange(FProperty* PropertyThatWillChange) override;
-	virtual void PostEditChangeChainProperty(struct FPropertyChangedChainEvent& PropertyChangedEvent) override;
+	UMG_API virtual void PreEditChange(FProperty* PropertyThatWillChange) override;
+	UMG_API virtual void PostEditChangeChainProperty(struct FPropertyChangedChainEvent& PropertyChangedEvent) override;
 	// End of UObject interface
 
 	/** Stores the current layout information about the slot and parent canvas. */
-	void SaveBaseLayout();
+	UMG_API void SaveBaseLayout();
 
-	void SetDesiredPosition(FVector2D InPosition);
+	UMG_API void SetDesiredPosition(FVector2D InPosition);
 
 	/** Compares the saved base layout against the current state.  Updates the necessary properties to maintain a stable position. */
-	void RebaseLayout(bool PreserveSize = true);
+	UMG_API void RebaseLayout(bool PreserveSize = true);
 #endif
 
 private:
diff --git a/Engine/Source/Runtime/UMG/Public/Components/CircularThrobber.h b/Engine/Source/Runtime/UMG/Public/Components/CircularThrobber.h
--- a/Engine/Source/Runtime/UMG/Public/Components/CircularThrobber.h
+++ b/Engine/Source/Runtime/UMG/Public/Components/CircularThrobber.h
@@ -18,8 +18,8 @@ class USlateBrushAsset;
  * * No Children
  * * Spinner Progress
  */
-UCLASS()
-class UMG_API UCircularThrobber : public UWidget
+UCLASS(MinimalAPI)
+class UCircularThrobber : public UWidget
 {
 	GENERATED_UCLASS_BODY()
 
@@ -49,34 +49,34 @@ public:
 
 	/** Sets how many pieces there are. */
 	UFUNCTION(BlueprintCallable, Category = "Appearance")
-	void SetNumberOfPieces(int32 InNumberOfPieces);
-	int32 GetNumberOfPieces() const;
+	UMG_API void SetNumberOfPieces(int32 InNumberOfPieces);
+	UMG_API int32 GetNumberOfPieces() const;
 
 	/** Sets the amount of time for a full circle (in seconds). */
 	UFUNCTION(BlueprintCallable, Category = "Appearance")
-	void SetPeriod(float InPeriod);
-	float GetPeriod() const;
+	UMG_API void SetPeriod(float InPeriod);
+	UMG_API float GetPeriod() const;
 
 	/** Sets the radius of the circle. */
 	UFUNCTION(BlueprintCallable, Category = "Appearance")
-	void SetRadius(float InRadius);
-	float GetRadius() const;
+	UMG_API void SetRadius(float InRadius);
+	UMG_API float GetRadius() const;
 
 	/** Sets the throbber image. */
-	void SetImage(const FSlateBrush& InRadius);
-	const FSlateBrush& GetImage() const;
+	UMG_API void SetImage(const FSlateBrush& InRadius);
+	UMG_API const FSlateBrush& GetImage() const;
 
 	//~ Begin UWidget Interface
 public:
-	virtual void SynchronizeProperties() override;
-	virtual void ReleaseSlateResources(bool bReleaseChildren) override;
+	UMG_API virtual void SynchronizeProperties() override;
+	UMG_API virtual void ReleaseSlateResources(bool bReleaseChildren) override;
 
 #if WITH_EDITOR
-	virtual const FText GetPaletteCategory() override;
+	UMG_API virtual const FText GetPaletteCategory() override;
 #endif
 
 protected:
-	virtual TSharedRef<SWidget> RebuildWidget() override;
+	UMG_API virtual TSharedRef<SWidget> RebuildWidget() override;
 	//~ End UWidget Interface
 
 private:
diff --git a/Engine/Source/Runtime/UMG/Public/Components/ComboBox.h b/Engine/Source/Runtime/UMG/Public/Components/ComboBox.h
--- a/Engine/Source/Runtime/UMG/Public/Components/ComboBox.h
+++ b/Engine/Source/Runtime/UMG/Public/Components/ComboBox.h
@@ -12,8 +12,8 @@
 /**
  * The combobox allows you to display a list of options to the user in a dropdown menu for them to select one.
  */
-UCLASS(Experimental, meta=( DisplayName="ComboBox (Object)" ))
-class UMG_API UComboBox : public UWidget
+UCLASS(Experimental, meta=( DisplayName="ComboBox (Object)" ), MinimalAPI)
+class UComboBox : public UWidget
 {
 	GENERATED_UCLASS_BODY()
 
@@ -30,20 +30,20 @@ public:
 	UPROPERTY(EditAnywhere, Category=Events, meta=( IsBindableEvent="True" ))
 	FGenerateWidgetForObject OnGenerateWidgetEvent;
 
-	virtual void ReleaseSlateResources(bool bReleaseChildren) override;
+	UMG_API virtual void ReleaseSlateResources(bool bReleaseChildren) override;
 
 #if WITH_EDITOR
-	virtual const FText GetPaletteCategory() override;
+	UMG_API virtual const FText GetPaletteCategory() override;
 #endif
 
 	UPROPERTY(EditAnywhere, BlueprintReadOnly, Category=Interaction)
 	bool bIsFocusable;
 
 protected:
-	TSharedRef<SWidget> HandleGenerateWidget(UObject* Item) const;
+	UMG_API TSharedRef<SWidget> HandleGenerateWidget(UObject* Item) const;
 
 	// UWidget interface
-	virtual TSharedRef<SWidget> RebuildWidget() override;
+	UMG_API virtual TSharedRef<SWidget> RebuildWidget() override;
 	// End of UWidget interface
 
 protected:
diff --git a/Engine/Source/Runtime/UMG/Public/Components/ComboBoxKey.h b/Engine/Source/Runtime/UMG/Public/Components/ComboBoxKey.h
--- a/Engine/Source/Runtime/UMG/Public/Components/ComboBoxKey.h
+++ b/Engine/Source/Runtime/UMG/Public/Components/ComboBoxKey.h
@@ -14,13 +14,13 @@
  * The combobox allows you to display a list of options to the user in a dropdown menu for them to select one.
  * Use OnGenerateConentWidgetEvent to return a custom built widget.
  */
-UCLASS(meta = (DisplayName = "ComboBox (Key)"))
-class UMG_API UComboBoxKey : public UWidget
+UCLASS(meta = (DisplayName = "ComboBox (Key)"), MinimalAPI)
+class UComboBoxKey : public UWidget
 {
 	GENERATED_BODY()
 
 public:
-	UComboBoxKey();
+	UMG_API UComboBoxKey();
 
 	DECLARE_DYNAMIC_MULTICAST_DELEGATE_TwoParams(FOnSelectionChangedEvent, FName, SelectedItem, ESelectInfo::Type, SelectionType);
 	DECLARE_DYNAMIC_MULTICAST_DELEGATE(FOnOpeningEvent);
@@ -110,95 +110,95 @@ public:
 
 	/** Add an element to the option list. */
 	UFUNCTION(BlueprintCallable, Category = "ComboBox")
-	void AddOption(FName Option);
+	UMG_API void AddOption(FName Option);
 
 	/** Remove an element to the option list. */
 	UFUNCTION(BlueprintCallable, Category = "ComboBox")
-	bool RemoveOption(FName Option);
+	UMG_API bool RemoveOption(FName Option);
 
 	/** Remove all the elements of the option list. */
 	UFUNCTION(BlueprintCallable, Category = "ComboBox")
-	void ClearOptions();
+	UMG_API void ClearOptions();
 
 	/** Clear the current selection. */
 	UFUNCTION(BlueprintCallable, Category = "ComboBox")
-	void ClearSelection();
+	UMG_API void ClearSelection();
 
 	/** Set the current selected option. */
 	UFUNCTION(BlueprintCallable, Category = "ComboBox")
-	void SetSelectedOption(FName Option);
+	UMG_API void SetSelectedOption(FName Option);
 
 	/** Get the current selected option */
 	UFUNCTION(BlueprintCallable, Category = "ComboBox")
-	FName GetSelectedOption() const;
+	UMG_API FName GetSelectedOption() const;
 
 	/** Is the combobox menu opened. */
 	UFUNCTION(BlueprintCallable, Category = "ComboBox", Meta = (ReturnDisplayName = "bOpen"))
-	bool IsOpen() const;
+	UMG_API bool IsOpen() const;
 
 	//~ Begin UVisual Interface
-	virtual void ReleaseSlateResources(bool bReleaseChildren) override;
+	UMG_API virtual void ReleaseSlateResources(bool bReleaseChildren) override;
 	//~ End UVisual Interface
 
 	/** Set the padding for content. */
-	void SetContentPadding(FMargin InPadding);
+	UMG_API void SetContentPadding(FMargin InPadding);
 
 	/** Get the padding for content. */
-	FMargin GetContentPadding() const;
+	UMG_API FMargin GetContentPadding() const;
 
 	/** Is the combobox navigated by gamepad. */
-	bool IsEnableGamepadNavigationMode() const;
+	UMG_API bool IsEnableGamepadNavigationMode() const;
 
 	/** Set whether the combobox is navigated by gamepad. */
-	void SetEnableGamepadNavigationMode(bool InEnableGamepadNavigationMode);
+	UMG_API void SetEnableGamepadNavigationMode(bool InEnableGamepadNavigationMode);
 
 	/** Is the combobox arrow showing. */
-	bool IsHasDownArrow() const;
+	UMG_API bool IsHasDownArrow() const;
 
 	/** Set whether the combobox arrow is showing. */
-	void SetHasDownArrow(bool InHasDownArrow);
+	UMG_API void SetHasDownArrow(bool InHasDownArrow);
 
 	/** Get the maximum height of the combobox list. */
-	float GetMaxListHeight() const;
+	UMG_API float GetMaxListHeight() const;
 
 	/** Set the maximum height of the combobox list. */
-	void SetMaxListHeight(float InMaxHeight);
+	UMG_API void SetMaxListHeight(float InMaxHeight);
 
 	/** Get the style of the combobox. */
-	const FComboBoxStyle& GetWidgetStyle() const;
+	UMG_API const FComboBoxStyle& GetWidgetStyle() const;
 
 	/** Set the style of the combobox. */
-	void SetWidgetStyle(const FComboBoxStyle& InWidgetStyle);
+	UMG_API void SetWidgetStyle(const FComboBoxStyle& InWidgetStyle);
 
 	/** Get the style of the items. */
-	const FTableRowStyle& GetItemStyle() const;
+	UMG_API const FTableRowStyle& GetItemStyle() const;
 
 	/** Set the style of the items. */
-	void SetItemStyle(const FTableRowStyle& InItemStyle);
+	UMG_API void SetItemStyle(const FTableRowStyle& InItemStyle);
 
 	/** Get the style of the scrollbar. */
-	const FScrollBarStyle& GetScrollBarStyle() const;
+	UMG_API const FScrollBarStyle& GetScrollBarStyle() const;
 
 	/** Is the combobox focusable. */
-	bool IsFocusable() const;
+	UMG_API bool IsFocusable() const;
 
 	/** Get the foreground color of the button. */
-	FSlateColor GetForegroundColor() const;
+	UMG_API FSlateColor GetForegroundColor() const;
 
 #if WITH_EDITOR
-	virtual const FText GetPaletteCategory() override;
+	UMG_API virtual const FText GetPaletteCategory() override;
 #endif
 
 protected:
 
 	/** Initialize the scrollbar style in the constructor before the SWidget is constructed. */
-	void InitScrollBarStyle(const FScrollBarStyle& InScrollBarStyle);
+	UMG_API void InitScrollBarStyle(const FScrollBarStyle& InScrollBarStyle);
 
 	/** Initialize IsFocusable in the constructor before the SWidget is constructed. */
-	void InitIsFocusable(bool InIsFocusable);
+	UMG_API void InitIsFocusable(bool InIsFocusable);
 
 	/** Initialize ForegroundColor in the constructor before the SWidget is constructed. */
-	void InitForegroundColor(FSlateColor InForegroundColor);
+	UMG_API void InitForegroundColor(FSlateColor InForegroundColor);
 
 private:
 	/** Called by slate when it needs to generate the widget in the content box */
@@ -214,7 +214,7 @@ private:
 	void HandleOpening();
 
 	//~ Begin UWidget Interface
-	virtual TSharedRef<SWidget> RebuildWidget() override;
+	UMG_API virtual TSharedRef<SWidget> RebuildWidget() override;
 	//~ End UWidget Interface
 
 private:
diff --git a/Engine/Source/Runtime/UMG/Public/Components/ComboBoxString.h b/Engine/Source/Runtime/UMG/Public/Components/ComboBoxString.h
--- a/Engine/Source/Runtime/UMG/Public/Components/ComboBoxString.h
+++ b/Engine/Source/Runtime/UMG/Public/Components/ComboBoxString.h
@@ -17,8 +17,8 @@
 /**
  * The combobox allows you to display a list of options to the user in a dropdown menu for them to select one.
  */
-UCLASS(meta=( DisplayName="ComboBox (String)"))
-class UMG_API UComboBoxString : public UWidget
+UCLASS(meta=( DisplayName="ComboBox (String)"), MinimalAPI)
+class UComboBoxString : public UWidget
 {
 	GENERATED_UCLASS_BODY()
 
@@ -111,139 +111,139 @@ public: // Events
 public:
 
 	UFUNCTION(BlueprintCallable, Category="ComboBox")
-	void AddOption(const FString& Option);
+	UMG_API void AddOption(const FString& Option);
 
 	UFUNCTION(BlueprintCallable, Category="ComboBox")
-	bool RemoveOption(const FString& Option);
+	UMG_API bool RemoveOption(const FString& Option);
 
 	UFUNCTION(BlueprintCallable, Category="ComboBox")
-	int32 FindOptionIndex(const FString& Option) const;
+	UMG_API int32 FindOptionIndex(const FString& Option) const;
 
 	UFUNCTION(BlueprintCallable, Category="ComboBox")
-	FString GetOptionAtIndex(int32 Index) const;
+	UMG_API FString GetOptionAtIndex(int32 Index) const;
 
 	UFUNCTION(BlueprintCallable, Category="ComboBox")
-	void ClearOptions();
+	UMG_API void ClearOptions();
 
 	UFUNCTION(BlueprintCallable, Category="ComboBox")
-	void ClearSelection();
+	UMG_API void ClearSelection();
 
 	/**
 	 * Refreshes the list of options.  If you added new ones, and want to update the list even if it's
 	 * currently being displayed use this.
 	 */
 	UFUNCTION(BlueprintCallable, Category="ComboBox")
-	void RefreshOptions();
+	UMG_API void RefreshOptions();
 
 	UFUNCTION(BlueprintCallable, Category="ComboBox")
-	void SetSelectedOption(FString Option);
+	UMG_API void SetSelectedOption(FString Option);
 
 	UFUNCTION(BlueprintCallable, Category = "ComboBox")
-	void SetSelectedIndex(const int32 Index);
+	UMG_API void SetSelectedIndex(const int32 Index);
 
 	UFUNCTION(BlueprintCallable, Category="ComboBox")
-	FString GetSelectedOption() const;
+	UMG_API FString GetSelectedOption() const;
 
 	UFUNCTION(BlueprintCallable, Category="ComboBox")
-	int32 GetSelectedIndex() const;
+	UMG_API int32 GetSelectedIndex() const;
 
 	/** Returns the number of options */
 	UFUNCTION(BlueprintCallable, Category="ComboBox")
-	int32 GetOptionCount() const;
+	UMG_API int32 GetOptionCount() const;
 	
 	UFUNCTION(BlueprintCallable, Category="ComboBox", Meta = (ReturnDisplayName = "bOpen"))
-	bool IsOpen() const;
+	UMG_API bool IsOpen() const;
 
 	//~ Begin UVisual Interface
-	virtual void ReleaseSlateResources(bool bReleaseChildren) override;
+	UMG_API virtual void ReleaseSlateResources(bool bReleaseChildren) override;
 	//~ End UVisual Interface
 
 	/** Set the padding for content. */
-	void SetContentPadding(FMargin InPadding);
+	UMG_API void SetContentPadding(FMargin InPadding);
 
 	/** Get the padding for content. */
-	FMargin GetContentPadding() const;
+	UMG_API FMargin GetContentPadding() const;
 
 	/** Is the combobox navigated by gamepad. */
-	bool IsEnableGamepadNavigationMode() const;
+	UMG_API bool IsEnableGamepadNavigationMode() const;
 
 	/** Set whether the combobox is navigated by gamepad. */
-	void SetEnableGamepadNavigationMode(bool InEnableGamepadNavigationMode);
+	UMG_API void SetEnableGamepadNavigationMode(bool InEnableGamepadNavigationMode);
 
 	/** Is the combobox arrow showing. */
-	bool IsHasDownArrow() const;
+	UMG_API bool IsHasDownArrow() const;
 
 	/** Set whether the combobox arrow is showing. */
-	void SetHasDownArrow(bool InHasDownArrow);
+	UMG_API void SetHasDownArrow(bool InHasDownArrow);
 
 	/** Get the maximum height of the combobox list. */
-	float GetMaxListHeight() const;
+	UMG_API float GetMaxListHeight() const;
 
 	/** Set the maximum height of the combobox list. */
-	void SetMaxListHeight(float InMaxHeight);
+	UMG_API void SetMaxListHeight(float InMaxHeight);
 
 	/** Get the default font for Combobox if no OnGenerateWidgetEvent is specified. */
-	const FSlateFontInfo& GetFont() const;
+	UMG_API const FSlateFontInfo& GetFont() const;
 
 	/** Get the style of the combobox. */
-	const FComboBoxStyle& GetWidgetStyle() const;
+	UMG_API const FComboBoxStyle& GetWidgetStyle() const;
 
 	/** Set the style of the combobox. */
-	void SetWidgetStyle(const FComboBoxStyle& InWidgetStyle);
+	UMG_API void SetWidgetStyle(const FComboBoxStyle& InWidgetStyle);
 
 	/** Get the style of the items. */
-	const FTableRowStyle& GetItemStyle() const;
+	UMG_API const FTableRowStyle& GetItemStyle() const;
 
 	/** Set the style of the items. */
-	void SetItemStyle(const FTableRowStyle& InItemStyle);
+	UMG_API void SetItemStyle(const FTableRowStyle& InItemStyle);
 
 	/** Get the style of the scrollbar. */
-	const FScrollBarStyle& GetScrollBarStyle() const;
+	UMG_API const FScrollBarStyle& GetScrollBarStyle() const;
 
 	/** Is the combobox focusable. */
-	bool IsFocusable() const;
+	UMG_API bool IsFocusable() const;
 
 	/** Get the foreground color of the button. */
-	FSlateColor GetForegroundColor() const;
+	UMG_API FSlateColor GetForegroundColor() const;
 
 	//~ Begin UObject Interface
-	virtual void PostInitProperties() override;
-	virtual void Serialize(FArchive& Ar) override;
-	virtual void PostLoad() override;
+	UMG_API virtual void PostInitProperties() override;
+	UMG_API virtual void Serialize(FArchive& Ar) override;
+	UMG_API virtual void PostLoad() override;
 	//~ End UObject Interface
 
 #if WITH_EDITOR
-	virtual const FText GetPaletteCategory() override;
+	UMG_API virtual const FText GetPaletteCategory() override;
 #endif
 
 protected:
 	/** Refresh ComboBoxContent with the correct widget/data when the selected option changes */
-	void UpdateOrGenerateWidget(TSharedPtr<FString> Item);
+	UMG_API void UpdateOrGenerateWidget(TSharedPtr<FString> Item);
 
 	/** Called by slate when it needs to generate a new item for the combobox */
-	virtual TSharedRef<SWidget> HandleGenerateWidget(TSharedPtr<FString> Item) const;
+	UMG_API virtual TSharedRef<SWidget> HandleGenerateWidget(TSharedPtr<FString> Item) const;
 
 	/** Called by slate when the underlying combobox selection changes */
-	virtual void HandleSelectionChanged(TSharedPtr<FString> Item, ESelectInfo::Type SelectionType);
+	UMG_API virtual void HandleSelectionChanged(TSharedPtr<FString> Item, ESelectInfo::Type SelectionType);
 
 	/** Called by slate when the underlying combobox is opening */
-	virtual void HandleOpening();
+	UMG_API virtual void HandleOpening();
 
 	//~ Begin UWidget Interface
-	virtual TSharedRef<SWidget> RebuildWidget() override;
+	UMG_API virtual TSharedRef<SWidget> RebuildWidget() override;
 	//~ End UWidget Interface
 
 	/** Initialize the scrollbar style in the constructor before the SWidget is constructed. */
-	void InitScrollBarStyle(const FScrollBarStyle& InScrollBarStyle);
+	UMG_API void InitScrollBarStyle(const FScrollBarStyle& InScrollBarStyle);
 
 	/** Initialize the default font in the constructor before the SWidget is constructed. */
-	void InitFont(FSlateFontInfo InFont);
+	UMG_API void InitFont(FSlateFontInfo InFont);
 
 	/** Initialize IsFocusable in the constructor before the SWidget is constructed. */
-	void InitIsFocusable(bool InIsFocusable);
+	UMG_API void InitIsFocusable(bool InIsFocusable);
 
 	/** Initialize ForegroundColor in the constructor before the SWidget is constructed. */
-	void InitForegroundColor(FSlateColor InForegroundColor);
+	UMG_API void InitForegroundColor(FSlateColor InForegroundColor);
 
 protected:
 	/** The true objects bound to the Slate combobox. */
diff --git a/Engine/Source/Runtime/UMG/Public/Components/ContentWidget.h b/Engine/Source/Runtime/UMG/Public/Components/ContentWidget.h
--- a/Engine/Source/Runtime/UMG/Public/Components/ContentWidget.h
+++ b/Engine/Source/Runtime/UMG/Public/Components/ContentWidget.h
@@ -8,27 +8,27 @@
 #include "ContentWidget.generated.h"
 
 /**  */
-UCLASS(Abstract)
-class UMG_API UContentWidget : public UPanelWidget
+UCLASS(Abstract, MinimalAPI)
+class UContentWidget : public UPanelWidget
 {
 	GENERATED_UCLASS_BODY()
 
 public:
 	/**  */
 	UFUNCTION(BlueprintCallable, Category="Widget|Panel")
-	UPanelSlot* GetContentSlot() const;
+	UMG_API UPanelSlot* GetContentSlot() const;
 
 	/**  */
 	UFUNCTION(BlueprintCallable, Category="Widget|Panel")
-	UPanelSlot* SetContent(UWidget* Content);
+	UMG_API UPanelSlot* SetContent(UWidget* Content);
 
 	/**  */
 	UFUNCTION(BlueprintCallable, Category="Widget|Panel")
-	UWidget* GetContent() const;
+	UMG_API UWidget* GetContent() const;
 
 protected:
 
 	// UPanelWidget
-	virtual UClass* GetSlotClass() const override;
+	UMG_API virtual UClass* GetSlotClass() const override;
 	// End UPanelWidget
 };
diff --git a/Engine/Source/Runtime/UMG/Public/Components/DynamicEntryBox.h b/Engine/Source/Runtime/UMG/Public/Components/DynamicEntryBox.h
--- a/Engine/Source/Runtime/UMG/Public/Components/DynamicEntryBox.h
+++ b/Engine/Source/Runtime/UMG/Public/Components/DynamicEntryBox.h
@@ -13,8 +13,8 @@
  *
  * No children can be manually added in the designer - all are auto-generated based on the given entry class.
  */
-UCLASS()
-class UMG_API UDynamicEntryBox : public UDynamicEntryBoxBase
+UCLASS(MinimalAPI)
+class UDynamicEntryBox : public UDynamicEntryBoxBase
 {
 	GENERATED_BODY()
 
@@ -48,14 +48,14 @@ public:
 
 	/** Clear out the box entries, optionally deleting the underlying Slate widgets entirely as well. */
 	UFUNCTION(BlueprintCallable, Category = DynamicEntryBox)
-	void Reset(bool bDeleteWidgets = false);
+	UMG_API void Reset(bool bDeleteWidgets = false);
 
 	UFUNCTION(BlueprintCallable, Category = DynamicEntryBox)
-	void RemoveEntry(UUserWidget* EntryWidget);
+	UMG_API void RemoveEntry(UUserWidget* EntryWidget);
 
 	//~ Begin UWidget Interface
 #if WITH_EDITOR	
-	virtual void ValidateCompiledDefaults(class IWidgetCompilerLog& CompileLog) const override;
+	UMG_API virtual void ValidateCompiledDefaults(class IWidgetCompilerLog& CompileLog) const override;
 #endif
 	//~ End UWidget Interface
 
@@ -70,16 +70,16 @@ public:
 	TFunction<void(UUserWidget*)> OnPreviewEntryCreatedFunc;
 #endif
 
-	virtual void SynchronizeProperties() override;
+	UMG_API virtual void SynchronizeProperties() override;
 
 private:
 	/** Creates and establishes a new dynamic entry in the box */
 	UFUNCTION(BlueprintCallable, Category = DynamicEntryBox, meta = (DisplayName = "Create Entry", AllowPrivateAccess = true))
-	UUserWidget* BP_CreateEntry();
+	UMG_API UUserWidget* BP_CreateEntry();
 
 	/** Creates and establishes a new dynamic entry in the box using the specified class instead of the default. */
 	UFUNCTION(BlueprintCallable, Category = DynamicEntryBox, meta = (DisplayName = "Create Entry of Class", AllowPrivateAccess = true, DeterminesOutputType = "EntryClass"))
-	UUserWidget* BP_CreateEntryOfClass(TSubclassOf<UUserWidget> EntryClass);
+	UMG_API UUserWidget* BP_CreateEntryOfClass(TSubclassOf<UUserWidget> EntryClass);
 
 	/**
 	 * The class of widget to create entries of.
@@ -90,4 +90,4 @@ private:
 	
 	// Let the details customization manipulate us directly
 	friend class FDynamicEntryBoxDetails;
-};
\ No newline at end of file
+};
diff --git a/Engine/Source/Runtime/UMG/Public/Components/DynamicEntryBoxBase.h b/Engine/Source/Runtime/UMG/Public/Components/DynamicEntryBoxBase.h
--- a/Engine/Source/Runtime/UMG/Public/Components/DynamicEntryBoxBase.h
+++ b/Engine/Source/Runtime/UMG/Public/Components/DynamicEntryBoxBase.h
@@ -28,8 +28,8 @@ enum class EDynamicBoxType : uint8
  *
  * @see UDynamicEntryBox for a ready-to-use version
  */
-UCLASS(Abstract)
-class UMG_API UDynamicEntryBoxBase : public UWidget
+UCLASS(Abstract, MinimalAPI)
+class UDynamicEntryBoxBase : public UWidget
 {
 	GENERATED_BODY()
 
@@ -80,18 +80,18 @@ protected:
 	FRadialBoxSettings RadialBoxSettings;
 
 public:
-	UDynamicEntryBoxBase(const FObjectInitializer& Initializer);
-	virtual void ReleaseSlateResources(bool bReleaseChildren) override;
+	UMG_API UDynamicEntryBoxBase(const FObjectInitializer& Initializer);
+	UMG_API virtual void ReleaseSlateResources(bool bReleaseChildren) override;
 
-	EDynamicBoxType GetBoxType() const;
-	const FVector2D& GetEntrySpacing() const;
+	UMG_API EDynamicBoxType GetBoxType() const;
+	UMG_API const FVector2D& GetEntrySpacing() const;
 
 	UFUNCTION(BlueprintCallable, Category = DynamicEntryBox)
-	const TArray<UUserWidget*>& GetAllEntries() const;
+	UMG_API const TArray<UUserWidget*>& GetAllEntries() const;
 
-	const FSlateChildSize& GetEntrySizeRule() const;
+	UMG_API const FSlateChildSize& GetEntrySizeRule() const;
 
-	const FRadialBoxSettings& GetRadialBoxSettings() const;
+	UMG_API const FRadialBoxSettings& GetRadialBoxSettings() const;
 
 	template <typename EntryWidgetT = UUserWidget>
 	TArray<EntryWidgetT*> GetTypedEntries() const
@@ -108,32 +108,32 @@ public:
 	}
 
 	UFUNCTION(BlueprintCallable, Category = DynamicEntryBox)
-	int32 GetNumEntries() const;
+	UMG_API int32 GetNumEntries() const;
 
 	UFUNCTION(BlueprintCallable, Category = DynamicEntryBox)
-	void SetEntrySpacing(const FVector2D& InEntrySpacing);
+	UMG_API void SetEntrySpacing(const FVector2D& InEntrySpacing);
 
 	UFUNCTION(BlueprintCallable, Category = DynamicEntryBox)
-	void SetRadialSettings(const FRadialBoxSettings& InSettings);
+	UMG_API void SetRadialSettings(const FRadialBoxSettings& InSettings);
 
-	EVerticalAlignment GetEntryVerticalAlignment() const;
+	UMG_API EVerticalAlignment GetEntryVerticalAlignment() const;
 
-	EHorizontalAlignment GetEntryHorizontalAlignment() const;
+	UMG_API EHorizontalAlignment GetEntryHorizontalAlignment() const;
 
-	int32 GetMaxElementSize() const;
+	UMG_API int32 GetMaxElementSize() const;
 
 protected:
-	virtual TSharedRef<SWidget> RebuildWidget() override;
-	virtual void SynchronizeProperties() override;
-	virtual void AddEntryChild(UUserWidget& ChildWidget);
+	UMG_API virtual TSharedRef<SWidget> RebuildWidget() override;
+	UMG_API virtual void SynchronizeProperties() override;
+	UMG_API virtual void AddEntryChild(UUserWidget& ChildWidget);
 
-	bool IsEntryClassValid(TSubclassOf<UUserWidget> InEntryClass) const;
-	UUserWidget* CreateEntryInternal(TSubclassOf<UUserWidget> InEntryClass);
-	void RemoveEntryInternal(UUserWidget* EntryWidget);
-	FMargin BuildEntryPadding(const FVector2D& DesiredSpacing);
+	UMG_API bool IsEntryClassValid(TSubclassOf<UUserWidget> InEntryClass) const;
+	UMG_API UUserWidget* CreateEntryInternal(TSubclassOf<UUserWidget> InEntryClass);
+	UMG_API void RemoveEntryInternal(UUserWidget* EntryWidget);
+	UMG_API FMargin BuildEntryPadding(const FVector2D& DesiredSpacing);
 
 	/** Clear out the box entries, optionally deleting the underlying Slate widgets entirely as well. */
-	void ResetInternal(bool bDeleteWidgets = false);
+	UMG_API void ResetInternal(bool bDeleteWidgets = false);
 
 	/** Clear out the box entries, executing the provided reset function for each and optionally deleting the underlying Slate widgets entirely as well. */
 	template <typename WidgetT = UUserWidget>
@@ -151,24 +151,24 @@ protected:
 	}
 
 #if WITH_EDITOR
-	virtual const FText GetPaletteCategory() override;
-	virtual void PostEditChangeProperty(struct FPropertyChangedEvent& PropertyChangedEvent) override;
+	UMG_API virtual const FText GetPaletteCategory() override;
+	UMG_API virtual void PostEditChangeProperty(struct FPropertyChangedEvent& PropertyChangedEvent) override;
 #endif
 
 	// Initialize EntryBoxType in the constructor before the SWidget is constructed.
-	void InitEntryBoxType(EDynamicBoxType InEntryBoxType);
+	UMG_API void InitEntryBoxType(EDynamicBoxType InEntryBoxType);
 
 	// Initialize EntrySizeRule in the constructor before the SWidget is constructed.
-	void InitEntrySizeRule(FSlateChildSize InEntrySizeRule);
+	UMG_API void InitEntrySizeRule(FSlateChildSize InEntrySizeRule);
 
 	// Initialize EntryHorizontalAlignment in the constructor before the SWidget is constructed.
-	void InitEntryHorizontalAlignment(EHorizontalAlignment InEntryHorizontalAlignment);
+	UMG_API void InitEntryHorizontalAlignment(EHorizontalAlignment InEntryHorizontalAlignment);
 
 	// Initialize EntryVerticalAlignment in the constructor before the SWidget is constructed.
-	void InitEntryVerticalAlignment(EVerticalAlignment InEntryVerticalAlignment);
+	UMG_API void InitEntryVerticalAlignment(EVerticalAlignment InEntryVerticalAlignment);
 
 	// Initialize MaxElementSize in the constructor before the SWidget is constructed.
-	void InitMaxElementSize(int32 InMaxElementSize);
+	UMG_API void InitMaxElementSize(int32 InMaxElementSize);
 
 protected:
 	// Can be a horizontal, vertical, wrap box, or overlay
@@ -180,4 +180,4 @@ private:
 
 	// Let the details customization manipulate us directly
 	friend class FDynamicEntryBoxBaseDetails;
-};
\ No newline at end of file
+};
