[Engine]

diff --git a/Engine/Source/Runtime/UMG/Public/Components/EditableText.h b/Engine/Source/Runtime/UMG/Public/Components/EditableText.h
--- a/Engine/Source/Runtime/UMG/Public/Components/EditableText.h
+++ b/Engine/Source/Runtime/UMG/Public/Components/EditableText.h
@@ -21,8 +21,8 @@ class USlateWidgetStyleAsset;
 /**
  * Editable text box widget
  */
-UCLASS()
-class UMG_API UEditableText : public UWidget
+UCLASS(MinimalAPI)
+class UEditableText : public UWidget
 {
 	GENERATED_UCLASS_BODY()
 
@@ -144,7 +144,7 @@ public:
 	 * @return The widget text
 	 */
 	UFUNCTION(BlueprintCallable, Category = "Widget", meta=(DisplayName="GetText (Editable Text)"))
-	FText GetText() const;
+	UMG_API FText GetText() const;
 
 	/**
 	 * Directly sets the widget text.
@@ -152,22 +152,22 @@ public:
 	 * @param InText The text to assign to the widget
 	 */
 	UFUNCTION(BlueprintCallable, Category = "Widget", meta=(DisplayName="SetText (Editable Text)"))
-	void SetText(FText InText);
+	UMG_API void SetText(FText InText);
 
-	bool GetIsPassword() const;
+	UMG_API bool GetIsPassword() const;
 
 	UFUNCTION(BlueprintCallable, Category = "Widget")
-	void SetIsPassword(UPARAM(DisplayName="IsPassword") bool InbIsPassword);
+	UMG_API void SetIsPassword(UPARAM(DisplayName="IsPassword") bool InbIsPassword);
 
 	/** Gets the Hint text that appears when there is no text in the text box */
 	UFUNCTION(BlueprintCallable, Category = "Widget", meta = (DisplayName = "GetHintText (Editable Text)"))
-	FText GetHintText() const;
+	UMG_API FText GetHintText() const;
 
 	UFUNCTION(BlueprintCallable, Category = "Widget")
-	void SetHintText(FText InHintText);
+	UMG_API void SetHintText(FText InHintText);
 
 	/** @return the minimum desired width for this text box */
-	float GetMinimumDesiredWidth() const;
+	UMG_API float GetMinimumDesiredWidth() const;
 
 	/**
 	 *  Set the minimum desired width for this text box
@@ -175,97 +175,97 @@ public:
 	 *  @param InMinDesiredWidth new minimum desired width
 	*/
 	UFUNCTION(BlueprintCallable, Category = "Appearance")
-	void SetMinimumDesiredWidth(float InMinDesiredWidth);
+	UMG_API void SetMinimumDesiredWidth(float InMinDesiredWidth);
 
 	/** When set to true the caret is moved when gaining focus */
-	void SetIsCaretMovedWhenGainFocus(bool bIsCaretMovedWhenGainFocus);
+	UMG_API void SetIsCaretMovedWhenGainFocus(bool bIsCaretMovedWhenGainFocus);
 
 	/** Return true when the caret is moved when gaining focus */
-	bool GetIsCaretMovedWhenGainFocus() const;
+	UMG_API bool GetIsCaretMovedWhenGainFocus() const;
 
 	/** Set to true to select all text when the user clicks to give focus on the widget */
-	void SetSelectAllTextWhenFocused(bool bSelectAllTextWhenFocused);
+	UMG_API void SetSelectAllTextWhenFocused(bool bSelectAllTextWhenFocused);
 
 	/** Whether to select all text when the user clicks to give focus on the widget */
-	bool GetSelectAllTextWhenFocused() const;
+	UMG_API bool GetSelectAllTextWhenFocused() const;
 
 	/** Set to true to allow the user to back out of changes when they press the escape key */
-	void SetRevertTextOnEscape(bool bRevertTextOnEscape);
+	UMG_API void SetRevertTextOnEscape(bool bRevertTextOnEscape);
 
 	/** Whether to allow the user to back out of changes when they press the escape key  */
-	bool GetRevertTextOnEscape() const;
+	UMG_API bool GetRevertTextOnEscape() const;
 
 	/** Whether to clear keyboard focus when pressing enter to commit changes */
-	bool GetClearKeyboardFocusOnCommit() const;
+	UMG_API bool GetClearKeyboardFocusOnCommit() const;
 
 	/** Set to true to select all text when pressing enter to commit changes */
-	void SetSelectAllTextOnCommit(bool bSelectAllTextOnCommit);
+	UMG_API void SetSelectAllTextOnCommit(bool bSelectAllTextOnCommit);
 
 	/** Whether to select all text when pressing enter to commit changes */
-	bool GetSelectAllTextOnCommit() const;
+	UMG_API bool GetSelectAllTextOnCommit() const;
 
 	/** Set the style and ensure the visual will be updated accordingly */
-	void SetWidgetStyle(const FEditableTextStyle& InEditableTextStyle);
+	UMG_API void SetWidgetStyle(const FEditableTextStyle& InEditableTextStyle);
 
-	bool GetIsReadOnly() const;
+	UMG_API bool GetIsReadOnly() const;
 
 	UFUNCTION(BlueprintCallable, Category = "Widget", meta=(DisplayName="SetIsReadOnly (Editable Text)"))
-	void SetIsReadOnly(UPARAM(DisplayName="ReadyOnly") bool InbIsReadyOnly);
+	UMG_API void SetIsReadOnly(UPARAM(DisplayName="ReadyOnly") bool InbIsReadyOnly);
 
 	UFUNCTION(BlueprintCallable, Category = "Widget")
-	ETextJustify::Type GetJustification() const;
+	UMG_API ETextJustify::Type GetJustification() const;
 
 	UFUNCTION(BlueprintSetter)
-	void SetJustification(ETextJustify::Type InJustification);
+	UMG_API void SetJustification(ETextJustify::Type InJustification);
 
 	/** @return the text overflow policy for this text block. */
-	ETextOverflowPolicy GetTextOverflowPolicy() const;
+	UMG_API ETextOverflowPolicy GetTextOverflowPolicy() const;
 
 	UFUNCTION(BlueprintSetter)
-	void SetTextOverflowPolicy(ETextOverflowPolicy InOverflowPolicy);
+	UMG_API void SetTextOverflowPolicy(ETextOverflowPolicy InOverflowPolicy);
 
 	/** Set to true to clear keyboard focus when pressing enter to commit changes */
-	void SetClearKeyboardFocusOnCommit(bool bInClearKeyboardFocusOnCommit);
+	UMG_API void SetClearKeyboardFocusOnCommit(bool bInClearKeyboardFocusOnCommit);
 
-	void SetKeyboardType(EVirtualKeyboardType::Type Type);
+	UMG_API void SetKeyboardType(EVirtualKeyboardType::Type Type);
 
 	UFUNCTION(BlueprintCallable, Category = "Appearance")
-	const FSlateFontInfo& GetFont() const;
+	UMG_API const FSlateFontInfo& GetFont() const;
 
 	UFUNCTION(BlueprintCallable, Category = "Appearance")
-	void SetFont(FSlateFontInfo InFontInfo);
+	UMG_API void SetFont(FSlateFontInfo InFontInfo);
 
 	UFUNCTION(BlueprintCallable, Category = "Appearance")
-	void SetFontMaterial(UMaterialInterface* InMaterial);
+	UMG_API void SetFontMaterial(UMaterialInterface* InMaterial);
 
 	UFUNCTION(BlueprintCallable, Category = "Appearance")
-	void SetFontOutlineMaterial(UMaterialInterface* InMaterial);
+	UMG_API void SetFontOutlineMaterial(UMaterialInterface* InMaterial);
 
 public:
 	
 	//~ Begin UWidget Interface
-	virtual void SynchronizeProperties() override;
+	UMG_API virtual void SynchronizeProperties() override;
 	//~ End UWidget Interface
 
 	//~ Begin UVisual Interface
-	virtual void ReleaseSlateResources(bool bReleaseChildren) override;
+	UMG_API virtual void ReleaseSlateResources(bool bReleaseChildren) override;
 	//~ End UVisual Interface
 
 #if WITH_EDITOR
-	virtual const FText GetPaletteCategory() override;
+	UMG_API virtual const FText GetPaletteCategory() override;
 #endif
 	//~ End UWidget Interface
 
 protected:
 	//~ Begin UWidget Interface
-	virtual TSharedRef<SWidget> RebuildWidget() override;
+	UMG_API virtual TSharedRef<SWidget> RebuildWidget() override;
 	// End of UWidget
 
-	void HandleOnTextChanged(const FText& Text);
-	void HandleOnTextCommitted(const FText& Text, ETextCommit::Type CommitMethod);
+	UMG_API void HandleOnTextChanged(const FText& Text);
+	UMG_API void HandleOnTextCommitted(const FText& Text, ETextCommit::Type CommitMethod);
 
 #if WITH_ACCESSIBILITY
-	virtual TSharedPtr<SWidget> GetAccessibleWidget() const override;
+	UMG_API virtual TSharedPtr<SWidget> GetAccessibleWidget() const override;
 #endif
 
 protected:
diff --git a/Engine/Source/Runtime/UMG/Public/Components/EditableTextBox.h b/Engine/Source/Runtime/UMG/Public/Components/EditableTextBox.h
--- a/Engine/Source/Runtime/UMG/Public/Components/EditableTextBox.h
+++ b/Engine/Source/Runtime/UMG/Public/Components/EditableTextBox.h
@@ -22,8 +22,8 @@ class USlateWidgetStyleAsset;
  * * No Children
  * * Text Entry
  */
-UCLASS(meta=(DisplayName="Text Box"))
-class UMG_API UEditableTextBox : public UWidget
+UCLASS(meta=(DisplayName="Text Box"), MinimalAPI)
+class UEditableTextBox : public UWidget
 {
 	GENERATED_UCLASS_BODY()
 
@@ -152,7 +152,7 @@ public:
 	* @return The widget text
 	*/
 	UFUNCTION(BlueprintCallable, Category="Widget", meta=(DisplayName="GetText (Text Box)"))
-	FText GetText() const;
+	UMG_API FText GetText() const;
 
 	/**
 	* Directly sets the widget text.
@@ -160,86 +160,86 @@ public:
 	* @param InText The text to assign to the widget
 	*/
 	UFUNCTION(BlueprintCallable, Category="Widget", meta=(DisplayName="SetText (Text Box)"))
-	void SetText(FText InText);
+	UMG_API void SetText(FText InText);
 	
 	/** Gets the Hint text that appears when there is no text in the text box */
-	FText GetHintText() const;
+	UMG_API FText GetHintText() const;
 
 	/**
 	* Sets the Hint text that appears when there is no text in the text box
 	* @param InHintText The text that appears when there is no text in the text box
 	*/
 	UFUNCTION(BlueprintCallable, Category = "Widget", meta = (DisplayName = "Set Hint Text (Text Box)"))
-	void SetHintText(FText InText);
+	UMG_API void SetHintText(FText InText);
 
 	/** @return the minimum desired width for this text box */
-	float GetMinimumDesiredWidth() const;
+	UMG_API float GetMinimumDesiredWidth() const;
 
 	/**
 	*  Set the minimum desired width for this text box
 	*
 	*  @param InMinDesiredWidth new minimum desired width
 	*/
-	void SetMinDesiredWidth(float InMinDesiredWidth);
+	UMG_API void SetMinDesiredWidth(float InMinDesiredWidth);
 	
 	/** When set to true the caret is moved when gaining focus */
-	void SetIsCaretMovedWhenGainFocus(bool bIsCaretMovedWhenGainFocus);
+	UMG_API void SetIsCaretMovedWhenGainFocus(bool bIsCaretMovedWhenGainFocus);
 
 	/** Return true when the caret is moved when gaining focus */
-	bool GetIsCaretMovedWhenGainFocus() const;
+	UMG_API bool GetIsCaretMovedWhenGainFocus() const;
 
 	/** Set to true to select all text when the user clicks to give focus on the widget */
-	void SetSelectAllTextWhenFocused(bool bSelectAllTextWhenFocused);
+	UMG_API void SetSelectAllTextWhenFocused(bool bSelectAllTextWhenFocused);
 
 	/** Whether to select all text when the user clicks to give focus on the widget */
-	bool GetSelectAllTextWhenFocused() const;
+	UMG_API bool GetSelectAllTextWhenFocused() const;
 
 	/** Set to true to allow the user to back out of changes when they press the escape key */
-	void SetRevertTextOnEscape(bool bRevertTextOnEscape);
+	UMG_API void SetRevertTextOnEscape(bool bRevertTextOnEscape);
 
 	/** Whether to allow the user to back out of changes when they press the escape key  */
-	bool GetRevertTextOnEscape() const;
+	UMG_API bool GetRevertTextOnEscape() const;
 
 	/** Set to true to clear keyboard focus when pressing enter to commit changes */
-	void SetClearKeyboardFocusOnCommit(bool bClearKeyboardFocusOnCommit);
+	UMG_API void SetClearKeyboardFocusOnCommit(bool bClearKeyboardFocusOnCommit);
 
 	/** Whether to clear keyboard focus when pressing enter to commit changes */
-	bool GetClearKeyboardFocusOnCommit() const;
+	UMG_API bool GetClearKeyboardFocusOnCommit() const;
 
 	/** Set to true to select all text when pressing enter to commit changes */
-	void SetSelectAllTextOnCommit(bool bSelectAllTextOnCommit);
+	UMG_API void SetSelectAllTextOnCommit(bool bSelectAllTextOnCommit);
 
 	/** Whether to select all text when pressing enter to commit changes */
-	bool GetSelectAllTextOnCommit() const;
+	UMG_API bool GetSelectAllTextOnCommit() const;
 
 	UFUNCTION(BlueprintCallable, Category="Widget",  meta=(DisplayName="SetError (Text Box)"))
-	void SetError(FText InError);
+	UMG_API void SetError(FText InError);
 
 	/** Return true when this text cannot be modified interactively by the user */
-	bool GetIsReadOnly() const;
+	UMG_API bool GetIsReadOnly() const;
 
 	/** Sets the Text as Readonly to prevent it from being modified interactively by the user */
 	UFUNCTION(BlueprintCallable, Category = "Widget", meta = (DisplayName = "SetIsReadOnly (Editable Text)"))
-	void SetIsReadOnly(UPARAM(DisplayName = "ReadyOnly") bool bReadOnly);
+	UMG_API void SetIsReadOnly(UPARAM(DisplayName = "ReadyOnly") bool bReadOnly);
 
-	bool GetIsPassword() const;
+	UMG_API bool GetIsPassword() const;
 
 	UFUNCTION(BlueprintCallable, Category = "Widget", meta = (DisplayName = "IsPassword"))
-	void SetIsPassword(bool bIsPassword);
+	UMG_API void SetIsPassword(bool bIsPassword);
 
 	UFUNCTION(BlueprintCallable, Category="Widget")
-	void ClearError();
+	UMG_API void ClearError();
 
 	UFUNCTION(BlueprintCallable, Category = "Widget")
-	bool HasError() const;
+	UMG_API bool HasError() const;
 
-	ETextJustify::Type GetJustification() const;
+	UMG_API ETextJustify::Type GetJustification() const;
 
 	UFUNCTION(BlueprintSetter)
-	void SetJustification(ETextJustify::Type InJustification);
+	UMG_API void SetJustification(ETextJustify::Type InJustification);
 
 	/** @return the text overflow policy for this text block. */
-	ETextOverflowPolicy GetTextOverflowPolicy() const;
+	UMG_API ETextOverflowPolicy GetTextOverflowPolicy() const;
 
 	/**
 	 * Set the text overflow policy for this text box.
@@ -247,34 +247,34 @@ public:
 	 * @param InOverflowPolicy the new text overflow policy.
 	 */
 	UFUNCTION(BlueprintSetter)
-	void SetTextOverflowPolicy(ETextOverflowPolicy InOverflowPolicy);
+	UMG_API void SetTextOverflowPolicy(ETextOverflowPolicy InOverflowPolicy);
 
 	UFUNCTION(BlueprintCallable, Category="Widget", meta=(DisplayName="SetForegroundColor (Text Box)"))
-	void SetForegroundColor(FLinearColor color);
+	UMG_API void SetForegroundColor(FLinearColor color);
 
 	//~ Begin UWidget Interface
-	virtual void SynchronizeProperties() override;
+	UMG_API virtual void SynchronizeProperties() override;
 	//~ End UWidget Interface
 
 	//~ Begin UVisual Interface
-	virtual void ReleaseSlateResources(bool bReleaseChildren) override;
+	UMG_API virtual void ReleaseSlateResources(bool bReleaseChildren) override;
 	//~ End UVisual Interface
 
 #if WITH_EDITOR
-	virtual const FText GetPaletteCategory() override;
+	UMG_API virtual const FText GetPaletteCategory() override;
 #endif
-	virtual void Serialize(FArchive& Ar) override;
+	UMG_API virtual void Serialize(FArchive& Ar) override;
 
 protected:
 	//~ Begin UWidget Interface
-	virtual TSharedRef<SWidget> RebuildWidget() override;
+	UMG_API virtual TSharedRef<SWidget> RebuildWidget() override;
 	// End of UWidget
 
-	virtual void HandleOnTextChanged(const FText& Text);
-	virtual void HandleOnTextCommitted(const FText& Text, ETextCommit::Type CommitMethod);
+	UMG_API virtual void HandleOnTextChanged(const FText& Text);
+	UMG_API virtual void HandleOnTextCommitted(const FText& Text, ETextCommit::Type CommitMethod);
 
 #if WITH_ACCESSIBILITY
-	virtual TSharedPtr<SWidget> GetAccessibleWidget() const override;
+	UMG_API virtual TSharedPtr<SWidget> GetAccessibleWidget() const override;
 #endif
 
 protected:
diff --git a/Engine/Source/Runtime/UMG/Public/Components/ExpandableArea.h b/Engine/Source/Runtime/UMG/Public/Components/ExpandableArea.h
--- a/Engine/Source/Runtime/UMG/Public/Components/ExpandableArea.h
+++ b/Engine/Source/Runtime/UMG/Public/Components/ExpandableArea.h
@@ -21,8 +21,8 @@ DECLARE_DYNAMIC_MULTICAST_DELEGATE_TwoParams(FOnExpandableAreaExpansionChanged,
 /**
  * 
  */
-UCLASS()
-class UMG_API UExpandableArea : public UWidget, public INamedSlotInterface
+UCLASS(MinimalAPI)
+class UExpandableArea : public UWidget, public INamedSlotInterface
 {
 	GENERATED_UCLASS_BODY()
 
@@ -66,65 +66,65 @@ public:
 public:
 
 	UFUNCTION(BlueprintCallable, Category="Expansion")
-	bool GetIsExpanded() const;
+	UMG_API bool GetIsExpanded() const;
 
 	UFUNCTION(BlueprintCallable, Category="Expansion")
-	void SetIsExpanded(bool IsExpanded);
+	UMG_API void SetIsExpanded(bool IsExpanded);
 
 	UFUNCTION(BlueprintCallable, Category = "Expansion")
-	void SetIsExpanded_Animated(bool IsExpanded);
+	UMG_API void SetIsExpanded_Animated(bool IsExpanded);
 	
-	const FExpandableAreaStyle& GetStyle() const;
+	UMG_API const FExpandableAreaStyle& GetStyle() const;
 
-	void SetStyle(const FExpandableAreaStyle& InStyle);
+	UMG_API void SetStyle(const FExpandableAreaStyle& InStyle);
 
-	const FSlateBrush& GetBorderBrush() const;
+	UMG_API const FSlateBrush& GetBorderBrush() const;
 
-	void SetBorderBrush(const FSlateBrush& InBorderBrush);
+	UMG_API void SetBorderBrush(const FSlateBrush& InBorderBrush);
 
-	const FSlateColor& GetBorderColor() const;
+	UMG_API const FSlateColor& GetBorderColor() const;
 
-	void SetBorderColor(const FSlateColor& InBorderColor);
+	UMG_API void SetBorderColor(const FSlateColor& InBorderColor);
 
-	float GetMaxHeight() const;
+	UMG_API float GetMaxHeight() const;
 
-	void SetMaxHeight(float InMaxHeight);
+	UMG_API void SetMaxHeight(float InMaxHeight);
 
-	FMargin GetHeaderPadding() const;
+	UMG_API FMargin GetHeaderPadding() const;
 
-	void SetHeaderPadding(FMargin InHeaderPadding);
+	UMG_API void SetHeaderPadding(FMargin InHeaderPadding);
 
-	FMargin GetAreaPadding() const;
-	void SetAreaPadding(FMargin InAreaPadding);
+	UMG_API FMargin GetAreaPadding() const;
+	UMG_API void SetAreaPadding(FMargin InAreaPadding);
 
 	// Begin INamedSlotInterface
-	virtual void GetSlotNames(TArray<FName>& SlotNames) const override;
-	virtual UWidget* GetContentForSlot(FName SlotName) const override;
-	virtual void SetContentForSlot(FName SlotName, UWidget* Content) override;
+	UMG_API virtual void GetSlotNames(TArray<FName>& SlotNames) const override;
+	UMG_API virtual UWidget* GetContentForSlot(FName SlotName) const override;
+	UMG_API virtual void SetContentForSlot(FName SlotName, UWidget* Content) override;
 	// End INamedSlotInterface
 
 public:
 	
 	// UWidget interface
-	virtual void SynchronizeProperties() override;
+	UMG_API virtual void SynchronizeProperties() override;
 	// End of UWidget interface
 
 	// UVisual interface
-	virtual void ReleaseSlateResources(bool bReleaseChildren) override;
+	UMG_API virtual void ReleaseSlateResources(bool bReleaseChildren) override;
 	// End of UVisual interface
 
 #if WITH_EDITOR
-	virtual const FText GetPaletteCategory() override;
-	virtual void OnDescendantSelectedByDesigner(UWidget* DescendantWidget) override;
-	virtual void OnDescendantDeselectedByDesigner(UWidget* DescendantWidget) override;
+	UMG_API virtual const FText GetPaletteCategory() override;
+	UMG_API virtual void OnDescendantSelectedByDesigner(UWidget* DescendantWidget) override;
+	UMG_API virtual void OnDescendantDeselectedByDesigner(UWidget* DescendantWidget) override;
 #endif
 
 protected:
 	// UWidget interface
-	virtual TSharedRef<SWidget> RebuildWidget() override;
+	UMG_API virtual TSharedRef<SWidget> RebuildWidget() override;
 	// End of UWidget interface
 
-	void SlateExpansionChanged(bool NewState);
+	UMG_API void SlateExpansionChanged(bool NewState);
 
 protected:
 	UPROPERTY()
diff --git a/Engine/Source/Runtime/UMG/Public/Components/GridPanel.h b/Engine/Source/Runtime/UMG/Public/Components/GridPanel.h
--- a/Engine/Source/Runtime/UMG/Public/Components/GridPanel.h
+++ b/Engine/Source/Runtime/UMG/Public/Components/GridPanel.h
@@ -17,8 +17,8 @@ class UGridSlot;
  * 
  * * Many Children
  */
-UCLASS()
-class UMG_API UGridPanel : public UPanelWidget
+UCLASS(MinimalAPI)
+class UGridPanel : public UPanelWidget
 {
 	GENERATED_UCLASS_BODY()
 
@@ -36,32 +36,32 @@ public:
 
 	/**  */
 	UFUNCTION(BlueprintCallable, Category="Widget")
-	UGridSlot* AddChildToGrid(UWidget* Content, int32 InRow = 0, int32 InColumn = 0);
+	UMG_API UGridSlot* AddChildToGrid(UWidget* Content, int32 InRow = 0, int32 InColumn = 0);
 
 	UFUNCTION(BlueprintCallable, Category = "Widget")
-	void SetColumnFill(int32 ColumnIndex, float Coefficient);
+	UMG_API void SetColumnFill(int32 ColumnIndex, float Coefficient);
 
 	UFUNCTION(BlueprintCallable, Category = "Widget")
-	void SetRowFill(int32 RowIndex, float Coefficient);
+	UMG_API void SetRowFill(int32 RowIndex, float Coefficient);
 
 public:
 
 	// UWidget interface
-	virtual void SynchronizeProperties() override;
+	UMG_API virtual void SynchronizeProperties() override;
 	// End of UWidget interface
 
-	virtual void ReleaseSlateResources(bool bReleaseChildren) override;
+	UMG_API virtual void ReleaseSlateResources(bool bReleaseChildren) override;
 
 #if WITH_EDITOR
-	virtual const FText GetPaletteCategory() override;
+	UMG_API virtual const FText GetPaletteCategory() override;
 #endif
 
 protected:
 
 	// UPanelWidget
-	virtual UClass* GetSlotClass() const override;
-	virtual void OnSlotAdded(UPanelSlot* Slot) override;
-	virtual void OnSlotRemoved(UPanelSlot* Slot) override;
+	UMG_API virtual UClass* GetSlotClass() const override;
+	UMG_API virtual void OnSlotAdded(UPanelSlot* Slot) override;
+	UMG_API virtual void OnSlotRemoved(UPanelSlot* Slot) override;
 	// End UPanelWidget
 
 protected:
@@ -70,6 +70,6 @@ protected:
 
 protected:
 	// UWidget interface
-	virtual TSharedRef<SWidget> RebuildWidget() override;
+	UMG_API virtual TSharedRef<SWidget> RebuildWidget() override;
 	// End of UWidget interface
 };
diff --git a/Engine/Source/Runtime/UMG/Public/Components/GridSlot.h b/Engine/Source/Runtime/UMG/Public/Components/GridSlot.h
--- a/Engine/Source/Runtime/UMG/Public/Components/GridSlot.h
+++ b/Engine/Source/Runtime/UMG/Public/Components/GridSlot.h
@@ -15,8 +15,8 @@
  * A slot for UGridPanel, these slots all share the same size as the largest slot
  * in the grid.
  */
-UCLASS()
-class UMG_API UGridSlot : public UPanelSlot
+UCLASS(MinimalAPI)
+class UGridSlot : public UPanelSlot
 {
 	GENERATED_UCLASS_BODY()
 
@@ -69,79 +69,79 @@ public:
 
 public:
 
-	FMargin GetPadding() const;
+	UMG_API FMargin GetPadding() const;
 
 	UFUNCTION(BlueprintCallable, Category="Layout|Border Slot")
-	void SetPadding(FMargin InPadding);
+	UMG_API void SetPadding(FMargin InPadding);
 
 	/** Gets the row index of the slot */
-	int32 GetRow() const;
+	UMG_API int32 GetRow() const;
 
 	/** Sets the row index of the slot, this determines what cell the slot is in the panel */
 	UFUNCTION(BlueprintCallable, Category="Layout|Grid Slot")
-	void SetRow(int32 InRow);
+	UMG_API void SetRow(int32 InRow);
 
 	/** Gets how many rows this this slot spans over */
-	int32 GetRowSpan() const;
+	UMG_API int32 GetRowSpan() const;
 
 	/** How many rows this this slot spans over */
 	UFUNCTION(BlueprintCallable, Category="Layout|Grid Slot")
-	void SetRowSpan(int32 InRowSpan);
+	UMG_API void SetRowSpan(int32 InRowSpan);
 
 	/** Gets the column index of the slot */
-	int32 GetColumn() const;
+	UMG_API int32 GetColumn() const;
 
 	/** Sets the column index of the slot, this determines what cell the slot is in the panel */
 	UFUNCTION(BlueprintCallable, Category="Layout|Grid Slot")
-	void SetColumn(int32 InColumn);
+	UMG_API void SetColumn(int32 InColumn);
 
 	/** Gets how many columns this slot spans over */
-	int32 GetColumnSpan() const;
+	UMG_API int32 GetColumnSpan() const;
 
 	/** How many columns this slot spans over */
 	UFUNCTION(BlueprintCallable, Category="Layout|Grid Slot")
-	void SetColumnSpan(int32 InColumnSpan);
+	UMG_API void SetColumnSpan(int32 InColumnSpan);
 
 	/** Gets the offset this cell drawn on top of others. */
-	int32 GetLayer() const;
+	UMG_API int32 GetLayer() const;
 
 	/** Sets positive values offset this cell to be hit-tested and drawn on top of others. */
 	UFUNCTION(BlueprintCallable, Category = "Layout|Grid Slot")
-	void SetLayer(int32 InLayer);
+	UMG_API void SetLayer(int32 InLayer);
 
 	/** Gets the offset for this slot's content */
-	FVector2D GetNudge() const;
+	UMG_API FVector2D GetNudge() const;
 
 	/** Sets the offset for this slot's content by some amount; positive values offset to lower right*/
 	UFUNCTION(BlueprintCallable, Category = "Layout|Grid Slot")
-	void SetNudge(FVector2D InNudge);
+	UMG_API void SetNudge(FVector2D InNudge);
 
 	/** */
-	EHorizontalAlignment GetHorizontalAlignment() const;
+	UMG_API EHorizontalAlignment GetHorizontalAlignment() const;
 
 	/**  */
 	UFUNCTION(BlueprintCallable, Category="Layout|Grid Slot")
-	void SetHorizontalAlignment(EHorizontalAlignment InHorizontalAlignment);
+	UMG_API void SetHorizontalAlignment(EHorizontalAlignment InHorizontalAlignment);
 
 	/**  */
-	EVerticalAlignment GetVerticalAlignment() const;
+	UMG_API EVerticalAlignment GetVerticalAlignment() const;
 
 	/**  */
 	UFUNCTION(BlueprintCallable, Category="Layout|Grid Slot")
-	void SetVerticalAlignment(EVerticalAlignment InVerticalAlignment);
+	UMG_API void SetVerticalAlignment(EVerticalAlignment InVerticalAlignment);
 
 	// UPanelSlot interface
-	virtual void SynchronizeProperties() override;
+	UMG_API virtual void SynchronizeProperties() override;
 	// End of UPanelSlot interface
 
 	/** Builds the underlying FSlot for the Slate layout panel. */
-	void BuildSlot(TSharedRef<SGridPanel> GridPanel);
+	UMG_API void BuildSlot(TSharedRef<SGridPanel> GridPanel);
 
-	virtual void ReleaseSlateResources(bool bReleaseChildren) override;
+	UMG_API virtual void ReleaseSlateResources(bool bReleaseChildren) override;
 
 #if WITH_EDITOR
-	virtual bool NudgeByDesigner(const FVector2D& NudgeDirection, const TOptional<int32>& GridSnapSize) override;
-	virtual void SynchronizeFromTemplate(const UPanelSlot* const TemplateSlot) override;
+	UMG_API virtual bool NudgeByDesigner(const FVector2D& NudgeDirection, const TOptional<int32>& GridSnapSize) override;
+	UMG_API virtual void SynchronizeFromTemplate(const UPanelSlot* const TemplateSlot) override;
 #endif //WITH_EDITOR
 
 private:
diff --git a/Engine/Source/Runtime/UMG/Public/Components/HorizontalBox.h b/Engine/Source/Runtime/UMG/Public/Components/HorizontalBox.h
--- a/Engine/Source/Runtime/UMG/Public/Components/HorizontalBox.h
+++ b/Engine/Source/Runtime/UMG/Public/Components/HorizontalBox.h
@@ -17,29 +17,29 @@ class UHorizontalBoxSlot;
  * * Many Children
  * * Flow Horizontal
  */
-UCLASS(meta = (ShortTooltip = "A layout panel for automatically laying child widgets out horizontally"))
-class UMG_API UHorizontalBox : public UPanelWidget
+UCLASS(meta = (ShortTooltip = "A layout panel for automatically laying child widgets out horizontally"), MinimalAPI)
+class UHorizontalBox : public UPanelWidget
 {
 	GENERATED_UCLASS_BODY()
 
 	/**  */
 	UFUNCTION(BlueprintCallable, Category="Widget")
-	UHorizontalBoxSlot* AddChildToHorizontalBox(UWidget* Content);
+	UMG_API UHorizontalBoxSlot* AddChildToHorizontalBox(UWidget* Content);
 
 #if WITH_EDITOR
 	// UWidget interface
-	virtual const FText GetPaletteCategory() override;
+	UMG_API virtual const FText GetPaletteCategory() override;
 	// End UWidget interface
 #endif
 
-	virtual void ReleaseSlateResources(bool bReleaseChildren) override;
+	UMG_API virtual void ReleaseSlateResources(bool bReleaseChildren) override;
 
 protected:
 
 	// UPanelWidget
-	virtual UClass* GetSlotClass() const override;
-	virtual void OnSlotAdded(UPanelSlot* Slot) override;
-	virtual void OnSlotRemoved(UPanelSlot* Slot) override;
+	UMG_API virtual UClass* GetSlotClass() const override;
+	UMG_API virtual void OnSlotAdded(UPanelSlot* Slot) override;
+	UMG_API virtual void OnSlotRemoved(UPanelSlot* Slot) override;
 	// End UPanelWidget
 
 protected:
@@ -48,6 +48,6 @@ protected:
 
 protected:
 	// UWidget interface
-	virtual TSharedRef<SWidget> RebuildWidget() override;
+	UMG_API virtual TSharedRef<SWidget> RebuildWidget() override;
 	// End of UWidget interface
 };
diff --git a/Engine/Source/Runtime/UMG/Public/Components/HorizontalBoxSlot.h b/Engine/Source/Runtime/UMG/Public/Components/HorizontalBoxSlot.h
--- a/Engine/Source/Runtime/UMG/Public/Components/HorizontalBoxSlot.h
+++ b/Engine/Source/Runtime/UMG/Public/Components/HorizontalBoxSlot.h
@@ -12,8 +12,8 @@
 
 #include "HorizontalBoxSlot.generated.h"
 
-UCLASS()
-class UMG_API UHorizontalBoxSlot : public UPanelSlot
+UCLASS(MinimalAPI)
+class UHorizontalBoxSlot : public UPanelSlot
 {
 	GENERATED_UCLASS_BODY()
 
@@ -42,38 +42,38 @@ public:
 
 public:
 
-	FMargin GetPadding() const;
+	UMG_API FMargin GetPadding() const;
 
 	UFUNCTION(BlueprintCallable, Category="Layout|Horizontal Box Slot")
-	void SetPadding(FMargin InPadding);
+	UMG_API void SetPadding(FMargin InPadding);
 
-	FSlateChildSize GetSize() const;
+	UMG_API FSlateChildSize GetSize() const;
 
 	UFUNCTION(BlueprintCallable, Category="Layout|Horizontal Box Slot")
-	void SetSize(FSlateChildSize InSize);
+	UMG_API void SetSize(FSlateChildSize InSize);
 
-	EHorizontalAlignment GetHorizontalAlignment() const;
+	UMG_API EHorizontalAlignment GetHorizontalAlignment() const;
 
 	UFUNCTION(BlueprintCallable, Category="Layout|Horizontal Box Slot")
-	void SetHorizontalAlignment(EHorizontalAlignment InHorizontalAlignment);
+	UMG_API void SetHorizontalAlignment(EHorizontalAlignment InHorizontalAlignment);
 
-	EVerticalAlignment GetVerticalAlignment() const;
+	UMG_API EVerticalAlignment GetVerticalAlignment() const;
 
 	UFUNCTION(BlueprintCallable, Category="Layout|Horizontal Box Slot")
-	void SetVerticalAlignment(EVerticalAlignment InVerticalAlignment);
+	UMG_API void SetVerticalAlignment(EVerticalAlignment InVerticalAlignment);
 
 public:
 
-	void BuildSlot(TSharedRef<SHorizontalBox> HorizontalBox);
+	UMG_API void BuildSlot(TSharedRef<SHorizontalBox> HorizontalBox);
 
 	// UPanelSlot interface
-	virtual void SynchronizeProperties() override;
+	UMG_API virtual void SynchronizeProperties() override;
 	// End of UPanelSlot interface
 
-	virtual void ReleaseSlateResources(bool bReleaseChildren) override;
+	UMG_API virtual void ReleaseSlateResources(bool bReleaseChildren) override;
 
 #if WITH_EDITOR
-	virtual bool NudgeByDesigner(const FVector2D& NudgeDirection, const TOptional<int32>& GridSnapSize) override;
-	virtual void SynchronizeFromTemplate(const UPanelSlot* const TemplateSlot) override;
+	UMG_API virtual bool NudgeByDesigner(const FVector2D& NudgeDirection, const TOptional<int32>& GridSnapSize) override;
+	UMG_API virtual void SynchronizeFromTemplate(const UPanelSlot* const TemplateSlot) override;
 #endif //WITH_EDITOR
 };
diff --git a/Engine/Source/Runtime/UMG/Public/Components/InputKeySelector.h b/Engine/Source/Runtime/UMG/Public/Components/InputKeySelector.h
--- a/Engine/Source/Runtime/UMG/Public/Components/InputKeySelector.h
+++ b/Engine/Source/Runtime/UMG/Public/Components/InputKeySelector.h
@@ -17,8 +17,8 @@ class SInputKeySelector;
 struct FButtonStyle;
 
 /** A widget for selecting a single key or a single key with a modifier. */
-UCLASS()
-class UMG_API UInputKeySelector : public UWidget
+UCLASS(MinimalAPI)
+class UInputKeySelector : public UWidget
 {
 	GENERATED_UCLASS_BODY()
 
@@ -85,95 +85,95 @@ public:
 
 	/** Sets the currently selected key. */
 	UFUNCTION( BlueprintCallable, Category = "Widget" )
-	void SetSelectedKey(const FInputChord& InSelectedKey);
+	UMG_API void SetSelectedKey(const FInputChord& InSelectedKey);
 
 	/** Returns the currently selected key. */
-	FInputChord GetSelectedKey() const;
+	UMG_API FInputChord GetSelectedKey() const;
 
 	/** Sets the text which is displayed while selecting keys. */
 	UFUNCTION( BlueprintCallable, Category = "Widget" )
-	void SetKeySelectionText( FText InKeySelectionText );
+	UMG_API void SetKeySelectionText( FText InKeySelectionText );
 
 	/** Returns the text which is displayed while selecting keys. */
-	const FText& GetKeySelectionText() const;
+	UMG_API const FText& GetKeySelectionText() const;
 
 	/** Sets the text to display when no key text is available or not selecting a key. */
 	UFUNCTION(BlueprintCallable, Category = "Widget")
-	void SetNoKeySpecifiedText(FText InNoKeySpecifiedText);
+	UMG_API void SetNoKeySpecifiedText(FText InNoKeySpecifiedText);
 
 	/** Returns the text to display when no key text is available or not selecting a key. */
-	const FText& GetNoKeySpecifiedText() const;
+	UMG_API const FText& GetNoKeySpecifiedText() const;
 
 	/** Sets whether or not modifier keys are allowed in the selected key. */
 	UFUNCTION(BlueprintCallable, Category = "Widget" )
-	void SetAllowModifierKeys(bool bInAllowModifierKeys );
+	UMG_API void SetAllowModifierKeys(bool bInAllowModifierKeys );
 
 	/** Returns true if modifier keys keys are allowed, otherwise returns false. */
-	bool AllowModifierKeys() const;
+	UMG_API bool AllowModifierKeys() const;
 
 	/** Sets whether or not gamepad keys are allowed in the selected key. */
 	UFUNCTION(BlueprintCallable, Category = "Widget")
-	void SetAllowGamepadKeys(bool bInAllowGamepadKeys);
+	UMG_API void SetAllowGamepadKeys(bool bInAllowGamepadKeys);
 
 	/** Returns true if gamepad keys are allowed, otherwise returns false. */
-	bool AllowGamepadKeys() const;
+	UMG_API bool AllowGamepadKeys() const;
 
 	/** Returns true if the widget is currently selecting a key, otherwise returns false. */
 	UFUNCTION( BlueprintCallable, Category = "Widget" )
-	bool GetIsSelectingKey() const;
+	UMG_API bool GetIsSelectingKey() const;
 
 	/** Sets the visibility of the text block. */
 	UFUNCTION(BlueprintCallable, Category = "Widget")
-	void SetTextBlockVisibility(const ESlateVisibility InVisibility);
+	UMG_API void SetTextBlockVisibility(const ESlateVisibility InVisibility);
 
 	/** Sets the style of the button used to start key selection mode. */
-	void SetButtonStyle(const FButtonStyle& ButtonStyle);
+	UMG_API void SetButtonStyle(const FButtonStyle& ButtonStyle);
 
 	/** Sets the style of the button used to start key selection mode. */
 	UE_DEPRECATED(5.2, "SetButtonStyle with pointer type parameter is deprecated. Please pass a reference instead.")
-	void SetButtonStyle(const FButtonStyle* ButtonStyle);
+	UMG_API void SetButtonStyle(const FButtonStyle* ButtonStyle);
 
 	/** Returns the style of the button used to start key selection mode. */
-	const FButtonStyle& GetButtonStyle() const;
+	UMG_API const FButtonStyle& GetButtonStyle() const;
 
 	/** Sets the style of the text used inside the button. */
-	void SetTextStyle(const FTextBlockStyle& InTextStyle);
+	UMG_API void SetTextStyle(const FTextBlockStyle& InTextStyle);
 
 	/** Returns the style of the text used inside the button. */
-	const FTextBlockStyle& GetTextStyle() const;
+	UMG_API const FTextBlockStyle& GetTextStyle() const;
 
 	/** Sets the amount of blank space around the text used to display the currently selected key. */
-	void SetMargin(const FMargin& InMargin);
+	UMG_API void SetMargin(const FMargin& InMargin);
 
 	/** Returns the amount of blank space around the text used to display the currently selected key. */
-	const FMargin& GetMargin() const;
+	UMG_API const FMargin& GetMargin() const;
 
 	/** Sets escape keys. */
 	UFUNCTION(BlueprintCallable, Category = "Widget")
-	void SetEscapeKeys(const TArray<FKey>& InKeys);
+	UMG_API void SetEscapeKeys(const TArray<FKey>& InKeys);
 
 	//~ Begin UWidget Interface
 #if WITH_EDITOR	
-	virtual const FText GetPaletteCategory() override;
+	UMG_API virtual const FText GetPaletteCategory() override;
 #endif
-	virtual void SynchronizeProperties() override;
+	UMG_API virtual void SynchronizeProperties() override;
 	//~ End UWidget Interface
 
 protected:
 	//~ Begin UObject Interface
-	virtual void Serialize(FArchive& Ar) override;
+	UMG_API virtual void Serialize(FArchive& Ar) override;
 	//~ End UObject Interface
 
 	//~ Begin UWidget Interface
-	virtual TSharedRef<SWidget> RebuildWidget() override;
+	UMG_API virtual TSharedRef<SWidget> RebuildWidget() override;
 	//~ End UWidget Interface
 
 	//~ Begin UVisual Interface
-	virtual void ReleaseSlateResources(bool bReleaseChildren) override;
+	UMG_API virtual void ReleaseSlateResources(bool bReleaseChildren) override;
 	//~ End UVisual Interface
 
 private:
-	virtual void HandleKeySelected(const FInputChord& InSelectedKey);
+	UMG_API virtual void HandleKeySelected(const FInputChord& InSelectedKey);
 	void HandleIsSelectingKeyChanged();
 
 private:
diff --git a/Engine/Source/Runtime/UMG/Public/Components/InvalidationBox.h b/Engine/Source/Runtime/UMG/Public/Components/InvalidationBox.h
--- a/Engine/Source/Runtime/UMG/Public/Components/InvalidationBox.h
+++ b/Engine/Source/Runtime/UMG/Public/Components/InvalidationBox.h
@@ -15,8 +15,8 @@ class SInvalidationPanel;
  * * Single Child
  * * Caching / Performance
  */
-UCLASS()
-class UMG_API UInvalidationBox : public UContentWidget
+UCLASS(MinimalAPI)
+class UInvalidationBox : public UContentWidget
 {
 	GENERATED_UCLASS_BODY()
 
@@ -27,37 +27,37 @@ public:
 	 */
 	UE_DEPRECATED(4.27, "InvalidationCache is not used.")
 	UFUNCTION(BlueprintCallable, Category="Invalidation Box", meta=(DeprecatedFunction))
-	void InvalidateCache();
+	UMG_API void InvalidateCache();
 
 	/**
 	 * @returns true when the invalidation box cache the widgets.
 	 * The widgets will be updated only if they get invalidated.
 	 */
 	UFUNCTION(BlueprintCallable, Category="Invalidation Box")
-	bool GetCanCache() const;
+	UMG_API bool GetCanCache() const;
 
 	/**
 	 * Tell the InvalidationBox to use the invalidation process.
 	 * @note the other internal flags can disable the option.
 	 */
 	UFUNCTION(BlueprintCallable, Category="Invalidation Box")
-	void SetCanCache(bool CanCache);
+	UMG_API void SetCanCache(bool CanCache);
 
-	virtual void ReleaseSlateResources(bool bReleaseChildren) override;
+	UMG_API virtual void ReleaseSlateResources(bool bReleaseChildren) override;
 
 #if WITH_EDITOR
-	virtual const FText GetPaletteCategory() override;
+	UMG_API virtual const FText GetPaletteCategory() override;
 #endif
 
 protected:
 
 	//~ Begin UPanelWidget interface
-	virtual void OnSlotAdded(UPanelSlot* Slot) override;
-	virtual void OnSlotRemoved(UPanelSlot* Slot) override;
+	UMG_API virtual void OnSlotAdded(UPanelSlot* Slot) override;
+	UMG_API virtual void OnSlotRemoved(UPanelSlot* Slot) override;
 	//~ End UPanelWidget interface
 
 	//~ Begin UWidget interface
-	virtual TSharedRef<SWidget> RebuildWidget() override;
+	UMG_API virtual TSharedRef<SWidget> RebuildWidget() override;
 	//~ End UWidget interface
 
 protected:
diff --git a/Engine/Source/Runtime/UMG/Public/Components/ListView.h b/Engine/Source/Runtime/UMG/Public/Components/ListView.h
--- a/Engine/Source/Runtime/UMG/Public/Components/ListView.h
+++ b/Engine/Source/Runtime/UMG/Public/Components/ListView.h
@@ -23,18 +23,18 @@ DECLARE_DYNAMIC_MULTICAST_DELEGATE_TwoParams(FOnListViewScrolledDynamic, float,
  *
  * To make a widget usable as an entry in a ListView, it must inherit from the IUserObjectListEntry interface.
  */
-UCLASS(meta = (EntryInterface = "/Script/UMG.UserObjectListEntry"))
-class UMG_API UListView : public UListViewBase, public ITypedUMGListView<UObject*>
+UCLASS(meta = (EntryInterface = "/Script/UMG.UserObjectListEntry"), MinimalAPI)
+class UListView : public UListViewBase, public ITypedUMGListView<UObject*>
 {
 	GENERATED_BODY()
 
 	IMPLEMENT_TYPED_UMG_LIST(UObject*, MyListView)
 
 public:
-	UListView(const FObjectInitializer& Initializer);
+	UMG_API UListView(const FObjectInitializer& Initializer);
 
-	virtual void ReleaseSlateResources(bool bReleaseChildren) override;
-	virtual void PostLoad() override;
+	UMG_API virtual void ReleaseSlateResources(bool bReleaseChildren) override;
+	UMG_API virtual void PostLoad() override;
 
 	/** Set the list of items to display within this listview */
 	template <typename ItemObjectT, typename AllocatorType = FDefaultAllocator>
@@ -63,7 +63,7 @@ public:
 	{
 		return Item ? ITypedUMGListView<UObject*>::GetEntryWidgetFromItem<RowWidgetT>(const_cast<UObject*>(Item)) : nullptr;
 	}
-	void SetSelectedItem(const UObject* Item);
+	UMG_API void SetSelectedItem(const UObject* Item);
 
 	/** Gets the first selected item, if any; recommended that you only use this for single selection lists. */
 	template <typename ObjectT = UObject>
@@ -81,75 +81,75 @@ public:
 	
 	/** Adds an the item to the list */
 	UFUNCTION(BlueprintCallable, Category = ListView)
-	void AddItem(UObject* Item);
+	UMG_API void AddItem(UObject* Item);
 
 	/** Removes an the item from the list */
 	UFUNCTION(BlueprintCallable, Category = ListView)
-	void RemoveItem(UObject* Item);
+	UMG_API void RemoveItem(UObject* Item);
 
 	/** Returns the item at the given index */
 	UFUNCTION(BlueprintCallable, Category = ListView)
-	UObject* GetItemAt(int32 Index) const;
+	UMG_API UObject* GetItemAt(int32 Index) const;
 
 	/** Returns the total number of items */
 	UFUNCTION(BlueprintCallable, Category = ListView)
-	int32 GetNumItems() const;
+	UMG_API int32 GetNumItems() const;
 
 	/** Returns the index that the specified item is at. Will return the first found, or -1 for not found */
 	UFUNCTION(BlueprintCallable, Category = ListView)
-	int32 GetIndexForItem(const UObject* Item) const;
+	UMG_API int32 GetIndexForItem(const UObject* Item) const;
 
 	/** Removes all items from the list */
 	UFUNCTION(BlueprintCallable, Category = ListView)
-	void ClearListItems();
+	UMG_API void ClearListItems();
 
 	/** Sets the new selection mode, preserving the current selection where possible. */
 	UFUNCTION(BlueprintCallable, Category = ListView)
-	void SetSelectionMode(TEnumAsByte<ESelectionMode::Type> SelectionMode);
+	UMG_API void SetSelectionMode(TEnumAsByte<ESelectionMode::Type> SelectionMode);
 
 	/** Returns true if a refresh is pending and the list will be rebuilt on the next tick */
 	UFUNCTION(BlueprintCallable, Category = ListView)
-	bool IsRefreshPending() const;
+	UMG_API bool IsRefreshPending() const;
 
 	/** Requests that the item at the given index is scrolled into view */
 	UFUNCTION(BlueprintCallable, Category = ListView)
-	void ScrollIndexIntoView(int32 Index);
+	UMG_API void ScrollIndexIntoView(int32 Index);
 
 	/** Sets the item at the given index as the sole selected item. */
 	UFUNCTION(BlueprintCallable, Category = ListView)
-	void SetSelectedIndex(int32 Index);
+	UMG_API void SetSelectedIndex(int32 Index);
 
 	/** Requests that the item at the given index navigated to, scrolling it into view if needed. */
 	UFUNCTION(BlueprintCallable, Category = ListView)
-	void NavigateToIndex(int32 Index);
+	UMG_API void NavigateToIndex(int32 Index);
 
 protected:
-	virtual void OnItemsChanged(const TArray<UObject*>& AddedItems, const TArray<UObject*>& RemovedItems);
+	UMG_API virtual void OnItemsChanged(const TArray<UObject*>& AddedItems, const TArray<UObject*>& RemovedItems);
 
 	UFUNCTION()
-	void OnListItemEndPlayed(AActor* Item, EEndPlayReason::Type EndPlayReason);
+	UMG_API void OnListItemEndPlayed(AActor* Item, EEndPlayReason::Type EndPlayReason);
 
 	UFUNCTION()
-	void OnListItemOuterEndPlayed(AActor* ItemOuter, EEndPlayReason::Type EndPlayReason);
+	UMG_API void OnListItemOuterEndPlayed(AActor* ItemOuter, EEndPlayReason::Type EndPlayReason);
 
-	virtual TSharedRef<STableViewBase> RebuildListWidget() override;
-	virtual void HandleListEntryHovered(UUserWidget& EntryWidget) override;
-	virtual void HandleListEntryUnhovered(UUserWidget& EntryWidget) override;
+	UMG_API virtual TSharedRef<STableViewBase> RebuildListWidget() override;
+	UMG_API virtual void HandleListEntryHovered(UUserWidget& EntryWidget) override;
+	UMG_API virtual void HandleListEntryUnhovered(UUserWidget& EntryWidget) override;
 	
 #if WITH_EDITOR
-	virtual void OnRefreshDesignerItems() override;
+	UMG_API virtual void OnRefreshDesignerItems() override;
 #endif
 
-	virtual UUserWidget& OnGenerateEntryWidgetInternal(UObject* Item, TSubclassOf<UUserWidget> DesiredEntryClass, const TSharedRef<STableViewBase>& OwnerTable) override;
-	virtual FMargin GetDesiredEntryPadding(UObject* Item) const override;
+	UMG_API virtual UUserWidget& OnGenerateEntryWidgetInternal(UObject* Item, TSubclassOf<UUserWidget> DesiredEntryClass, const TSharedRef<STableViewBase>& OwnerTable) override;
+	UMG_API virtual FMargin GetDesiredEntryPadding(UObject* Item) const override;
 
-	virtual void OnItemClickedInternal(UObject* Item) override;
-	virtual void OnItemDoubleClickedInternal(UObject* Item) override;
-	virtual void OnSelectionChangedInternal(UObject* FirstSelectedItem) override;
-	virtual void OnItemScrolledIntoViewInternal(UObject* Item, UUserWidget& EntryWidget) override;
-	virtual void OnListViewScrolledInternal(float ItemOffset, float DistanceRemaining) override;
+	UMG_API virtual void OnItemClickedInternal(UObject* Item) override;
+	UMG_API virtual void OnItemDoubleClickedInternal(UObject* Item) override;
+	UMG_API virtual void OnSelectionChangedInternal(UObject* FirstSelectedItem) override;
+	UMG_API virtual void OnItemScrolledIntoViewInternal(UObject* Item, UUserWidget& EntryWidget) override;
+	UMG_API virtual void OnListViewScrolledInternal(float ItemOffset, float DistanceRemaining) override;
 
-	void HandleOnEntryInitializedInternal(UObject* Item, const TSharedRef<ITableRow>& TableRow);
+	UMG_API void HandleOnEntryInitializedInternal(UObject* Item, const TSharedRef<ITableRow>& TableRow);
 
 	/** SListView construction helper - useful if using a custom STreeView subclass */
 	template <template<typename> class ListViewT = SListView>
@@ -232,55 +232,55 @@ public:
 	float GetVerticalEntrySpacing() const { return VerticalEntrySpacing; }
 
 protected:
-	void InitHorizontalEntrySpacing(float InHorizontalEntrySpacing);
-	void InitVerticalEntrySpacing(float InVerticalEntrySpacing);
+	UMG_API void InitHorizontalEntrySpacing(float InHorizontalEntrySpacing);
+	UMG_API void InitVerticalEntrySpacing(float InVerticalEntrySpacing);
 
 private:
 	// BP exposure of ITypedUMGListView API
 
 	/** Sets the given item as the sole selected item. */
 	UFUNCTION(BlueprintCallable, Category = ListView, meta = (AllowPrivateAccess = true, DisplayName = "Set Selected Item"))
-	void BP_SetSelectedItem(UObject* Item);
+	UMG_API void BP_SetSelectedItem(UObject* Item);
 
 	/** Sets whether the given item is selected. */
 	UFUNCTION(BlueprintCallable, Category = ListView, meta = (AllowPrivateAccess = true, DisplayName = "Set Item Selection"))
-	void BP_SetItemSelection(UObject* Item, bool bSelected);
+	UMG_API void BP_SetItemSelection(UObject* Item, bool bSelected);
 
 	/** Clear selection */
 	UFUNCTION(BlueprintCallable, Category = ListView, meta = (AllowPrivateAccess = true, DisplayName = "Clear Selection"))
-	void BP_ClearSelection();
+	UMG_API void BP_ClearSelection();
 
 	/** Gets the number of items currently selected in the list */
 	UFUNCTION(BlueprintCallable, Category = ListView, meta = (AllowPrivateAccess = true, DisplayName = "Get Num Items Selected"))
-	int32 BP_GetNumItemsSelected() const;
+	UMG_API int32 BP_GetNumItemsSelected() const;
 
 	/** Gets a list of all the currently selected items */
 	UFUNCTION(BlueprintCallable, BlueprintPure = false, Category = ListView, meta = (AllowPrivateAccess = true, DisplayName = "GetSelectedItems"))
-	bool BP_GetSelectedItems(TArray<UObject*>& Items) const;
+	UMG_API bool BP_GetSelectedItems(TArray<UObject*>& Items) const;
 
 	/** Gets whether the entry for the given object is currently visible in the list */
 	UFUNCTION(BlueprintCallable, Category = ListView, meta = (AllowPrivateAccess = true, DisplayName = "Is Item Visible"))
-	bool BP_IsItemVisible(UObject* Item) const;
+	UMG_API bool BP_IsItemVisible(UObject* Item) const;
 
 	/** Requests that the given item is navigated to, scrolling it into view if needed. */
 	UFUNCTION(BlueprintCallable, Category = ListView, meta = (AllowPrivateAccess = true, DisplayName = "Navigate To Item"))
-	void BP_NavigateToItem(UObject* Item);
+	UMG_API void BP_NavigateToItem(UObject* Item);
 
 	/** Requests that the given item is scrolled into view */
 	UFUNCTION(BlueprintCallable, Category = ListView, meta = (AllowPrivateAccess = true, DisplayName = "Scroll Item Into View"))
-	void BP_ScrollItemIntoView(UObject* Item);
+	UMG_API void BP_ScrollItemIntoView(UObject* Item);
 
 	/** Cancels a previous request to scroll and item into view. */
 	UFUNCTION(BlueprintCallable, Category = ListView, meta = (AllowPrivateAccess = true, DisplayName = "Cancel Scroll Into View"))
-	void BP_CancelScrollIntoView();
+	UMG_API void BP_CancelScrollIntoView();
 
 	/** Sets the array of objects to display rows for in the list */
 	UFUNCTION(BlueprintCallable, Category = ListView, meta = (AllowPrivateAccess = true, DisplayName = "Set List Items"))
-	void BP_SetListItems(const TArray<UObject*>& InListItems);
+	UMG_API void BP_SetListItems(const TArray<UObject*>& InListItems);
 
 	/** Gets the first selected item, if any; recommended that you only use this for single selection lists. */
 	UFUNCTION(BlueprintCallable, Category = ListView, meta = (DisplayName = "Get Selected Item", AllowPrivateAccess = true))
-	UObject* BP_GetSelectedItem() const;
+	UMG_API UObject* BP_GetSelectedItem() const;
 
 private:
 	/** Called when a row widget is generated for a list item */
diff --git a/Engine/Source/Runtime/UMG/Public/Components/ListViewBase.h b/Engine/Source/Runtime/UMG/Public/Components/ListViewBase.h
--- a/Engine/Source/Runtime/UMG/Public/Components/ListViewBase.h
+++ b/Engine/Source/Runtime/UMG/Public/Components/ListViewBase.h
@@ -489,53 +489,53 @@ private:
  * }
  *
  */
-UCLASS(Abstract, NotBlueprintable, hidedropdown, meta = (EntryInterface = UserListEntry))
-class UMG_API UListViewBase : public UWidget
+UCLASS(Abstract, NotBlueprintable, hidedropdown, meta = (EntryInterface = UserListEntry), MinimalAPI)
+class UListViewBase : public UWidget
 {
 	GENERATED_BODY()
 
 public:
-	UListViewBase(const FObjectInitializer& ObjectInitializer);
+	UMG_API UListViewBase(const FObjectInitializer& ObjectInitializer);
 
 #if WITH_EDITOR
-	virtual const FText GetPaletteCategory() override;
-	virtual void ValidateCompiledDefaults(IWidgetCompilerLog& CompileLog) const override;
+	UMG_API virtual const FText GetPaletteCategory() override;
+	UMG_API virtual void ValidateCompiledDefaults(IWidgetCompilerLog& CompileLog) const override;
 #endif
 
 	TSubclassOf<UUserWidget> GetEntryWidgetClass() const { return EntryWidgetClass; }
 
 	/** Gets all of the list entry widgets currently being displayed by the list */
 	UFUNCTION(BlueprintCallable, Category = ListViewBase)
-	const TArray<UUserWidget*>& GetDisplayedEntryWidgets() const;
+	UMG_API const TArray<UUserWidget*>& GetDisplayedEntryWidgets() const;
 
 	/** Get the scroll offset of this view (in items) */
 	UFUNCTION(BlueprintCallable, Category = ListViewBase)
-	float GetScrollOffset() const;
+	UMG_API float GetScrollOffset() const;
 
 	/**
 	 * Full regeneration of all entries in the list. Note that the entry UWidget instances will not be destroyed, but they will be released and re-generated.
 	 * In other words, entry widgets will not receive Destruct/Construct events. They will receive OnEntryReleased and IUserObjectListEntry implementations will receive OnListItemObjectSet.
 	 */
 	UFUNCTION(BlueprintCallable, Category = ListViewBase)
-	void RegenerateAllEntries();
+	UMG_API void RegenerateAllEntries();
 
 	/** Scroll the entire list up to the first item */
 	UFUNCTION(BlueprintCallable, Category = ListViewBase)
-	void ScrollToTop();
+	UMG_API void ScrollToTop();
 
 	/** Scroll the entire list down to the bottom-most item */
 	UFUNCTION(BlueprintCallable, Category = ListViewBase)
-	void ScrollToBottom();
+	UMG_API void ScrollToBottom();
 
 	/** Set the scroll offset of this view (in items) */
 	UFUNCTION(BlueprintCallable, Category = ListView)
-	void SetScrollOffset(const float InScrollOffset);
+	UMG_API void SetScrollOffset(const float InScrollOffset);
 
 	UFUNCTION(BlueprintCallable, Category = ListViewBase)
-	void SetWheelScrollMultiplier(float NewWheelScrollMultiplier);
+	UMG_API void SetWheelScrollMultiplier(float NewWheelScrollMultiplier);
 
 	UFUNCTION(BlueprintCallable, Category = ListViewBase)
-	void SetScrollbarVisibility(ESlateVisibility InVisibility);
+	UMG_API void SetScrollbarVisibility(ESlateVisibility InVisibility);
 
 	/**
 	 * Sets the list to refresh on the next tick.
@@ -547,7 +547,7 @@ public:
 	 * This can be onerous to set up for simple cases, so it's also reasonable (though not ideal) to call RegenerateAllEntries when changes within N list items need to be reflected.
 	 */
 	UFUNCTION(BlueprintCallable, Category = ListViewBase)
-	void RequestRefresh();
+	UMG_API void RequestRefresh();
 
 	DECLARE_EVENT_OneParam(UListView, FOnListEntryGenerated, UUserWidget&);
 	FOnListEntryGenerated& OnEntryWidgetGenerated() { return OnListEntryGeneratedEvent; }
@@ -556,17 +556,17 @@ public:
 	FOnEntryWidgetReleased& OnEntryWidgetReleased() { return OnEntryWidgetReleasedEvent; }
 
 protected:
-	virtual TSharedRef<SWidget> RebuildWidget() override final;
-	virtual void ReleaseSlateResources(bool bReleaseChildren) override;
-	virtual void SynchronizeProperties() override;
+	UMG_API virtual TSharedRef<SWidget> RebuildWidget() override final;
+	UMG_API virtual void ReleaseSlateResources(bool bReleaseChildren) override;
+	UMG_API virtual void SynchronizeProperties() override;
 
 	/** Implement in child classes to construct the actual ListView Slate widget */
-	virtual TSharedRef<STableViewBase> RebuildListWidget();
+	UMG_API virtual TSharedRef<STableViewBase> RebuildListWidget();
 
 	//@todo DanH: Should probably have the events for native & BP built in up here - need to update existing binds to UListView's version
 	virtual void HandleListEntryHovered(UUserWidget& EntryWidget) {}
 	virtual void HandleListEntryUnhovered(UUserWidget& EntryWidget) {}
-	virtual	void FinishGeneratingEntry(UUserWidget& GeneratedEntry);
+	UMG_API virtual	void FinishGeneratingEntry(UUserWidget& GeneratedEntry);
    
     /** Called when a row widget is generated for a list item */
     UPROPERTY(BlueprintAssignable, Category = Events, meta = (DisplayName = "On Entry Generated"))
@@ -654,7 +654,7 @@ protected:
 #endif
 
 	/** Expected to be bound to the actual ListView widget created by a child class (automatically taken care of via the construction helpers within ITypedUMGListView) */
-	void HandleRowReleased(const TSharedRef<ITableRow>& Row);
+	UMG_API void HandleRowReleased(const TSharedRef<ITableRow>& Row);
 
 	// Note: Options for this property can be configured via class and property metadata. See class declaration comment above.
 	/** The type of widget to create for each entry displayed in the list. */
@@ -698,7 +698,7 @@ protected:
 	virtual void NativeOnEntryReleased(UUserWidget* EntryWidget) {};
 
 private:
-	virtual void HandleAnnounceGeneratedEntries();
+	UMG_API virtual void HandleAnnounceGeneratedEntries();
 
 #if WITH_EDITORONLY_DATA
 	bool bNeedsToCallRefreshDesignerItems = false;;
@@ -755,4 +755,4 @@ public:	\
 	virtual FOnFinishedScrolling& OnFinishedScrolling() const override { return OnFinishedScrollingEvent; }	\
 	virtual FOnItemExpansionChanged& OnItemExpansionChanged() const override { return OnItemExpansionChangedEvent; }	\
 	virtual FOnGetEntryClassForItem& OnGetEntryClassForItem() const override { return OnGetEntryClassForItemDelegate; } \
-	virtual FOnIsItemSelectableOrNavigable& OnIsItemSelectableOrNavigable() const override { return OnIsItemSelectableOrNavigableDelegate; }
\ No newline at end of file
+	virtual FOnIsItemSelectableOrNavigable& OnIsItemSelectableOrNavigable() const override { return OnIsItemSelectableOrNavigableDelegate; }
diff --git a/Engine/Source/Runtime/UMG/Public/Components/MenuAnchor.h b/Engine/Source/Runtime/UMG/Public/Components/MenuAnchor.h
--- a/Engine/Source/Runtime/UMG/Public/Components/MenuAnchor.h
+++ b/Engine/Source/Runtime/UMG/Public/Components/MenuAnchor.h
@@ -19,8 +19,8 @@ DECLARE_DYNAMIC_MULTICAST_DELEGATE_OneParam(FOnMenuOpenChangedEvent, bool, bIsOp
  * * Single Child
  * * Popup
  */
-UCLASS()
-class UMG_API UMenuAnchor : public UContentWidget
+UCLASS(MinimalAPI)
+class UMenuAnchor : public UContentWidget
 {
 	GENERATED_UCLASS_BODY()
 
@@ -73,18 +73,18 @@ public:
 	//TODO UMG Add Set MenuClass
 
 	UFUNCTION(BlueprintCallable, Category = "Menu Anchor")
-	void SetPlacement(EMenuPlacement InPlacement);
+	UMG_API void SetPlacement(EMenuPlacement InPlacement);
 
-	EMenuPlacement GetPlacement() const;
+	UMG_API EMenuPlacement GetPlacement() const;
 
 	UFUNCTION(BlueprintCallable, Category = "Menu Anchor")
-	void FitInWindow(bool bFit);
+	UMG_API void FitInWindow(bool bFit);
 
-	bool IsFitInWindow() const;
+	UMG_API bool IsFitInWindow() const;
 
-	bool IsDeferPaintingAfterWindowContent() const;
+	UMG_API bool IsDeferPaintingAfterWindowContent() const;
 
-	bool IsUseApplicationMenuStack() const;
+	UMG_API bool IsUseApplicationMenuStack() const;
 
 public:
 
@@ -94,19 +94,19 @@ public:
 	 * @param bFocusOnOpen  Should we focus the popup as soon as it opens?
 	 */
 	UFUNCTION(BlueprintCallable, Category="Menu Anchor")
-	void ToggleOpen(bool bFocusOnOpen);
+	UMG_API void ToggleOpen(bool bFocusOnOpen);
 
 	/** Opens the menu if it is not already open */
 	UFUNCTION(BlueprintCallable, Category="Menu Anchor")
-	void Open(bool bFocusMenu);
+	UMG_API void Open(bool bFocusMenu);
 
 	/** Closes the menu if it is currently open. */
 	UFUNCTION(BlueprintCallable, Category="Menu Anchor")
-	void Close();
+	UMG_API void Close();
 
 	/** Returns true if the popup is open; false otherwise. */
 	UFUNCTION(BlueprintCallable, Category="Menu Anchor")
-	bool IsOpen() const;
+	UMG_API bool IsOpen() const;
 
 	/**
 	 * Returns true if we should open the menu due to a click. Sometimes we should not, if
@@ -114,43 +114,43 @@ public:
 	 * happens to land on the button.
 	 */
 	UFUNCTION(BlueprintCallable, Category="Menu Anchor")
-	bool ShouldOpenDueToClick() const;
+	UMG_API bool ShouldOpenDueToClick() const;
 
 	/** Returns the current menu position */
 	UFUNCTION(BlueprintCallable, Category="Menu Anchor")
-	FVector2D GetMenuPosition() const;
+	UMG_API FVector2D GetMenuPosition() const;
 
 	/** Returns whether this menu has open submenus */
 	UFUNCTION(BlueprintCallable, Category="Menu Anchor")
-	bool HasOpenSubMenus() const;
+	UMG_API bool HasOpenSubMenus() const;
 
-	virtual void ReleaseSlateResources(bool bReleaseChildren) override;
+	UMG_API virtual void ReleaseSlateResources(bool bReleaseChildren) override;
 
 #if WITH_EDITOR
-	virtual const FText GetPaletteCategory() override;
+	UMG_API virtual const FText GetPaletteCategory() override;
 #endif
 
 protected:
 
 	// UPanelWidget
-	virtual void OnSlotAdded(UPanelSlot* Slot) override;
-	virtual void OnSlotRemoved(UPanelSlot* Slot) override;
+	UMG_API virtual void OnSlotAdded(UPanelSlot* Slot) override;
+	UMG_API virtual void OnSlotRemoved(UPanelSlot* Slot) override;
 	// End UPanelWidget
 
 	// UWidget interface
-	virtual TSharedRef<SWidget> RebuildWidget() override;
+	UMG_API virtual TSharedRef<SWidget> RebuildWidget() override;
 	// End of UWidget interface
 
 	// Initialize ShouldDeferPaintingAfterWindowContent in the constructor before the SWidget is constructed.
-	void InitShouldDeferPaintingAfterWindowContent(bool InShouldDeferPaintingAfterWindowContent);
+	UMG_API void InitShouldDeferPaintingAfterWindowContent(bool InShouldDeferPaintingAfterWindowContent);
 
 	// Initialize UseApplicationMenuStack in the constructor before the SWidget is constructed.
-	void InitUseApplicationMenuStack(bool InUseApplicationMenuStack);
+	UMG_API void InitUseApplicationMenuStack(bool InUseApplicationMenuStack);
 
 
 protected:
-	TSharedRef<SWidget> HandleGetMenuContent();
-	void HandleMenuOpenChanged(bool bIsOpen);
+	UMG_API TSharedRef<SWidget> HandleGetMenuContent();
+	UMG_API void HandleMenuOpenChanged(bool bIsOpen);
 
 protected:
 	TSharedPtr<SMenuAnchor> MyMenuAnchor;
diff --git a/Engine/Source/Runtime/UMG/Public/Components/MultiLineEditableText.h b/Engine/Source/Runtime/UMG/Public/Components/MultiLineEditableText.h
--- a/Engine/Source/Runtime/UMG/Public/Components/MultiLineEditableText.h
+++ b/Engine/Source/Runtime/UMG/Public/Components/MultiLineEditableText.h
@@ -17,8 +17,8 @@ class SMultiLineEditableText;
 /**
  * Editable text box widget
  */
-UCLASS(meta=( DisplayName="Editable Text (Multi-Line)" ))
-class UMG_API UMultiLineEditableText : public UTextLayoutWidget
+UCLASS(meta=( DisplayName="Editable Text (Multi-Line)" ), MinimalAPI)
+class UMultiLineEditableText : public UTextLayoutWidget
 {
 	GENERATED_UCLASS_BODY()
 
@@ -99,95 +99,95 @@ public:
 	* @return The widget text
 	*/
 	UFUNCTION(BlueprintCallable, Category="Widget", meta=(DisplayName="GetText (Multi-Line Editable Text)"))
-	FText GetText() const;
+	UMG_API FText GetText() const;
 
 	/**
 	* Directly sets the widget text.
 	* @param InText The text to assign to the widget
 	*/
 	UFUNCTION(BlueprintCallable, Category="Widget", meta=(DisplayName="SetText (Multi-Line Editable Text)"))
-	void SetText(FText InText);
+	UMG_API void SetText(FText InText);
 
 	/** Returns the Hint text that appears when there is no text in the text box */
 	UFUNCTION(BlueprintCallable, Category="Widget", meta=(DisplayName="GetHintText (Multi-Line Editable Text)"))
-	FText GetHintText() const;
+	UMG_API FText GetHintText() const;
 
 	/** 
 	* Sets the Hint text that appears when there is no text in the text box 
 	* @param InHintText The text that appears when there is no text in the text box 
 	*/
 	UFUNCTION(BlueprintCallable, Category="Widget", meta=(DisplayName="SetHintText (Multi-Line Editable Text)"))
-	void SetHintText(FText InHintText);
+	UMG_API void SetHintText(FText InHintText);
 
 	/** Set to true to select all text when the user clicks to give focus on the widget */
-	void SetSelectAllTextWhenFocused(bool bSelectAllTextWhenFocused);
+	UMG_API void SetSelectAllTextWhenFocused(bool bSelectAllTextWhenFocused);
 
 	/** Whether to select all text when the user clicks to give focus on the widget */
-	bool GetSelectAllTextWhenFocused() const;
+	UMG_API bool GetSelectAllTextWhenFocused() const;
 
 	/** Set to true to clear text selection when focus is lost */
-	void SetClearTextSelectionOnFocusLoss(bool bClearTextSelectionOnFocusLoss);
+	UMG_API void SetClearTextSelectionOnFocusLoss(bool bClearTextSelectionOnFocusLoss);
 
 	/** Whether to clear text selection when focus is lost */
-	bool GetClearTextSelectionOnFocusLoss() const;
+	UMG_API bool GetClearTextSelectionOnFocusLoss() const;
 
 	/** Set to true to allow the user to back out of changes when they press the escape key */
-	void SetRevertTextOnEscape(bool bRevertTextOnEscape);
+	UMG_API void SetRevertTextOnEscape(bool bRevertTextOnEscape);
 
 	/** Whether to allow the user to back out of changes when they press the escape key  */
-	bool GetRevertTextOnEscape() const;
+	UMG_API bool GetRevertTextOnEscape() const;
 
 	/** Set to true to clear keyboard focus when pressing enter to commit changes */
-	void SetClearKeyboardFocusOnCommit(bool bClearKeyboardFocusOnCommit);
+	UMG_API void SetClearKeyboardFocusOnCommit(bool bClearKeyboardFocusOnCommit);
 
 	/** Whether to clear keyboard focus when pressing enter to commit changes */
-	bool GetClearKeyboardFocusOnCommit() const;	
+	UMG_API bool GetClearKeyboardFocusOnCommit() const;	
 
 	/** Return true when this text cannot be modified interactively by the user */
-	bool GetIsReadOnly() const;
+	UMG_API bool GetIsReadOnly() const;
 
 	/** Sets the Text as Readonly to prevent it from being modified interactively by the user */
 	UFUNCTION(BlueprintCallable, Category="Widget", meta=(DisplayName="SetIsReadOnly (Multi-Line Editable Text"))
-	void SetIsReadOnly(bool bReadOnly);
+	UMG_API void SetIsReadOnly(bool bReadOnly);
 
 	UFUNCTION(BlueprintSetter)
-	void SetWidgetStyle(const FTextBlockStyle& InWidgetStyle);
+	UMG_API void SetWidgetStyle(const FTextBlockStyle& InWidgetStyle);
 
 	UFUNCTION(BlueprintCallable, Category = "Appearance")
-	const FSlateFontInfo& GetFont() const;
+	UMG_API const FSlateFontInfo& GetFont() const;
 
 	UFUNCTION(BlueprintCallable, Category = "Appearance")
-	void SetFont(FSlateFontInfo InFontInfo);
+	UMG_API void SetFont(FSlateFontInfo InFontInfo);
 
 	UFUNCTION(BlueprintCallable, Category = "Appearance")
-	void SetFontMaterial(UMaterialInterface* InMaterial);
+	UMG_API void SetFontMaterial(UMaterialInterface* InMaterial);
 
 	UFUNCTION(BlueprintCallable, Category = "Appearance")
-	void SetFontOutlineMaterial(UMaterialInterface* InMaterial);
+	UMG_API void SetFontOutlineMaterial(UMaterialInterface* InMaterial);
 
 	//~ Begin UTextLayoutWidget Interface
-	virtual void SetJustification(ETextJustify::Type InJustification) override;
+	UMG_API virtual void SetJustification(ETextJustify::Type InJustification) override;
 	//~ End UTextLayoutWidget Interface
 
 	//~ Begin UWidget Interface
-	virtual void SynchronizeProperties() override;
+	UMG_API virtual void SynchronizeProperties() override;
 	//~ End UWidget Interface
 
 	//~ Begin UVisual Interface
-	virtual void ReleaseSlateResources(bool bReleaseChildren) override;
+	UMG_API virtual void ReleaseSlateResources(bool bReleaseChildren) override;
 	//~ End UVisual Interface
 
 #if WITH_EDITOR
-	virtual const FText GetPaletteCategory() override;
+	UMG_API virtual const FText GetPaletteCategory() override;
 #endif
 
 protected:
 	//~ Begin UWidget Interface
-	virtual TSharedRef<SWidget> RebuildWidget() override;
+	UMG_API virtual TSharedRef<SWidget> RebuildWidget() override;
 	// End of UWidget
 
-	void HandleOnTextChanged(const FText& Text);
-	void HandleOnTextCommitted(const FText& Text, ETextCommit::Type CommitMethod);
+	UMG_API void HandleOnTextChanged(const FText& Text);
+	UMG_API void HandleOnTextCommitted(const FText& Text, ETextCommit::Type CommitMethod);
 
 protected:
 	TSharedPtr<SMultiLineEditableText> MyMultiLineEditableText;
diff --git a/Engine/Source/Runtime/UMG/Public/Components/MultiLineEditableTextBox.h b/Engine/Source/Runtime/UMG/Public/Components/MultiLineEditableTextBox.h
--- a/Engine/Source/Runtime/UMG/Public/Components/MultiLineEditableTextBox.h
+++ b/Engine/Source/Runtime/UMG/Public/Components/MultiLineEditableTextBox.h
@@ -17,8 +17,8 @@ class USlateWidgetStyleAsset;
 /**
  * Allows a user to enter multiple lines of text
  */
-UCLASS(meta=(DisplayName="Text Box (Multi-Line)"))
-class UMG_API UMultiLineEditableTextBox : public UTextLayoutWidget
+UCLASS(meta=(DisplayName="Text Box (Multi-Line)"), MinimalAPI)
+class UMultiLineEditableTextBox : public UTextLayoutWidget
 {
 	GENERATED_UCLASS_BODY()
 
@@ -90,7 +90,7 @@ public:
 	 * @return The widget text
 	 */
 	UFUNCTION(BlueprintCallable, Category="Widget", meta=(DisplayName="GetText (Multi-Line Text Box)"))
-	FText GetText() const;
+	UMG_API FText GetText() const;
 
 	/**
 	 * Directly sets the widget text.
@@ -98,34 +98,34 @@ public:
 	 * @param InText The text to assign to the widget
 	 */
 	UFUNCTION(BlueprintCallable, Category="Widget", meta=(DisplayName="SetText (Multi-Line Text Box)"))
-	void SetText(FText InText);
+	UMG_API void SetText(FText InText);
 
 	/** Returns the Hint text that appears when there is no text in the text box */
 	UFUNCTION(BlueprintCallable, Category = "Widget", meta = (DisplayName = "GetHintText (Multi-Line Text Box)"))
-	FText GetHintText() const;
+	UMG_API FText GetHintText() const;
 
 	/**
 	* Sets the Hint text that appears when there is no text in the text box
 	* @param InHintText The text that appears when there is no text in the text box
 	*/
 	UFUNCTION(BlueprintCallable, Category="Widget", meta=(DisplayName="SetHintText (Multi-Line Text Box)"))
-	void SetHintText(FText InHintText);
+	UMG_API void SetHintText(FText InHintText);
 
 	UFUNCTION(BlueprintCallable, Category="Widget", meta=(DisplayName="SetError (Multi-Line Text Box)"))
-	void SetError(FText InError);
+	UMG_API void SetError(FText InError);
 
 	/** Return true when this text cannot be modified interactively by the user */
-	bool GetIsReadOnly() const;
+	UMG_API bool GetIsReadOnly() const;
 
 	/** Sets the Text as Readonly to prevent it from being modified interactively by the user */
 	UFUNCTION(BlueprintCallable, Category = "Widget", meta = (DisplayName = "SetIsReadOnly (Multi-Line Text Box)"))
-	void SetIsReadOnly(UPARAM(DisplayName = "ReadyOnly") bool bReadOnly);
+	UMG_API void SetIsReadOnly(UPARAM(DisplayName = "ReadyOnly") bool bReadOnly);
 
 	UFUNCTION(BlueprintCallable, Category="Widget", meta=(DisplayName="SetTextStyle (Multi-Line Text Box)"))
-	void SetTextStyle(const FTextBlockStyle& InTextStyle);
+	UMG_API void SetTextStyle(const FTextBlockStyle& InTextStyle);
 
 	UFUNCTION(BlueprintCallable, Category="Widget", meta=(DisplayName="SetForegroundColor (Multi-Line Text Box)"))
-	void SetForegroundColor(FLinearColor color);
+	UMG_API void SetForegroundColor(FLinearColor color);
 
 	//TODO UMG Add Set ReadOnlyForegroundColor
 	//TODO UMG Add Set BackgroundColor
@@ -133,29 +133,29 @@ public:
 
 public:
 	//~ Begin UTextLayoutWidget Interface
-	virtual void SetJustification(ETextJustify::Type InJustification) override;
+	UMG_API virtual void SetJustification(ETextJustify::Type InJustification) override;
 	//~ End UTextLayoutWidget Interface
 
 	//~ Begin UWidget Interface
-	virtual void SynchronizeProperties() override;
+	UMG_API virtual void SynchronizeProperties() override;
 	//~ End UWidget Interface
 
 	//~ Begin UVisual Interface
-	virtual void ReleaseSlateResources(bool bReleaseChildren) override;
+	UMG_API virtual void ReleaseSlateResources(bool bReleaseChildren) override;
 	//~ End UVisual Interface
 
 #if WITH_EDITOR
-	virtual const FText GetPaletteCategory() override;
+	UMG_API virtual const FText GetPaletteCategory() override;
 #endif
-	virtual void Serialize(FArchive& Ar) override;
+	UMG_API virtual void Serialize(FArchive& Ar) override;
 
 protected:
 	//~ Begin UWidget Interface
-	virtual TSharedRef<SWidget> RebuildWidget() override;
+	UMG_API virtual TSharedRef<SWidget> RebuildWidget() override;
 	// End of UWidget
 
-	void HandleOnTextChanged(const FText& Text);
-	void HandleOnTextCommitted(const FText& Text, ETextCommit::Type CommitMethod);
+	UMG_API void HandleOnTextChanged(const FText& Text);
+	UMG_API void HandleOnTextCommitted(const FText& Text, ETextCommit::Type CommitMethod);
 
 protected:
 	TSharedPtr<SMultiLineEditableTextBox> MyEditableTextBlock;
diff --git a/Engine/Source/Runtime/UMG/Public/Components/NamedSlot.h b/Engine/Source/Runtime/UMG/Public/Components/NamedSlot.h
--- a/Engine/Source/Runtime/UMG/Public/Components/NamedSlot.h
+++ b/Engine/Source/Runtime/UMG/Public/Components/NamedSlot.h
@@ -14,25 +14,25 @@ class SBox;
  * Allows you to expose an external slot for your user widget.  When others reuse your user control, 
  * they can put whatever they want in this named slot.
  */
-UCLASS()
-class UMG_API UNamedSlot : public UContentWidget
+UCLASS(MinimalAPI)
+class UNamedSlot : public UContentWidget
 {
 	GENERATED_UCLASS_BODY()
 
 public:
 
 	// UPanelWidget interface
-	virtual void OnSlotAdded(UPanelSlot* Slot) override;
-	virtual void OnSlotRemoved(UPanelSlot* Slot) override;
+	UMG_API virtual void OnSlotAdded(UPanelSlot* Slot) override;
+	UMG_API virtual void OnSlotRemoved(UPanelSlot* Slot) override;
 	// End of UPanelWidget interface
 
 	// UVisual interface
-	virtual void ReleaseSlateResources(bool bReleaseChildren) override;
+	UMG_API virtual void ReleaseSlateResources(bool bReleaseChildren) override;
 	// End of UVisual interface
 
 #if WITH_EDITOR
-	virtual const FText GetPaletteCategory() override;
-	FGuid GetSlotGUID() const;
+	UMG_API virtual const FText GetPaletteCategory() override;
+	UMG_API FGuid GetSlotGUID() const;
 #endif
 
 public:
@@ -47,12 +47,12 @@ public:
 	bool bExposeOnInstanceOnly = false;
 #endif
 	
-	virtual void Serialize(FArchive& Ar) override;
-	virtual void PostLoad() override;
+	UMG_API virtual void Serialize(FArchive& Ar) override;
+	UMG_API virtual void PostLoad() override;
 	
 protected:
 	// UWidget interface
-	virtual TSharedRef<SWidget> RebuildWidget() override;
+	UMG_API virtual TSharedRef<SWidget> RebuildWidget() override;
 	// End of UWidget interface
 
 protected:
diff --git a/Engine/Source/Runtime/UMG/Public/Components/NamedSlotInterface.h b/Engine/Source/Runtime/UMG/Public/Components/NamedSlotInterface.h
--- a/Engine/Source/Runtime/UMG/Public/Components/NamedSlotInterface.h
+++ b/Engine/Source/Runtime/UMG/Public/Components/NamedSlotInterface.h
@@ -13,13 +13,13 @@ class UWidget;
 /**
  * 
  */
-UINTERFACE(meta=( CannotImplementInterfaceInBlueprint ))
-class UMG_API UNamedSlotInterface : public UInterface
+UINTERFACE(meta=( CannotImplementInterfaceInBlueprint ), MinimalAPI)
+class UNamedSlotInterface : public UInterface
 {
 	GENERATED_UINTERFACE_BODY()
 };
 
-class UMG_API INamedSlotInterface
+class INamedSlotInterface
 {
 	GENERATED_IINTERFACE_BODY()
 
@@ -35,16 +35,16 @@ public:
 	virtual void SetContentForSlot(FName SlotName, UWidget* Content) = 0;
 
 	/** Determines if any slot holds the given widget. */
-	bool ContainsContent(UWidget* Content) const;
+	UMG_API bool ContainsContent(UWidget* Content) const;
 
 	/** Determines if any slot holds the given widget and the name of that slot. */
-	FName FindSlotForContent(UWidget* Content) const;
+	UMG_API FName FindSlotForContent(UWidget* Content) const;
 
 	/** Releases named slot related resources. */
-	void ReleaseNamedSlotSlateResources(bool bReleaseChildren);
+	UMG_API void ReleaseNamedSlotSlateResources(bool bReleaseChildren);
 
 #if WITH_EDITOR
 	/** Applies the designer flags to the all of the content in all the slots. */
-	void SetNamedSlotDesignerFlags(EWidgetDesignFlags NewFlags);
+	UMG_API void SetNamedSlotDesignerFlags(EWidgetDesignFlags NewFlags);
 #endif
 };
diff --git a/Engine/Source/Runtime/UMG/Public/Components/NativeWidgetHost.h b/Engine/Source/Runtime/UMG/Public/Components/NativeWidgetHost.h
--- a/Engine/Source/Runtime/UMG/Public/Components/NativeWidgetHost.h
+++ b/Engine/Source/Runtime/UMG/Public/Components/NativeWidgetHost.h
@@ -12,26 +12,26 @@
  * A NativeWidgetHost is a container widget that can contain one child slate widget.  This should
  * be used when all you need is to nest a native widget inside a UMG widget.
  */
-UCLASS()
-class UMG_API UNativeWidgetHost : public UWidget
+UCLASS(MinimalAPI)
+class UNativeWidgetHost : public UWidget
 {
 	GENERATED_UCLASS_BODY()
 
-	void SetContent(TSharedRef<SWidget> InContent);
+	UMG_API void SetContent(TSharedRef<SWidget> InContent);
 	TSharedPtr< SWidget > GetContent() const { return NativeWidget; }
 
-	virtual void ReleaseSlateResources(bool bReleaseChildren) override;
+	UMG_API virtual void ReleaseSlateResources(bool bReleaseChildren) override;
 
 #if WITH_EDITOR
-	virtual const FText GetPaletteCategory() override;
+	UMG_API virtual const FText GetPaletteCategory() override;
 #endif
 
 protected:
 	// UWidget interface
-	virtual TSharedRef<SWidget> RebuildWidget() override;
+	UMG_API virtual TSharedRef<SWidget> RebuildWidget() override;
 	// End of UWidget interface
 
-	TSharedRef<SWidget> GetDefaultContent();
+	UMG_API TSharedRef<SWidget> GetDefaultContent();
 
 protected:
 	TSharedPtr<SWidget> NativeWidget;
diff --git a/Engine/Source/Runtime/UMG/Public/Components/Overlay.h b/Engine/Source/Runtime/UMG/Public/Components/Overlay.h
--- a/Engine/Source/Runtime/UMG/Public/Components/Overlay.h
+++ b/Engine/Source/Runtime/UMG/Public/Components/Overlay.h
@@ -14,31 +14,31 @@ class UOverlaySlot;
 /**
  * Allows widgets to be stacked on top of each other, uses simple flow layout for content on each layer.
  */
-UCLASS()
-class UMG_API UOverlay : public UPanelWidget
+UCLASS(MinimalAPI)
+class UOverlay : public UPanelWidget
 {
 	GENERATED_UCLASS_BODY()
 
-	virtual void ReleaseSlateResources(bool bReleaseChildren) override;
+	UMG_API virtual void ReleaseSlateResources(bool bReleaseChildren) override;
 
 	/**  */
 	UFUNCTION(BlueprintCallable, Category="Widget")
-	UOverlaySlot* AddChildToOverlay(UWidget* Content);
+	UMG_API UOverlaySlot* AddChildToOverlay(UWidget* Content);
 
 	/** Replace the widget at the given index it with a different widget. */
 	UFUNCTION(BlueprintCallable, Category = "Panel")
-	bool ReplaceOverlayChildAt(int32 Index, UWidget* Content);
+	UMG_API bool ReplaceOverlayChildAt(int32 Index, UWidget* Content);
 
 #if WITH_EDITOR
-	virtual const FText GetPaletteCategory() override;
+	UMG_API virtual const FText GetPaletteCategory() override;
 #endif
 
 protected:
 
 	// UPanelWidget
-	virtual UClass* GetSlotClass() const override;
-	virtual void OnSlotAdded(UPanelSlot* Slot) override;
-	virtual void OnSlotRemoved(UPanelSlot* Slot) override;
+	UMG_API virtual UClass* GetSlotClass() const override;
+	UMG_API virtual void OnSlotAdded(UPanelSlot* Slot) override;
+	UMG_API virtual void OnSlotRemoved(UPanelSlot* Slot) override;
 	// End UPanelWidget
 
 protected:
@@ -47,6 +47,6 @@ protected:
 
 protected:
 	// UWidget interface
-	virtual TSharedRef<SWidget> RebuildWidget() override;
+	UMG_API virtual TSharedRef<SWidget> RebuildWidget() override;
 	// End of UWidget interface
 };
diff --git a/Engine/Source/Runtime/UMG/Public/Components/OverlaySlot.h b/Engine/Source/Runtime/UMG/Public/Components/OverlaySlot.h
--- a/Engine/Source/Runtime/UMG/Public/Components/OverlaySlot.h
+++ b/Engine/Source/Runtime/UMG/Public/Components/OverlaySlot.h
@@ -13,8 +13,8 @@
 #include "OverlaySlot.generated.h"
 
 /** Slot for the UOverlay panel.  Allows content to be hover above other content. */
-UCLASS()
-class UMG_API UOverlaySlot : public UPanelSlot
+UCLASS(MinimalAPI)
+class UOverlaySlot : public UPanelSlot
 {
 	GENERATED_UCLASS_BODY()
 
@@ -45,36 +45,36 @@ public:
 
 	/** Set padding area between the slot and the content it contains. */
 	UFUNCTION(BlueprintCallable, Category="Layout|Overlay Slot")
-	void SetPadding(FMargin InPadding);
+	UMG_API void SetPadding(FMargin InPadding);
 
 	/** Get padding area between the slot and the content it contains. */
-	FMargin GetPadding() const;
+	UMG_API FMargin GetPadding() const;
 
 	/** Set the alignment of the object horizontally. */
 	UFUNCTION(BlueprintCallable, Category="Layout|Overlay Slot")
-	void SetHorizontalAlignment(EHorizontalAlignment InHorizontalAlignment);
+	UMG_API void SetHorizontalAlignment(EHorizontalAlignment InHorizontalAlignment);
 	
 	/** Get the alignment of the object horizontally. */
-	EHorizontalAlignment GetHorizontalAlignment() const;
+	UMG_API EHorizontalAlignment GetHorizontalAlignment() const;
 
 	/** Set the alignment of the object vertically. */
 	UFUNCTION(BlueprintCallable, Category="Layout|Overlay Slot")
-	void SetVerticalAlignment(EVerticalAlignment InVerticalAlignment);
+	UMG_API void SetVerticalAlignment(EVerticalAlignment InVerticalAlignment);
 
 	/** Get the alignment of the object vertically. */
-	EVerticalAlignment GetVerticalAlignment() const;
+	UMG_API EVerticalAlignment GetVerticalAlignment() const;
 
 public:
 
 	// UPanelSlot interface
-	virtual void SynchronizeProperties() override;
+	UMG_API virtual void SynchronizeProperties() override;
 	// End of UPanelSlot interface
 
 	/** Builds the underlying FSlot for the Slate layout panel. */
-	virtual void BuildSlot(TSharedRef<SOverlay> InOverlay);
+	UMG_API virtual void BuildSlot(TSharedRef<SOverlay> InOverlay);
 
 	/** Replace the slot content. */
-	void ReplaceContent(UWidget* Content);
+	UMG_API void ReplaceContent(UWidget* Content);
 
-	virtual void ReleaseSlateResources(bool bReleaseChildren) override;
+	UMG_API virtual void ReleaseSlateResources(bool bReleaseChildren) override;
 };
diff --git a/Engine/Source/Runtime/UMG/Public/Components/PanelSlot.h b/Engine/Source/Runtime/UMG/Public/Components/PanelSlot.h
--- a/Engine/Source/Runtime/UMG/Public/Components/PanelSlot.h
+++ b/Engine/Source/Runtime/UMG/Public/Components/PanelSlot.h
@@ -8,8 +8,8 @@
 #include "PanelSlot.generated.h"
 
 /** The base class for all Slots in UMG. */
-UCLASS(BlueprintType)
-class UMG_API UPanelSlot : public UVisual
+UCLASS(BlueprintType, MinimalAPI)
+class UPanelSlot : public UVisual
 {
 	GENERATED_UCLASS_BODY()
 
@@ -22,12 +22,12 @@ public:
 	TObjectPtr<class UWidget> Content;
 	
 #if WITH_EDITOR
-	bool IsDesignTime() const;
+	UMG_API bool IsDesignTime() const;
 #else
 	FORCEINLINE bool IsDesignTime() const { return false; }
 #endif
 
-	virtual void ReleaseSlateResources(bool bReleaseChildren) override;
+	UMG_API virtual void ReleaseSlateResources(bool bReleaseChildren) override;
 
 	/** Applies all properties to the live slot if possible. */
 	virtual void SynchronizeProperties()
diff --git a/Engine/Source/Runtime/UMG/Public/Components/PanelWidget.h b/Engine/Source/Runtime/UMG/Public/Components/PanelWidget.h
--- a/Engine/Source/Runtime/UMG/Public/Components/PanelWidget.h
+++ b/Engine/Source/Runtime/UMG/Public/Components/PanelWidget.h
@@ -10,8 +10,8 @@
 #include "PanelWidget.generated.h"
 
 /** The base class for all UMG panel widgets.  Panel widgets layout a collection of child widgets. */
-UCLASS(Abstract)
-class UMG_API UPanelWidget : public UWidget
+UCLASS(Abstract, MinimalAPI)
+class UPanelWidget : public UWidget
 {
 	GENERATED_UCLASS_BODY()
 
@@ -25,7 +25,7 @@ public:
 
 	/** Gets number of child widgets in the container. */
 	UFUNCTION(BlueprintCallable, Category="Widget|Panel")
-	int32 GetChildrenCount() const;
+	UMG_API int32 GetChildrenCount() const;
 
 	/**
 	 * Gets the widget at an index.
@@ -33,30 +33,30 @@ public:
 	 * @return The widget at the given index, or nothing if there is no widget there.
 	 */
 	UFUNCTION(BlueprintCallable, Category="Widget|Panel")
-	UWidget* GetChildAt(int32 Index) const;
+	UMG_API UWidget* GetChildAt(int32 Index) const;
 	
 	/** Gets all widgets in the container */
 	UFUNCTION(BlueprintCallable, Category="Widget|Panel")
-	TArray<UWidget*> GetAllChildren() const;
+	UMG_API TArray<UWidget*> GetAllChildren() const;
 
 	/** Gets the index of a specific child widget */
 	UFUNCTION(BlueprintCallable, Category="Widget|Panel")
-	int32 GetChildIndex(const UWidget* Content) const;
+	UMG_API int32 GetChildIndex(const UWidget* Content) const;
 
 	/** Returns true if panel contains this widget */
 	UFUNCTION(BlueprintCallable, Category="Widget|Panel")
-	bool HasChild(UWidget* Content) const;
+	UMG_API bool HasChild(UWidget* Content) const;
 
 	/** Removes a child by it's index. */
 	UFUNCTION(BlueprintCallable, Category="Widget|Panel")
-	bool RemoveChildAt(int32 Index);
+	UMG_API bool RemoveChildAt(int32 Index);
 
 	/**
 	 * Adds a new child widget to the container.  Returns the base slot type, 
 	 * requires casting to turn it into the type specific to the container.
 	 */
 	UFUNCTION(BlueprintCallable, Category="Widget|Panel")
-	UPanelSlot* AddChild(UWidget* Content);
+	UMG_API UPanelSlot* AddChild(UWidget* Content);
 
 #if WITH_EDITOR
 
@@ -67,7 +67,7 @@ public:
 	 * @param Content The content to replace the child with.
 	 * @return true if the index existed and the child could be replaced.
 	 */
-	bool ReplaceChildAt(int32 Index, UWidget* Content);
+	UMG_API bool ReplaceChildAt(int32 Index, UWidget* Content);
 
 	/**
 	 * Swaps the child widget out of the slot, and replaces it with the new child widget.
@@ -75,20 +75,20 @@ public:
 	 * @param NewChild The new child widget being inserted.
 	 * @return true if the CurrentChild was found and the swap occurred, otherwise false.
 	 */
-	virtual bool ReplaceChild(UWidget* CurrentChild, UWidget* NewChild);
+	UMG_API virtual bool ReplaceChild(UWidget* CurrentChild, UWidget* NewChild);
 
 	/**
 	 * Inserts a widget at a specific index.  This does not update the live slate version, it requires
 	 * a rebuild of the whole UI to see a change.
 	 */
-	UPanelSlot* InsertChildAt(int32 Index, UWidget* Content);
+	UMG_API UPanelSlot* InsertChildAt(int32 Index, UWidget* Content);
 
 	/**
 	 * Moves the child widget from its current index to the new index provided.
 	 */
-	void ShiftChild(int32 Index, UWidget* Child);
+	UMG_API void ShiftChild(int32 Index, UWidget* Child);
 
-	virtual void SetDesignerFlags(EWidgetDesignFlags NewFlags) override;
+	UMG_API virtual void SetDesignerFlags(EWidgetDesignFlags NewFlags) override;
 
 #endif
 
@@ -97,18 +97,18 @@ public:
 	 * @return true if the widget was found and removed.
 	 */
 	UFUNCTION(BlueprintCallable, Category="Widget|Panel")
-	bool RemoveChild(UWidget* Content);
+	UMG_API bool RemoveChild(UWidget* Content);
 
 	/** Returns true if there are any child widgets in the panel */
 	UFUNCTION(BlueprintCallable, Category="Widget|Panel")
-	bool HasAnyChildren() const;
+	UMG_API bool HasAnyChildren() const;
 
 	/** Remove all child widgets from the panel widget. */
 	UFUNCTION(BlueprintCallable, Category="Widget|Panel")
-	virtual void ClearChildren();
+	UMG_API virtual void ClearChildren();
 
 	/** The slots in the widget holding the child widgets of this panel. */
-	const TArray<UPanelSlot*>& GetSlots() const;
+	UMG_API const TArray<UPanelSlot*>& GetSlots() const;
 
 	/** Returns true if the panel supports more than one child. */
 	bool CanHaveMultipleChildren() const
@@ -122,7 +122,7 @@ public:
 		return CanHaveMultipleChildren() || GetChildrenCount() == 0;
 	}
 
-	virtual void ReleaseSlateResources(bool bReleaseChildren) override;
+	UMG_API virtual void ReleaseSlateResources(bool bReleaseChildren) override;
 
 #if WITH_EDITOR
 	virtual bool LockToPanelOnDrag() const
@@ -144,13 +144,13 @@ public:
 #endif
 
 	// Begin UObject
-	virtual void PostLoad() override;
+	UMG_API virtual void PostLoad() override;
 	// End UObject
 
 protected:
 
 #if WITH_EDITOR
-	virtual TSharedRef<SWidget> RebuildDesignWidget(TSharedRef<SWidget> Content) override;
+	UMG_API virtual TSharedRef<SWidget> RebuildDesignWidget(TSharedRef<SWidget> Content) override;
 #endif
 
 	virtual UClass* GetSlotClass() const
diff --git a/Engine/Source/Runtime/UMG/Public/Components/ProgressBar.h b/Engine/Source/Runtime/UMG/Public/Components/ProgressBar.h
--- a/Engine/Source/Runtime/UMG/Public/Components/ProgressBar.h
+++ b/Engine/Source/Runtime/UMG/Public/Components/ProgressBar.h
@@ -17,8 +17,8 @@ class USlateBrushAsset;
  *
  * * No Children
  */
-UCLASS()
-class UMG_API UProgressBar : public UWidget
+UCLASS(MinimalAPI)
+class UProgressBar : public UWidget
 {
 	GENERATED_UCLASS_BODY()
 	
@@ -69,64 +69,64 @@ public:
 public:
 	
 	/** */
-	const FProgressBarStyle& GetWidgetStyle() const;
+	UMG_API const FProgressBarStyle& GetWidgetStyle() const;
 
 	/**/
-	void SetWidgetStyle(const FProgressBarStyle& InStyle);
+	UMG_API void SetWidgetStyle(const FProgressBarStyle& InStyle);
 
 	/** */
-	float GetPercent() const;
+	UMG_API float GetPercent() const;
 
 	/** Sets the current value of the ProgressBar. */
 	UFUNCTION(BlueprintCallable, Category="Progress")
-	void SetPercent(float InPercent);
+	UMG_API void SetPercent(float InPercent);
 
 	/** */
-	EProgressBarFillType::Type GetBarFillType() const;
+	UMG_API EProgressBarFillType::Type GetBarFillType() const;
 
 	/** */
-	void SetBarFillType(EProgressBarFillType::Type InBarFillType);
+	UMG_API void SetBarFillType(EProgressBarFillType::Type InBarFillType);
 
 	/** */
-	EProgressBarFillStyle::Type GetBarFillStyle() const;
+	UMG_API EProgressBarFillStyle::Type GetBarFillStyle() const;
 
 	/** */
-	void SetBarFillStyle(EProgressBarFillStyle::Type InBarFillStyle);
+	UMG_API void SetBarFillStyle(EProgressBarFillStyle::Type InBarFillStyle);
 
 	/** */
-	bool UseMarquee() const;
+	UMG_API bool UseMarquee() const;
 
 	/** Sets the progress bar to show as a marquee. */
 	UFUNCTION(BlueprintCallable, Category="Progress")
-	void SetIsMarquee(bool InbIsMarquee);
+	UMG_API void SetIsMarquee(bool InbIsMarquee);
 
 	/** */
-	FVector2D GetBorderPadding() const;
+	UMG_API FVector2D GetBorderPadding() const;
 
 	/** */
-	void SetBorderPadding(FVector2D InBorderPadding);
+	UMG_API void SetBorderPadding(FVector2D InBorderPadding);
 
 	/** */
-	FLinearColor GetFillColorAndOpacity() const;
+	UMG_API FLinearColor GetFillColorAndOpacity() const;
 
 	/** Sets the fill color of the progress bar. */
 	UFUNCTION(BlueprintCallable, Category="Progress")
-	void SetFillColorAndOpacity(FLinearColor InColor);
+	UMG_API void SetFillColorAndOpacity(FLinearColor InColor);
 
 public:
 	
 	//~ Begin UWidget Interface
-	virtual void SynchronizeProperties() override;
+	UMG_API virtual void SynchronizeProperties() override;
 	//~ End UWidget Interface
 
 	//~ Begin UVisual Interface
-	virtual void ReleaseSlateResources(bool bReleaseChildren) override;
+	UMG_API virtual void ReleaseSlateResources(bool bReleaseChildren) override;
 	//~ End UVisual Interface
 
 #if WITH_EDITOR
 	//~ Begin UWidget Interface
-	virtual const FText GetPaletteCategory() override;
-	virtual void OnCreationFromPalette() override;
+	UMG_API virtual const FText GetPaletteCategory() override;
+	UMG_API virtual void OnCreationFromPalette() override;
 	//~ End UWidget Interface
 #endif
 
@@ -135,7 +135,7 @@ protected:
 	TSharedPtr<SProgressBar> MyProgressBar;
 
 	//~ Begin UWidget Interface
-	virtual TSharedRef<SWidget> RebuildWidget() override;
+	UMG_API virtual TSharedRef<SWidget> RebuildWidget() override;
 	//~ End UWidget Interface
 
 	PROPERTY_BINDING_IMPLEMENTATION(FSlateColor, FillColorAndOpacity);
diff --git a/Engine/Source/Runtime/UMG/Public/Components/RetainerBox.h b/Engine/Source/Runtime/UMG/Public/Components/RetainerBox.h
--- a/Engine/Source/Runtime/UMG/Public/Components/RetainerBox.h
+++ b/Engine/Source/Runtime/UMG/Public/Components/RetainerBox.h
@@ -22,8 +22,8 @@ class UMaterialInterface;
  * * Single Child
  * * Caching / Performance
  */
-UCLASS()
-class UMG_API URetainerBox : public UContentWidget
+UCLASS(MinimalAPI)
+class URetainerBox : public UContentWidget
 {
 	GENERATED_UCLASS_BODY()
 
@@ -77,70 +77,70 @@ public:
 	 * Requests the retainer redrawn the contents it has.
 	 */
 	UFUNCTION(BlueprintCallable, Category="Retainer")
-	void SetRenderingPhase(int32 RenderPhase, int32 TotalPhases);
+	UMG_API void SetRenderingPhase(int32 RenderPhase, int32 TotalPhases);
 
 	/**
 	 * Requests the retainer redrawn the contents it has.
 	 */
 	UFUNCTION(BlueprintCallable, Category="Retainer")
-	void RequestRender();
+	UMG_API void RequestRender();
 
 	/**
 	 * Get the current dynamic effect material applied to the retainer box.
 	 */
 	UFUNCTION(BlueprintCallable, Category="Retainer|Effect")
-	UMaterialInstanceDynamic* GetEffectMaterial() const;
+	UMG_API UMaterialInstanceDynamic* GetEffectMaterial() const;
 
 	/**
 	 * Set a new effect material to the retainer widget.
 	 */
 	UFUNCTION(BlueprintCallable, Category="Retainer|Effect")
-	void SetEffectMaterial(UMaterialInterface* EffectMaterial);
+	UMG_API void SetEffectMaterial(UMaterialInterface* EffectMaterial);
 
 	/**
 	 * Sets the name of the texture parameter to set the render target to on the material.
 	 */
 	UFUNCTION(BlueprintCallable, Category="Retainer|Effect")
-	void SetTextureParameter(FName TextureParameter);
+	UMG_API void SetTextureParameter(FName TextureParameter);
 
 	/**
 	* Set the flag for if we retain the render or pass-through
 	*/
 	UFUNCTION(BlueprintCallable, Category = "Retainer")
-	void SetRetainRendering(bool bInRetainRendering);
+	UMG_API void SetRetainRendering(bool bInRetainRendering);
 
 	/**
 	 * Get the flag for if we retain the render or pass-through.
 	 */
-	bool IsRetainRendering() const;
+	UMG_API bool IsRetainRendering() const;
 
 	/**
 	 * Get the phase to render on.
 	 */
-	int32 GetPhase() const;
+	UMG_API int32 GetPhase() const;
 
 	/**
 	 * Get the total number of phases.
 	 */
-	int32 GetPhaseCount() const;
+	UMG_API int32 GetPhaseCount() const;
 
 	/**
 	 * Get whether this widget should redraw the contents it has every time it receives an invalidation request.
 	 */
-	bool IsRenderOnInvalidation() const;
+	UMG_API bool IsRenderOnInvalidation() const;
 
 	/**
 	 * Get whether this widget should redraw the contents it has every time the phase occurs.
 	 */
-	bool IsRenderOnPhase() const;
+	UMG_API bool IsRenderOnPhase() const;
 
-	virtual void ReleaseSlateResources(bool bReleaseChildren) override;
+	UMG_API virtual void ReleaseSlateResources(bool bReleaseChildren) override;
 
 #if WITH_EDITOR
-	virtual const FText GetPaletteCategory() override;
+	UMG_API virtual const FText GetPaletteCategory() override;
 #endif
 
-	FGeometry GetCachedAllottedGeometry() const;
+	UMG_API FGeometry GetCachedAllottedGeometry() const;
 
 protected:
 
@@ -172,42 +172,42 @@ protected:
 #endif
 
 	//~ Begin UPanelWidget interface
-	virtual void OnSlotAdded(UPanelSlot* Slot) override;
-	virtual void OnSlotRemoved(UPanelSlot* Slot) override;
+	UMG_API virtual void OnSlotAdded(UPanelSlot* Slot) override;
+	UMG_API virtual void OnSlotRemoved(UPanelSlot* Slot) override;
 	//~ End UPanelWidget interface
 
 	//~ Begin UWidget interface
-	virtual TSharedRef<SWidget> RebuildWidget() override;
-	virtual void SynchronizeProperties() override;
+	UMG_API virtual TSharedRef<SWidget> RebuildWidget() override;
+	UMG_API virtual void SynchronizeProperties() override;
 	//~ End of UWidget interface
 
 	//~ Begin UObject interface
 #if WITH_EDITOR
-	virtual bool CanEditChange(const FProperty* InProperty) const override;
+	UMG_API virtual bool CanEditChange(const FProperty* InProperty) const override;
 #endif
 	//~ End UObject interface
 
 	// Initialize RenderOnInvalidation in the constructor before the SWidget is constructed.
-	void InitRenderOnInvalidation(bool InRenderOnInvalidation);
+	UMG_API void InitRenderOnInvalidation(bool InRenderOnInvalidation);
 
 	// Initialize RenderOnPhase in the constructor before the SWidget is constructed.
-	void InitRenderOnPhase(bool InRenderOnPhase);
+	UMG_API void InitRenderOnPhase(bool InRenderOnPhase);
 
 	// Initialize Phase in the constructor before the SWidget is constructed.
-	void InitPhase(int32 InPhase);
+	UMG_API void InitPhase(int32 InPhase);
 
 	// Initialize PhaseCount in the constructor before the SWidget is constructed.
-	void InitPhaseCount(int32 InPhaseCount);
+	UMG_API void InitPhaseCount(int32 InPhaseCount);
 public:
 	/**
 	* Gets the name of the texture parameter to set the render target to on the material.
 	*/
-	const FName& GetTextureParameter() const;
+	UMG_API const FName& GetTextureParameter() const;
 
 	/**
 	* Gets the current dynamic effect material applied to the retainer box.
 	*/
-	const UMaterialInterface* GetEffectMaterialInterface() const;
+	UMG_API const UMaterialInterface* GetEffectMaterialInterface() const;
 
 protected:
 	TSharedPtr<class SRetainerWidget> MyRetainerWidget;
