[Engine]

diff --git a/Engine/Source/Runtime/UMG/Public/Components/RichTextBlock.h b/Engine/Source/Runtime/UMG/Public/Components/RichTextBlock.h
--- a/Engine/Source/Runtime/UMG/Public/Components/RichTextBlock.h
+++ b/Engine/Source/Runtime/UMG/Public/Components/RichTextBlock.h
@@ -35,8 +35,8 @@ public:
  * * Fancy Text
  * * No Children
  */
-UCLASS()
-class UMG_API URichTextBlock : public UTextLayoutWidget
+UCLASS(MinimalAPI)
+class URichTextBlock : public UTextLayoutWidget
 {
 	GENERATED_BODY()
 protected:
@@ -84,7 +84,7 @@ public:
 	 * @param InColorAndOpacity		The new text color and opacity
 	 */
 	UFUNCTION(BlueprintCallable, Category="Appearance")
-	void SetDefaultColorAndOpacity(FSlateColor InColorAndOpacity);
+	UMG_API void SetDefaultColorAndOpacity(FSlateColor InColorAndOpacity);
 
 	/**  
 	 * Sets the color and opacity of the default text drop shadow
@@ -92,104 +92,104 @@ public:
 	 * @param InShadowColorAndOpacity		The new drop shadow color and opacity
 	 */
 	UFUNCTION(BlueprintCallable, Category="Appearance")
-	void SetDefaultShadowColorAndOpacity(FLinearColor InShadowColorAndOpacity);
+	UMG_API void SetDefaultShadowColorAndOpacity(FLinearColor InShadowColorAndOpacity);
 
 	/**  
 	 * Sets the offset that the default text drop shadow should be drawn at
 	 * @param InShadowOffset		The new offset
 	 */
 	UFUNCTION(BlueprintCallable, Category="Appearance")
-	void SetDefaultShadowOffset(FVector2D InShadowOffset);
+	UMG_API void SetDefaultShadowOffset(FVector2D InShadowOffset);
 
 	/**
 	 * Dynamically set the default font info for this rich text block
 	 * @param InFontInfo The new font info
 	 */
 	UFUNCTION(BlueprintCallable, Category = "Appearance")
-	void SetDefaultFont(FSlateFontInfo InFontInfo);
+	UMG_API void SetDefaultFont(FSlateFontInfo InFontInfo);
 
 	/**
 	 * Dynamically set the default strike brush for this rich text block
 	 * @param InStrikeBrush The new brush to use to strike through text
 	 */
 	UFUNCTION(BlueprintCallable, Category = "Appearance")
-	void SetDefaultStrikeBrush(const FSlateBrush& InStrikeBrush);
+	UMG_API void SetDefaultStrikeBrush(const FSlateBrush& InStrikeBrush);
 
 	/**
 	 *  Set the minimum desired width for this rich text block
 	 *  @param InMinDesiredWidth new minimum desired width
 	 */
 	UFUNCTION(BlueprintCallable, Category = "Appearance")
-	void SetMinDesiredWidth(float InMinDesiredWidth);
+	UMG_API void SetMinDesiredWidth(float InMinDesiredWidth);
 	
 	/**
 	 * Set the auto wrap for this rich text block
 	 * @param InAutoTextWrap to turn wrap on or off
 	 */
 	UFUNCTION(BlueprintCallable, Category = "Appearance")
-	void SetAutoWrapText(bool InAutoTextWrap);
+	UMG_API void SetAutoWrapText(bool InAutoTextWrap);
 
 	/**
 	 * Set the text transformation policy for this text block.
 	 * @param InTransformPolicy the new text transformation policy.
 	 */
 	UFUNCTION(BlueprintCallable, Category = "Appearance")
-	void SetTextTransformPolicy(ETextTransformPolicy InTransformPolicy);
+	UMG_API void SetTextTransformPolicy(ETextTransformPolicy InTransformPolicy);
 
 	/**
 	* Set the text overflow policy for this text block.
 	* @param InOverflowPolicy the new text overflow policy.
 	*/
 	UFUNCTION(BlueprintCallable, Category = "Appearance")
-	void SetTextOverflowPolicy(ETextOverflowPolicy InOverflowPolicy);
+	UMG_API void SetTextOverflowPolicy(ETextOverflowPolicy InOverflowPolicy);
 
 	/** 
 	 * Wholesale override of the currently established default text style
 	 * @param InDefaultTextStyle The new text style to apply to all default (i.e. undecorated) text in the block
 	 */
 	UFUNCTION(BlueprintCallable, Category = Appearance)
-	void SetDefaultTextStyle(const FTextBlockStyle& InDefaultTextStyle);
+	UMG_API void SetDefaultTextStyle(const FTextBlockStyle& InDefaultTextStyle);
 
 	UFUNCTION(BlueprintCallable, Category = Appearance)
-	void SetDefaultMaterial(UMaterialInterface* InMaterial);
+	UMG_API void SetDefaultMaterial(UMaterialInterface* InMaterial);
 
 	/** Remove all overrides made to the default text style and return to the style specified in the style set data table */
 	UFUNCTION(BlueprintCallable, Category = Appearance)
-	void ClearAllDefaultStyleOverrides();
+	UMG_API void ClearAllDefaultStyleOverrides();
 
 	/**
 	 * Creates a dynamic material for the default font or returns it if it already
 	 * exists
 	 */
 	UFUNCTION(BlueprintCallable, Category = "Appearance")
-	UMaterialInstanceDynamic* GetDefaultDynamicMaterial();
+	UMG_API UMaterialInstanceDynamic* GetDefaultDynamicMaterial();
 
 	/**
 	 * Replaces the existing decorators with the list provided
 	 */
 	UFUNCTION(BlueprintCallable, Category = "Appearance")
-	void SetDecorators(const TArray<TSubclassOf<URichTextBlockDecorator>>& InDecoratorClasses);
+	UMG_API void SetDecorators(const TArray<TSubclassOf<URichTextBlockDecorator>>& InDecoratorClasses);
 
 public:
-	URichTextBlock(const FObjectInitializer& ObjectInitializer);
+	UMG_API URichTextBlock(const FObjectInitializer& ObjectInitializer);
 
 	//~ Begin UTextLayoutWidget Interface
-	virtual void SetJustification(ETextJustify::Type InJustification) override;
+	UMG_API virtual void SetJustification(ETextJustify::Type InJustification) override;
 	//~ End UTextLayoutWidget Interface
 
 	// UWidget interface
-	virtual void SynchronizeProperties() override;
+	UMG_API virtual void SynchronizeProperties() override;
 	// End of UWidget interface
 
 	// UVisual interface
-	virtual void ReleaseSlateResources(bool bReleaseChildren) override;
+	UMG_API virtual void ReleaseSlateResources(bool bReleaseChildren) override;
 	// End of UVisual interface
 
 #if WITH_EDITOR
 	// UWidget interface
-	virtual const FText GetPaletteCategory() override;
-	virtual void OnCreationFromPalette() override;
-	virtual void ValidateCompiledDefaults(IWidgetCompilerLog& CompileLog) const override;
+	UMG_API virtual const FText GetPaletteCategory() override;
+	UMG_API virtual void OnCreationFromPalette() override;
+	UMG_API virtual void ValidateCompiledDefaults(IWidgetCompilerLog& CompileLog) const override;
 	// End UWidget interface
 #endif
 
@@ -197,7 +197,7 @@ public:
 	 * Returns widgets text.
 	 */
 	UFUNCTION(BlueprintCallable, Category = "Widget")
-	FText GetText() const;
+	UMG_API FText GetText() const;
 	
 	/**
 	 * Directly sets the widget text.
@@ -205,43 +205,43 @@ public:
 	 * @param InText The text to assign to the widget
 	 */
 	UFUNCTION(BlueprintCallable, Category = "Widget")
-	virtual void SetText(const FText& InText);
+	UMG_API virtual void SetText(const FText& InText);
 
 	UFUNCTION(BlueprintCallable, Category = "Widget")
-	UDataTable* GetTextStyleSet() const;
+	UMG_API UDataTable* GetTextStyleSet() const;
 
 	UFUNCTION(BlueprintCallable, Category = "Widget")
-	void SetTextStyleSet(UDataTable* NewTextStyleSet);
+	UMG_API void SetTextStyleSet(UDataTable* NewTextStyleSet);
 
-	const FTextBlockStyle& GetDefaultTextStyle() const;
-	const FTextBlockStyle& GetCurrentDefaultTextStyle() const;
+	UMG_API const FTextBlockStyle& GetDefaultTextStyle() const;
+	UMG_API const FTextBlockStyle& GetCurrentDefaultTextStyle() const;
 
 	/**  */
 	UFUNCTION(BlueprintCallable, Category = "Widget")
-	URichTextBlockDecorator* GetDecoratorByClass(TSubclassOf<URichTextBlockDecorator> DecoratorClass);
+	UMG_API URichTextBlockDecorator* GetDecoratorByClass(TSubclassOf<URichTextBlockDecorator> DecoratorClass);
 
 	/**
 	 * Causes the text to reflow it's layout and re-evaluate any decorators
 	 */
 	UFUNCTION(BlueprintCallable, Category = "Widget")
-	void RefreshTextLayout();
+	UMG_API void RefreshTextLayout();
 
 protected:
-	virtual TSharedRef<SWidget> RebuildWidget() override;
+	UMG_API virtual TSharedRef<SWidget> RebuildWidget() override;
 	
-	virtual void UpdateStyleData();
-	void RebuildStyleInstance();
-	virtual void CreateDecorators(TArray< TSharedRef< class ITextDecorator > >& OutDecorators);
-	virtual TSharedPtr< class IRichTextMarkupParser > CreateMarkupParser();
-	virtual TSharedPtr< class IRichTextMarkupWriter > CreateMarkupWriter();
-
-	void BeginDefaultStyleOverride();
-	virtual void ApplyUpdatedDefaultTextStyle();
-
-	const FTextBlockStyle& GetDefaultTextStyleOverride() const;
-	float GetMinDesiredWidth() const;
-	ETextTransformPolicy GetTransformPolicy() const;
-	ETextOverflowPolicy GetOverflowPolicy() const;
+	UMG_API virtual void UpdateStyleData();
+	UMG_API void RebuildStyleInstance();
+	UMG_API virtual void CreateDecorators(TArray< TSharedRef< class ITextDecorator > >& OutDecorators);
+	UMG_API virtual TSharedPtr< class IRichTextMarkupParser > CreateMarkupParser();
+	UMG_API virtual TSharedPtr< class IRichTextMarkupWriter > CreateMarkupWriter();
+
+	UMG_API void BeginDefaultStyleOverride();
+	UMG_API virtual void ApplyUpdatedDefaultTextStyle();
+
+	UMG_API const FTextBlockStyle& GetDefaultTextStyleOverride() const;
+	UMG_API float GetMinDesiredWidth() const;
+	UMG_API ETextTransformPolicy GetTransformPolicy() const;
+	UMG_API ETextOverflowPolicy GetOverflowPolicy() const;
 
 protected:
 	UPROPERTY(Transient)
diff --git a/Engine/Source/Runtime/UMG/Public/Components/RichTextBlockDecorator.h b/Engine/Source/Runtime/UMG/Public/Components/RichTextBlockDecorator.h
--- a/Engine/Source/Runtime/UMG/Public/Components/RichTextBlockDecorator.h
+++ b/Engine/Source/Runtime/UMG/Public/Components/RichTextBlockDecorator.h
@@ -16,10 +16,10 @@ class ISlateStyle;
 class URichTextBlockDecorator;
 class URichTextBlock;
 
-class UMG_API FRichTextDecorator : public ITextDecorator
+class FRichTextDecorator : public ITextDecorator
 {
 public:
-	FRichTextDecorator(URichTextBlock* InOwner);
+	UMG_API FRichTextDecorator(URichTextBlock* InOwner);
 
 	virtual ~FRichTextDecorator() {}
 
@@ -29,25 +29,25 @@ public:
 		return false;
 	}
 
-	virtual TSharedRef<ISlateRun> Create(const TSharedRef<class FTextLayout>& TextLayout, const FTextRunParseResults& RunParseResult, const FString& OriginalText, const TSharedRef< FString >& InOutModelText, const ISlateStyle* Style) override final;
+	UMG_API virtual TSharedRef<ISlateRun> Create(const TSharedRef<class FTextLayout>& TextLayout, const FTextRunParseResults& RunParseResult, const FString& OriginalText, const TSharedRef< FString >& InOutModelText, const ISlateStyle* Style) override final;
 
 protected:
 	/** Override this function if you want to create a unique widget like an image */
-	virtual TSharedPtr<SWidget> CreateDecoratorWidget(const FTextRunInfo& RunInfo, const FTextBlockStyle& DefaultTextStyle) const;
+	UMG_API virtual TSharedPtr<SWidget> CreateDecoratorWidget(const FTextRunInfo& RunInfo, const FTextBlockStyle& DefaultTextStyle) const;
 
 	/** Override this function if you want to dynamically generate text, optionally changing the style. InOutString will start as the content between tags */
-	virtual void CreateDecoratorText(const FTextRunInfo& RunInfo, FTextBlockStyle& InOutTextStyle, FString& InOutString) const;
+	UMG_API virtual void CreateDecoratorText(const FTextRunInfo& RunInfo, FTextBlockStyle& InOutTextStyle, FString& InOutString) const;
 
 	URichTextBlock* Owner;
 };
 
-UCLASS(Abstract, Blueprintable)
-class UMG_API URichTextBlockDecorator : public UObject
+UCLASS(Abstract, Blueprintable, MinimalAPI)
+class URichTextBlockDecorator : public UObject
 {
 	GENERATED_BODY()
 
 public:
-	URichTextBlockDecorator(const FObjectInitializer& ObjectInitializer);
+	UMG_API URichTextBlockDecorator(const FObjectInitializer& ObjectInitializer);
 
-	virtual TSharedPtr<ITextDecorator> CreateDecorator(URichTextBlock* InOwner);
+	UMG_API virtual TSharedPtr<ITextDecorator> CreateDecorator(URichTextBlock* InOwner);
 };
diff --git a/Engine/Source/Runtime/UMG/Public/Components/RichTextBlockImageDecorator.h b/Engine/Source/Runtime/UMG/Public/Components/RichTextBlockImageDecorator.h
--- a/Engine/Source/Runtime/UMG/Public/Components/RichTextBlockImageDecorator.h
+++ b/Engine/Source/Runtime/UMG/Public/Components/RichTextBlockImageDecorator.h
@@ -18,7 +18,7 @@ class ISlateStyle;
 
 /** Simple struct for rich text styles */
 USTRUCT(Blueprintable, BlueprintType)
-struct UMG_API FRichImageRow : public FTableRowBase
+struct FRichImageRow : public FTableRowBase
 {
 	GENERATED_USTRUCT_BODY()
 
@@ -35,21 +35,21 @@ public:
  *
  * Understands the format <img id="NameOfBrushInTable"></>
  */
-UCLASS(Abstract, Blueprintable)
-class UMG_API URichTextBlockImageDecorator : public URichTextBlockDecorator
+UCLASS(Abstract, Blueprintable, MinimalAPI)
+class URichTextBlockImageDecorator : public URichTextBlockDecorator
 {
 	GENERATED_BODY()
 
 public:
-	URichTextBlockImageDecorator(const FObjectInitializer& ObjectInitializer);
+	UMG_API URichTextBlockImageDecorator(const FObjectInitializer& ObjectInitializer);
 
-	virtual TSharedPtr<ITextDecorator> CreateDecorator(URichTextBlock* InOwner) override;
+	UMG_API virtual TSharedPtr<ITextDecorator> CreateDecorator(URichTextBlock* InOwner) override;
 
-	virtual const FSlateBrush* FindImageBrush(FName TagOrId, bool bWarnIfMissing);
+	UMG_API virtual const FSlateBrush* FindImageBrush(FName TagOrId, bool bWarnIfMissing);
 
 protected:
 
-	FRichImageRow* FindImageRow(FName TagOrId, bool bWarnIfMissing);
+	UMG_API FRichImageRow* FindImageRow(FName TagOrId, bool bWarnIfMissing);
 
 	UPROPERTY(EditAnywhere, Category=Appearance, meta = (RequiredAssetDataTags = "RowStructure=/Script/UMG.RichImageRow"))
 	TObjectPtr<class UDataTable> ImageSet;
diff --git a/Engine/Source/Runtime/UMG/Public/Components/SafeZone.h b/Engine/Source/Runtime/UMG/Public/Components/SafeZone.h
--- a/Engine/Source/Runtime/UMG/Public/Components/SafeZone.h
+++ b/Engine/Source/Runtime/UMG/Public/Components/SafeZone.h
@@ -25,27 +25,27 @@
  *   1: Display the overlay for the title safe zone.
  *   2: Display the overlay for the action safe zone.
  */
-UCLASS()
-class UMG_API USafeZone : public UContentWidget
+UCLASS(MinimalAPI)
+class USafeZone : public UContentWidget
 {
 	GENERATED_BODY()
 public:
-	USafeZone();
+	UMG_API USafeZone();
 
 #if WITH_EDITOR
-	virtual const FText GetPaletteCategory() override;
+	UMG_API virtual const FText GetPaletteCategory() override;
 
-	virtual void OnDesignerChanged(const FDesignerChangedEventArgs& EventArgs) override;
+	UMG_API virtual void OnDesignerChanged(const FDesignerChangedEventArgs& EventArgs) override;
 #endif
 
-	virtual void OnSlotAdded( UPanelSlot* Slot ) override;
-	virtual void OnSlotRemoved( UPanelSlot* Slot ) override;
-	virtual UClass* GetSlotClass() const override;
+	UMG_API virtual void OnSlotAdded( UPanelSlot* Slot ) override;
+	UMG_API virtual void OnSlotRemoved( UPanelSlot* Slot ) override;
+	UMG_API virtual UClass* GetSlotClass() const override;
 
-	void UpdateWidgetProperties();
+	UMG_API void UpdateWidgetProperties();
 
 	UFUNCTION(BlueprintCallable, Category = "SafeZone")
-	void SetSidesToPad(bool InPadLeft, bool InPadRight, bool InPadTop, bool InPadBottom);
+	UMG_API void SetSidesToPad(bool InPadLeft, bool InPadRight, bool InPadTop, bool InPadBottom);
 
 public:
 
@@ -71,22 +71,22 @@ public:
 
 protected:
 	// UWidget interface
-	virtual TSharedRef<SWidget> RebuildWidget() override;
+	UMG_API virtual TSharedRef<SWidget> RebuildWidget() override;
 	// End of UWidget interface
 
-	virtual void ReleaseSlateResources(bool bReleaseChildren) override;
+	UMG_API virtual void ReleaseSlateResources(bool bReleaseChildren) override;
 
-	void SetPadLeft(bool InPadLeft);
-	bool GetPadLeft() const;
+	UMG_API void SetPadLeft(bool InPadLeft);
+	UMG_API bool GetPadLeft() const;
 
-	void SetPadRight(bool InPadRight);
-	bool GetPadRight() const;
+	UMG_API void SetPadRight(bool InPadRight);
+	UMG_API bool GetPadRight() const;
 
-	void SetPadTop(bool InPadTop);
-	bool GetPadTop() const;
+	UMG_API void SetPadTop(bool InPadTop);
+	UMG_API bool GetPadTop() const;
 
-	void SetPadBottom(bool InPadBottom);
-	bool GetPadBottom() const;
+	UMG_API void SetPadBottom(bool InPadBottom);
+	UMG_API bool GetPadBottom() const;
 
 	TSharedPtr< class SSafeZone > MySafeZone;
 
diff --git a/Engine/Source/Runtime/UMG/Public/Components/SafeZoneSlot.h b/Engine/Source/Runtime/UMG/Public/Components/SafeZoneSlot.h
--- a/Engine/Source/Runtime/UMG/Public/Components/SafeZoneSlot.h
+++ b/Engine/Source/Runtime/UMG/Public/Components/SafeZoneSlot.h
@@ -11,8 +11,8 @@
 
 class SSafeZone;
 
-UCLASS()
-class UMG_API USafeZoneSlot : public UPanelSlot
+UCLASS(MinimalAPI)
+class USafeZoneSlot : public UPanelSlot
 {
 	GENERATED_BODY()
 public:
@@ -37,31 +37,31 @@ public:
 	UPROPERTY( EditAnywhere, BlueprintReadWrite, Getter, Setter, Category = SafeZone )
 	FMargin Padding;
 
-	USafeZoneSlot();
+	UMG_API USafeZoneSlot();
 
-	void SetIsTitleSafe(bool InIsTitleSafe);
-	bool IsTitleSafe() const;
+	UMG_API void SetIsTitleSafe(bool InIsTitleSafe);
+	UMG_API bool IsTitleSafe() const;
 
-	void SetSafeAreaScale(const FMargin& InSafeAreaScale);
-	const FMargin& GetSafeAreaScale() const;
+	UMG_API void SetSafeAreaScale(const FMargin& InSafeAreaScale);
+	UMG_API const FMargin& GetSafeAreaScale() const;
 
-	void SetHorizontalAlignment(EHorizontalAlignment InHAlign);
-	const EHorizontalAlignment GetHorizontalAlignment() const;
+	UMG_API void SetHorizontalAlignment(EHorizontalAlignment InHAlign);
+	UMG_API const EHorizontalAlignment GetHorizontalAlignment() const;
 
-	void SetVerticalAlignment(EVerticalAlignment InVAlign);
-	const EVerticalAlignment GetVerticalAlignment() const;
+	UMG_API void SetVerticalAlignment(EVerticalAlignment InVAlign);
+	UMG_API const EVerticalAlignment GetVerticalAlignment() const;
 
-	void SetPadding(const FMargin& InPadding);
-	const FMargin& GetPadding() const;
+	UMG_API void SetPadding(const FMargin& InPadding);
+	UMG_API const FMargin& GetPadding() const;
 
 	//~ UPanelSlot interface
-	virtual void SynchronizeProperties() override;
+	UMG_API virtual void SynchronizeProperties() override;
 	//~ End of UPanelSlot interface
 
 	/** Builds the underlying slot for the slate button. */
-	void BuildSlot(TSharedRef<SSafeZone> InSafeZone);
+	UMG_API void BuildSlot(TSharedRef<SSafeZone> InSafeZone);
 
-	virtual void ReleaseSlateResources(bool bReleaseChildren) override;
+	UMG_API virtual void ReleaseSlateResources(bool bReleaseChildren) override;
 
 private:
 
diff --git a/Engine/Source/Runtime/UMG/Public/Components/ScaleBox.h b/Engine/Source/Runtime/UMG/Public/Components/ScaleBox.h
--- a/Engine/Source/Runtime/UMG/Public/Components/ScaleBox.h
+++ b/Engine/Source/Runtime/UMG/Public/Components/ScaleBox.h
@@ -17,8 +17,8 @@
  * * Single Child
  * * Aspect Ratio
  */
-UCLASS(config=Engine)
-class UMG_API UScaleBox : public UContentWidget
+UCLASS(config=Engine, MinimalAPI)
+class UScaleBox : public UContentWidget
 {
 	GENERATED_UCLASS_BODY()
 
@@ -46,47 +46,47 @@ public:
 
 public:
 	UFUNCTION(BlueprintCallable, Category = "Appearance")
-	void SetStretch(EStretch::Type InStretch);
+	UMG_API void SetStretch(EStretch::Type InStretch);
 
-	EStretch::Type GetStretch() const;
+	UMG_API EStretch::Type GetStretch() const;
 
 	UFUNCTION(BlueprintCallable, Category = "Appearance")
-	void SetStretchDirection(EStretchDirection::Type InStretchDirection);
+	UMG_API void SetStretchDirection(EStretchDirection::Type InStretchDirection);
 
-	EStretchDirection::Type GetStretchDirection() const;
+	UMG_API EStretchDirection::Type GetStretchDirection() const;
 
 	UFUNCTION(BlueprintCallable, Category = "Appearance")
-	void SetUserSpecifiedScale(float InUserSpecifiedScale);
+	UMG_API void SetUserSpecifiedScale(float InUserSpecifiedScale);
 
-	float GetUserSpecifiedScale() const;
+	UMG_API float GetUserSpecifiedScale() const;
 
 	UFUNCTION(BlueprintCallable, Category = "Appearance")
-	void SetIgnoreInheritedScale(bool bInIgnoreInheritedScale);
+	UMG_API void SetIgnoreInheritedScale(bool bInIgnoreInheritedScale);
 
-	bool IsIgnoreInheritedScale() const;
+	UMG_API bool IsIgnoreInheritedScale() const;
 
 public:
 
 	// UWidget interface
-	virtual void SynchronizeProperties() override;
+	UMG_API virtual void SynchronizeProperties() override;
 	// End of UWidget interface
 
-	virtual void ReleaseSlateResources(bool bReleaseChildren) override;
+	UMG_API virtual void ReleaseSlateResources(bool bReleaseChildren) override;
 
 #if WITH_EDITOR
-	virtual const FText GetPaletteCategory() override;
-	virtual bool CanEditChange(const FProperty* InProperty) const override;
-	virtual void OnDesignerChanged(const FDesignerChangedEventArgs& EventArgs) override;
+	UMG_API virtual const FText GetPaletteCategory() override;
+	UMG_API virtual bool CanEditChange(const FProperty* InProperty) const override;
+	UMG_API virtual void OnDesignerChanged(const FDesignerChangedEventArgs& EventArgs) override;
 #endif
 
-	virtual void Serialize(FArchive& Ar) override;
+	UMG_API virtual void Serialize(FArchive& Ar) override;
 
 protected:
 
 	// UPanelWidget
-	virtual UClass* GetSlotClass() const override;
-	virtual void OnSlotAdded(UPanelSlot* Slot) override;
-	virtual void OnSlotRemoved(UPanelSlot* Slot) override;
+	UMG_API virtual UClass* GetSlotClass() const override;
+	UMG_API virtual void OnSlotAdded(UPanelSlot* Slot) override;
+	UMG_API virtual void OnSlotRemoved(UPanelSlot* Slot) override;
 	// End UPanelWidget
 
 protected:
@@ -98,6 +98,6 @@ protected:
 
 protected:
 	// UWidget interface
-	virtual TSharedRef<SWidget> RebuildWidget() override;
+	UMG_API virtual TSharedRef<SWidget> RebuildWidget() override;
 	// End of UWidget interface
 };
diff --git a/Engine/Source/Runtime/UMG/Public/Components/ScaleBoxSlot.h b/Engine/Source/Runtime/UMG/Public/Components/ScaleBoxSlot.h
--- a/Engine/Source/Runtime/UMG/Public/Components/ScaleBoxSlot.h
+++ b/Engine/Source/Runtime/UMG/Public/Components/ScaleBoxSlot.h
@@ -13,8 +13,8 @@
 class SScaleBox;
 
 /** The Slot for the UScaleBoxSlot, contains the widget displayed in a button's single slot */
-UCLASS()
-class UMG_API UScaleBoxSlot : public UPanelSlot
+UCLASS(MinimalAPI)
+class UScaleBoxSlot : public UPanelSlot
 {
 	GENERATED_UCLASS_BODY()
 
@@ -39,27 +39,27 @@ public:
 public:
 	UE_DEPRECATED(5.1, "SetPadding is deprecated")
 	UFUNCTION(BlueprintCallable, Category="Layout|ScaleBox Slot", meta = (DeprecatedFunction, DeprecatedMessage = "The function is deprecated - There is no padding on ScaleBox."))
-	void SetPadding(FMargin InPadding);
+	UMG_API void SetPadding(FMargin InPadding);
 
-	EHorizontalAlignment GetHorizontalAlignment() const;
+	UMG_API EHorizontalAlignment GetHorizontalAlignment() const;
 
 	UFUNCTION(BlueprintCallable, Category="Layout|ScaleBox Slot")
-	void SetHorizontalAlignment(EHorizontalAlignment InHorizontalAlignment);
+	UMG_API void SetHorizontalAlignment(EHorizontalAlignment InHorizontalAlignment);
 
 
-	EVerticalAlignment GetVerticalAlignment() const;
+	UMG_API EVerticalAlignment GetVerticalAlignment() const;
 
 	UFUNCTION(BlueprintCallable, Category="Layout|ScaleBox Slot")
-	void SetVerticalAlignment(EVerticalAlignment InVerticalAlignment);
+	UMG_API void SetVerticalAlignment(EVerticalAlignment InVerticalAlignment);
 
 	//~ UPanelSlot interface
-	virtual void SynchronizeProperties() override;
+	UMG_API virtual void SynchronizeProperties() override;
 	//~ End of UPanelSlot interface
 
 	/** Builds the underlying slot for the slate ScaleBox. */
-	void BuildSlot(TSharedRef<SScaleBox> InScaleBox);
+	UMG_API void BuildSlot(TSharedRef<SScaleBox> InScaleBox);
 
-	virtual void ReleaseSlateResources(bool bReleaseChildren) override;
+	UMG_API virtual void ReleaseSlateResources(bool bReleaseChildren) override;
 
 private:
 	/** A pointer to the button to allow us to adjust the size, padding...etc at runtime. */
diff --git a/Engine/Source/Runtime/UMG/Public/Components/ScrollBar.h b/Engine/Source/Runtime/UMG/Public/Components/ScrollBar.h
--- a/Engine/Source/Runtime/UMG/Public/Components/ScrollBar.h
+++ b/Engine/Source/Runtime/UMG/Public/Components/ScrollBar.h
@@ -11,8 +11,8 @@
 #include "ScrollBar.generated.h"
 
 /** */
-UCLASS(Experimental)
-class UMG_API UScrollBar : public UWidget
+UCLASS(Experimental, MinimalAPI)
+class UScrollBar : public UWidget
 {
 	GENERATED_UCLASS_BODY()
 
@@ -59,55 +59,55 @@ public:
 	* @param InThumbSizeFraction  Size of thumbnail as a fraction of the total available scroll space.
 	*/
 	UFUNCTION(BlueprintCallable, Category="Scrolling")
-	void SetState(float InOffsetFraction, float InThumbSizeFraction);
+	UMG_API void SetState(float InOffsetFraction, float InThumbSizeFraction);
 
 	//~ Begin UWidget Interface
-	virtual void SynchronizeProperties() override;
+	UMG_API virtual void SynchronizeProperties() override;
 	//~ End UWidget Interface
 
 	//~ Begin UVisual Interface
-	virtual void ReleaseSlateResources(bool bReleaseChildren) override;
+	UMG_API virtual void ReleaseSlateResources(bool bReleaseChildren) override;
 	//~ End UVisual Interface
 
 	//~ Begin UObject Interface
 #if WITH_EDITORONLY_DATA
-	virtual void Serialize(FArchive& Ar) override;
+	UMG_API virtual void Serialize(FArchive& Ar) override;
 #endif // if WITH_EDITORONLY_DATA
 	//~ End UObject Interface
 
 #if WITH_EDITOR
 	//~ Begin UWidget Interface
-	virtual const FText GetPaletteCategory() override;
+	UMG_API virtual const FText GetPaletteCategory() override;
 	//~ End UWidget Interface
 #endif
 
 	/** @return the style of scrollbar. */
-	const FScrollBarStyle& GetWidgetStyle() const;
+	UMG_API const FScrollBarStyle& GetWidgetStyle() const;
 	/** Sets the style of scrollbar. */
-	void SetWidgetStyle(const FScrollBarStyle& InWidgetStyle);
+	UMG_API void SetWidgetStyle(const FScrollBarStyle& InWidgetStyle);
 
 	/** @return True if the scrollbar should always show. */
-	bool IsAlwaysShowScrollbar() const;
+	UMG_API bool IsAlwaysShowScrollbar() const;
 	/** Sets whether the scrollbar should always show. */
-	void SetAlwaysShowScrollbar(bool bNewValue);
+	UMG_API void SetAlwaysShowScrollbar(bool bNewValue);
 
 	/** @return True if the scrollbar track should always show. */
-	bool IsAlwaysShowScrollbarTrack() const;
+	UMG_API bool IsAlwaysShowScrollbarTrack() const;
 	/** Sets whether the scrollbar track should always show */
-	void SetAlwaysShowScrollbarTrack(bool bNewValue);
+	UMG_API void SetAlwaysShowScrollbarTrack(bool bNewValue);
 
 	/** @return the orientation of the scrollbar. */
-	EOrientation GetOrientation() const;
+	UMG_API EOrientation GetOrientation() const;
 
 	/** @return the thickness of the scrollbar. */
-	FVector2D GetThickness() const;
+	UMG_API FVector2D GetThickness() const;
 	/** Sets the thickness of the scrollbar. */
-	void SetThickness(const FVector2D& InThickness);
+	UMG_API void SetThickness(const FVector2D& InThickness);
 
 	/** Sets the padding of the scrollbar. */
-	FMargin GetPadding() const;
+	UMG_API FMargin GetPadding() const;
 	/** @return the padding of the scrollbar. */
-	void SetPadding(const FMargin& InPadding);
+	UMG_API void SetPadding(const FMargin& InPadding);
 
 protected:
 
@@ -115,9 +115,9 @@ protected:
 
 protected:
 	//~ Begin UWidget Interface
-	virtual TSharedRef<SWidget> RebuildWidget() override;
+	UMG_API virtual TSharedRef<SWidget> RebuildWidget() override;
 	//~ End UWidget Interface
 
 	/** Initialize the orientation of the scrollbar in the constructor before the SWidget is constructed. */
-	void InitOrientation(EOrientation InOrientation);
+	UMG_API void InitOrientation(EOrientation InOrientation);
 };
diff --git a/Engine/Source/Runtime/UMG/Public/Components/ScrollBox.h b/Engine/Source/Runtime/UMG/Public/Components/ScrollBox.h
--- a/Engine/Source/Runtime/UMG/Public/Components/ScrollBox.h
+++ b/Engine/Source/Runtime/UMG/Public/Components/ScrollBox.h
@@ -16,8 +16,8 @@ DECLARE_DYNAMIC_MULTICAST_DELEGATE_OneParam(FOnUserScrolledEvent, float, Current
 /**
  * An arbitrary scrollable collection of widgets.  Great for presenting 10-100 widgets in a list.  Doesn't support virtualization.
  */
-UCLASS()
-class UMG_API UScrollBox : public UPanelWidget
+UCLASS(MinimalAPI)
+class UScrollBox : public UPanelWidget
 {
 	GENERATED_UCLASS_BODY()
 
@@ -116,86 +116,86 @@ public:
 	UPROPERTY(EditAnywhere, BlueprintReadWrite, Getter, Setter, BlueprintSetter = "SetWheelScrollMultiplier", Category = "Scroll")
 	float WheelScrollMultiplier = 1.f;
 
-	void SetWidgetStyle(const FScrollBoxStyle& NewWidgetStyle);
+	UMG_API void SetWidgetStyle(const FScrollBoxStyle& NewWidgetStyle);
 
-	const FScrollBoxStyle& GetWidgetStyle() const;
+	UMG_API const FScrollBoxStyle& GetWidgetStyle() const;
 
-	void SetWidgetBarStyle(const FScrollBarStyle& NewWidgetBarStyle);
+	UMG_API void SetWidgetBarStyle(const FScrollBarStyle& NewWidgetBarStyle);
 
-	const FScrollBarStyle& GetWidgetBarStyle() const;
+	UMG_API const FScrollBarStyle& GetWidgetBarStyle() const;
 
 	UFUNCTION(BlueprintCallable, Category = "Scroll")
-	void SetConsumeMouseWheel(EConsumeMouseWheel NewConsumeMouseWheel);
+	UMG_API void SetConsumeMouseWheel(EConsumeMouseWheel NewConsumeMouseWheel);
 
-	EConsumeMouseWheel GetConsumeMouseWheel() const;
+	UMG_API EConsumeMouseWheel GetConsumeMouseWheel() const;
 
 	UFUNCTION(BlueprintCallable, Category = "Scroll")
-	void SetOrientation(EOrientation NewOrientation);
+	UMG_API void SetOrientation(EOrientation NewOrientation);
 
-	EOrientation GetOrientation() const;
+	UMG_API EOrientation GetOrientation() const;
 
 	UFUNCTION(BlueprintCallable, Category = "Scroll")
-	void SetScrollBarVisibility(ESlateVisibility NewScrollBarVisibility);
+	UMG_API void SetScrollBarVisibility(ESlateVisibility NewScrollBarVisibility);
 
-	ESlateVisibility GetScrollBarVisibility() const;
+	UMG_API ESlateVisibility GetScrollBarVisibility() const;
 
 	UFUNCTION(BlueprintCallable, Category = "Scroll")
-	void SetScrollbarThickness(const FVector2D& NewScrollbarThickness);
+	UMG_API void SetScrollbarThickness(const FVector2D& NewScrollbarThickness);
 
-	FVector2D GetScrollbarThickness() const;
+	UMG_API FVector2D GetScrollbarThickness() const;
 
 	UFUNCTION(BlueprintCallable, Category = "Scroll")
-	void SetScrollbarPadding(const FMargin& NewScrollbarPadding);
+	UMG_API void SetScrollbarPadding(const FMargin& NewScrollbarPadding);
 
-	FMargin GetScrollbarPadding() const;
+	UMG_API FMargin GetScrollbarPadding() const;
 
 	UFUNCTION(BlueprintCallable, Category = "Scroll")
-	void SetAlwaysShowScrollbar(bool NewAlwaysShowScrollbar);
+	UMG_API void SetAlwaysShowScrollbar(bool NewAlwaysShowScrollbar);
 	
-	bool IsAlwaysShowScrollbar() const;
+	UMG_API bool IsAlwaysShowScrollbar() const;
 
 	UFUNCTION(BlueprintCallable, Category = "Scroll")
-	void SetAllowOverscroll(bool NewAllowOverscroll);
+	UMG_API void SetAllowOverscroll(bool NewAllowOverscroll);
 
-	bool IsAllowOverscroll() const;
+	UMG_API bool IsAllowOverscroll() const;
 
 	UFUNCTION(BlueprintCallable, Category = "Scroll")
-	void SetAnimateWheelScrolling(bool bShouldAnimateWheelScrolling);
+	UMG_API void SetAnimateWheelScrolling(bool bShouldAnimateWheelScrolling);
 
-	bool IsAnimateWheelScrolling() const;
+	UMG_API bool IsAnimateWheelScrolling() const;
 
 	UFUNCTION(BlueprintCallable, Category = "Scroll")
-	void SetWheelScrollMultiplier(float NewWheelScrollMultiplier);
+	UMG_API void SetWheelScrollMultiplier(float NewWheelScrollMultiplier);
 
-	float GetWheelScrollMultiplier() const;
+	UMG_API float GetWheelScrollMultiplier() const;
 
 	UFUNCTION(BlueprintCallable, Category = "Scroll")
-	void SetScrollWhenFocusChanges(EScrollWhenFocusChanges NewScrollWhenFocusChanges);
+	UMG_API void SetScrollWhenFocusChanges(EScrollWhenFocusChanges NewScrollWhenFocusChanges);
 
-	EScrollWhenFocusChanges GetScrollWhenFocusChanges() const;
+	UMG_API EScrollWhenFocusChanges GetScrollWhenFocusChanges() const;
 
 	UFUNCTION(BlueprintCallable, Category = "Scroll")
-	void SetNavigationDestination(const EDescendantScrollDestination NewNavigationDestination);
+	UMG_API void SetNavigationDestination(const EDescendantScrollDestination NewNavigationDestination);
 
-	EDescendantScrollDestination GetNavigationDestination() const;
+	UMG_API EDescendantScrollDestination GetNavigationDestination() const;
 
-	void SetAlwaysShowScrollbarTrack(bool NewAlwaysShowScrollbarTrack);
+	UMG_API void SetAlwaysShowScrollbarTrack(bool NewAlwaysShowScrollbarTrack);
 
-	bool IsAlwaysShowScrollbarTrack() const;
+	UMG_API bool IsAlwaysShowScrollbarTrack() const;
 
-	float GetNavigationScrollPadding() const;
+	UMG_API float GetNavigationScrollPadding() const;
 
-	void SetAllowRightClickDragScrolling(bool bShouldAllowRightClickDragScrolling);
+	UMG_API void SetAllowRightClickDragScrolling(bool bShouldAllowRightClickDragScrolling);
 
-	bool IsAllowRightClickDragScrolling() const;
+	UMG_API bool IsAllowRightClickDragScrolling() const;
 
-	bool IsFrontPadScrolling() const;
+	UMG_API bool IsFrontPadScrolling() const;
 
-	bool IsBackPadScrolling() const;
+	UMG_API bool IsBackPadScrolling() const;
 
 	/** Instantly stops any inertial scrolling that is currently in progress */
 	UFUNCTION(BlueprintCallable, Category = "Scroll")
-	void EndInertialScrolling();
+	UMG_API void EndInertialScrolling();
 
 public:
 
@@ -208,75 +208,75 @@ public:
 	 * @param NewScrollOffset is in Slate Units.
 	 */
 	UFUNCTION(BlueprintCallable, Category="Widget")
-	void SetScrollOffset(float NewScrollOffset);
+	UMG_API void SetScrollOffset(float NewScrollOffset);
 	
 	/**
 	 * Gets the scroll offset of the scrollbox in Slate Units.
 	 */
 	UFUNCTION(BlueprintCallable, Category="Widget")
-	float GetScrollOffset() const;
+	UMG_API float GetScrollOffset() const;
 
 	/** Gets the scroll offset of the bottom of the ScrollBox in Slate Units. */
 	UFUNCTION(BlueprintCallable, Category = "Widget")
-	float GetScrollOffsetOfEnd() const;
+	UMG_API float GetScrollOffsetOfEnd() const;
 
 	/** Gets the fraction currently visible in the scrollbox */
 	UFUNCTION(BlueprintCallable, Category="Widget")
-	float GetViewFraction() const;
+	UMG_API float GetViewFraction() const;
 
 	UFUNCTION(BlueprintCallable, Category="Widget")
-	float GetViewOffsetFraction() const;
+	UMG_API float GetViewOffsetFraction() const;
 
 	/** Scrolls the ScrollBox to the top instantly */
 	UFUNCTION(BlueprintCallable, Category="Widget")
-	void ScrollToStart();
+	UMG_API void ScrollToStart();
 
 	/** Scrolls the ScrollBox to the bottom instantly during the next layout pass. */
 	UFUNCTION(BlueprintCallable, Category="Widget")
-	void ScrollToEnd();
+	UMG_API void ScrollToEnd();
 
 	/** Scrolls the ScrollBox to the widget during the next layout pass. */
 	UFUNCTION(BlueprintCallable, Category="Widget")
-	void ScrollWidgetIntoView(UWidget* WidgetToFind, bool AnimateScroll = true, EDescendantScrollDestination ScrollDestination = EDescendantScrollDestination::IntoView, float Padding = 0);
+	UMG_API void ScrollWidgetIntoView(UWidget* WidgetToFind, bool AnimateScroll = true, EDescendantScrollDestination ScrollDestination = EDescendantScrollDestination::IntoView, float Padding = 0);
 
 	//~ Begin UWidget Interface
-	virtual void SynchronizeProperties() override;
+	UMG_API virtual void SynchronizeProperties() override;
 	//~ End UWidget Interface
 
 	//~ Begin UVisual Interface
-	virtual void ReleaseSlateResources(bool bReleaseChildren) override;
+	UMG_API virtual void ReleaseSlateResources(bool bReleaseChildren) override;
 	//~ End UVisual Interface
 
 	//~ Begin UObject Interface
 #if WITH_EDITORONLY_DATA
-	virtual void Serialize(FArchive& Ar) override;
+	UMG_API virtual void Serialize(FArchive& Ar) override;
 #endif // if WITH_EDITORONLY_DATA
 	//~ End UObject Interface
 
 #if WITH_EDITOR
 	//~ Begin UWidget Interface
-	virtual const FText GetPaletteCategory() override;
-	virtual void OnDescendantSelectedByDesigner( UWidget* DescendantWidget ) override;
-	virtual void OnDescendantDeselectedByDesigner( UWidget* DescendantWidget ) override;
+	UMG_API virtual const FText GetPaletteCategory() override;
+	UMG_API virtual void OnDescendantSelectedByDesigner( UWidget* DescendantWidget ) override;
+	UMG_API virtual void OnDescendantDeselectedByDesigner( UWidget* DescendantWidget ) override;
 	//~ End UWidget Interface
 #endif
 
 protected:
 
 	// UPanelWidget
-	virtual UClass* GetSlotClass() const override;
-	virtual void OnSlotAdded(UPanelSlot* Slot) override;
-	virtual void OnSlotRemoved(UPanelSlot* Slot) override;
+	UMG_API virtual UClass* GetSlotClass() const override;
+	UMG_API virtual void OnSlotAdded(UPanelSlot* Slot) override;
+	UMG_API virtual void OnSlotRemoved(UPanelSlot* Slot) override;
 	// End UPanelWidget
 
-	void SlateHandleUserScrolled(float CurrentOffset);
+	UMG_API void SlateHandleUserScrolled(float CurrentOffset);
 
 	// Initialize IsFocusable in the constructor before the SWidget is constructed.
-	void InitBackPadScrolling(bool InBackPadScrolling);
+	UMG_API void InitBackPadScrolling(bool InBackPadScrolling);
 	// Initialize IsFocusable in the constructor before the SWidget is constructed.
-	void InitFrontPadScrolling(bool InFrontPadScrolling);
+	UMG_API void InitFrontPadScrolling(bool InFrontPadScrolling);
 	// Initialize IsFocusable in the constructor before the SWidget is constructed.
-	void InitNavigationScrollPadding(float InNavigationScrollPadding);
+	UMG_API void InitNavigationScrollPadding(float InNavigationScrollPadding);
 
 protected:
 	/** The desired scroll offset for the underlying scrollbox.  This is a cache so that it can be set before the widget is constructed. */
@@ -286,7 +286,7 @@ protected:
 
 protected:
 	//~ Begin UWidget Interface
-	virtual TSharedRef<SWidget> RebuildWidget() override;
+	UMG_API virtual TSharedRef<SWidget> RebuildWidget() override;
 	//~ End UWidget Interface
 
 #if WITH_EDITOR
diff --git a/Engine/Source/Runtime/UMG/Public/Components/ScrollBoxSlot.h b/Engine/Source/Runtime/UMG/Public/Components/ScrollBoxSlot.h
--- a/Engine/Source/Runtime/UMG/Public/Components/ScrollBoxSlot.h
+++ b/Engine/Source/Runtime/UMG/Public/Components/ScrollBoxSlot.h
@@ -13,8 +13,8 @@
 #include "ScrollBoxSlot.generated.h"
 
 /** The Slot for the UScrollBox, contains the widget that are scrollable */
-UCLASS()
-class UMG_API UScrollBoxSlot : public UPanelSlot
+UCLASS(MinimalAPI)
+class UScrollBoxSlot : public UPanelSlot
 {
 	GENERATED_UCLASS_BODY()
 
@@ -43,35 +43,35 @@ public:
 
 public:
 
-	FMargin GetPadding() const;
+	UMG_API FMargin GetPadding() const;
 
 	UFUNCTION(BlueprintCallable, Category="Layout|ScrollBox Slot")
-	void SetPadding(FMargin InPadding);
+	UMG_API void SetPadding(FMargin InPadding);
 
-	FSlateChildSize GetSize() const;
+	UMG_API FSlateChildSize GetSize() const;
 
-	void SetSize(FSlateChildSize InSize);
+	UMG_API void SetSize(FSlateChildSize InSize);
 
-	EHorizontalAlignment GetHorizontalAlignment() const;
+	UMG_API EHorizontalAlignment GetHorizontalAlignment() const;
 
 	UFUNCTION(BlueprintCallable, Category="Layout|ScrollBox Slot")
-	void SetHorizontalAlignment(EHorizontalAlignment InHorizontalAlignment);
+	UMG_API void SetHorizontalAlignment(EHorizontalAlignment InHorizontalAlignment);
 
-	EVerticalAlignment GetVerticalAlignment() const;
+	UMG_API EVerticalAlignment GetVerticalAlignment() const;
 
 	UFUNCTION(BlueprintCallable, Category = "Layout|ScrollBox Slot")
-	void SetVerticalAlignment(EVerticalAlignment InVerticalAlignment);
+	UMG_API void SetVerticalAlignment(EVerticalAlignment InVerticalAlignment);
 
 public:
 
 	//~ UPanelSlot interface
-	virtual void SynchronizeProperties() override;
+	UMG_API virtual void SynchronizeProperties() override;
 	//~ End of UPanelSlot interface
 
 	/** Builds the underlying FSlot for the Slate layout panel. */
-	void BuildSlot(TSharedRef<SScrollBox> ScrollBox);
+	UMG_API void BuildSlot(TSharedRef<SScrollBox> ScrollBox);
 
-	virtual void ReleaseSlateResources(bool bReleaseChildren) override;
+	UMG_API virtual void ReleaseSlateResources(bool bReleaseChildren) override;
 
 private:
 
diff --git a/Engine/Source/Runtime/UMG/Public/Components/SizeBox.h b/Engine/Source/Runtime/UMG/Public/Components/SizeBox.h
--- a/Engine/Source/Runtime/UMG/Public/Components/SizeBox.h
+++ b/Engine/Source/Runtime/UMG/Public/Components/SizeBox.h
@@ -17,8 +17,8 @@ class SBox;
  * * Single Child
  * * Fixed Size
  */
-UCLASS()
-class UMG_API USizeBox : public UContentWidget
+UCLASS(MinimalAPI)
+class USizeBox : public UContentWidget
 {
 	GENERATED_UCLASS_BODY()
 
@@ -111,105 +111,105 @@ public:
 public:
 
 	/** */
-	float GetWidthOverride() const;
+	UMG_API float GetWidthOverride() const;
 
 	/** When specified, ignore the content's desired size and report the WidthOverride as the Box's desired width. */
 	UFUNCTION(BlueprintCallable, Category="Layout|Size Box")
-	void SetWidthOverride(float InWidthOverride);
+	UMG_API void SetWidthOverride(float InWidthOverride);
 
 	UFUNCTION(BlueprintCallable, Category="Layout|Size Box")
-	void ClearWidthOverride();
+	UMG_API void ClearWidthOverride();
 
 	/** */
-	float GetHeightOverride() const;
+	UMG_API float GetHeightOverride() const;
 
 	/** When specified, ignore the content's desired size and report the HeightOverride as the Box's desired height. */
 	UFUNCTION(BlueprintCallable, Category="Layout|Size Box")
-	void SetHeightOverride(float InHeightOverride);
+	UMG_API void SetHeightOverride(float InHeightOverride);
 
 	UFUNCTION(BlueprintCallable, Category="Layout|Size Box")
-	void ClearHeightOverride();
+	UMG_API void ClearHeightOverride();
 
 	/** */
-	float GetMinDesiredWidth() const;
+	UMG_API float GetMinDesiredWidth() const;
 
 	/** When specified, will report the MinDesiredWidth if larger than the content's desired width. */
 	UFUNCTION(BlueprintCallable, Category="Layout|Size Box")
-	void SetMinDesiredWidth(float InMinDesiredWidth);
+	UMG_API void SetMinDesiredWidth(float InMinDesiredWidth);
 
 	UFUNCTION(BlueprintCallable, Category="Layout|Size Box")
-	void ClearMinDesiredWidth();
+	UMG_API void ClearMinDesiredWidth();
 
 	/** */
-	float GetMinDesiredHeight() const;
+	UMG_API float GetMinDesiredHeight() const;
 
 	/** When specified, will report the MinDesiredHeight if larger than the content's desired height. */
 	UFUNCTION(BlueprintCallable, Category="Layout|Size Box")
-	void SetMinDesiredHeight(float InMinDesiredHeight);
+	UMG_API void SetMinDesiredHeight(float InMinDesiredHeight);
 
 	UFUNCTION(BlueprintCallable, Category="Layout|Size Box")
-	void ClearMinDesiredHeight();
+	UMG_API void ClearMinDesiredHeight();
 
 	/** */
-	float GetMaxDesiredWidth() const;
+	UMG_API float GetMaxDesiredWidth() const;
 
 	/** When specified, will report the MaxDesiredWidth if smaller than the content's desired width. */
 	UFUNCTION(BlueprintCallable, Category="Layout|Size Box")
-	void SetMaxDesiredWidth(float InMaxDesiredWidth);
+	UMG_API void SetMaxDesiredWidth(float InMaxDesiredWidth);
 
 	UFUNCTION(BlueprintCallable, Category="Layout|Size Box")
-	void ClearMaxDesiredWidth();
+	UMG_API void ClearMaxDesiredWidth();
 
 	/** */
-	float GetMaxDesiredHeight() const;
+	UMG_API float GetMaxDesiredHeight() const;
 
 	/** When specified, will report the MaxDesiredHeight if smaller than the content's desired height. */
 	UFUNCTION(BlueprintCallable, Category="Layout|Size Box")
-	void SetMaxDesiredHeight(float InMaxDesiredHeight);
+	UMG_API void SetMaxDesiredHeight(float InMaxDesiredHeight);
 
 	UFUNCTION(BlueprintCallable, Category="Layout|Size Box")
-	void ClearMaxDesiredHeight();
+	UMG_API void ClearMaxDesiredHeight();
 
 	/** */
-	float GetMinAspectRatio() const;
+	UMG_API float GetMinAspectRatio() const;
 
 	UFUNCTION(BlueprintCallable, Category="Layout|Size Box")
-	void SetMinAspectRatio(float InMinAspectRatio);
+	UMG_API void SetMinAspectRatio(float InMinAspectRatio);
 
 	UFUNCTION(BlueprintCallable, Category="Layout|Size Box")
-	void ClearMinAspectRatio();
+	UMG_API void ClearMinAspectRatio();
 
 	/** */
-	float GetMaxAspectRatio() const;
+	UMG_API float GetMaxAspectRatio() const;
 
 	UFUNCTION(BlueprintCallable, Category="Layout|Size Box")
-	void SetMaxAspectRatio(float InMaxAspectRatio);
+	UMG_API void SetMaxAspectRatio(float InMaxAspectRatio);
 
 	UFUNCTION(BlueprintCallable, Category = "Layout|Size Box")
-	void ClearMaxAspectRatio();
+	UMG_API void ClearMaxAspectRatio();
 
 
 	// UWidget interface
-	virtual void SynchronizeProperties() override;
+	UMG_API virtual void SynchronizeProperties() override;
 	// End of UWidget interface
 
 	// UVisual interface
-	virtual void ReleaseSlateResources(bool bReleaseChildren) override;
+	UMG_API virtual void ReleaseSlateResources(bool bReleaseChildren) override;
 	// End of UVisual interface
 
 #if WITH_EDITOR
-	virtual const FText GetPaletteCategory() override;
+	UMG_API virtual const FText GetPaletteCategory() override;
 #endif
 
 protected:
 
 	// UPanelWidget
-	virtual UClass* GetSlotClass() const override;
-	virtual void OnSlotAdded(UPanelSlot* Slot) override;
-	virtual void OnSlotRemoved(UPanelSlot* Slot) override;
+	UMG_API virtual UClass* GetSlotClass() const override;
+	UMG_API virtual void OnSlotAdded(UPanelSlot* Slot) override;
+	UMG_API virtual void OnSlotRemoved(UPanelSlot* Slot) override;
 	// End UPanelWidget
 
 	// UWidget interface
-	virtual TSharedRef<SWidget> RebuildWidget() override;
+	UMG_API virtual TSharedRef<SWidget> RebuildWidget() override;
 	// End of UWidget interface
-};
\ No newline at end of file
+};
diff --git a/Engine/Source/Runtime/UMG/Public/Components/SizeBoxSlot.h b/Engine/Source/Runtime/UMG/Public/Components/SizeBoxSlot.h
--- a/Engine/Source/Runtime/UMG/Public/Components/SizeBoxSlot.h
+++ b/Engine/Source/Runtime/UMG/Public/Components/SizeBoxSlot.h
@@ -13,8 +13,8 @@
 class SBox;
 
 /** The Slot for the USizeBoxSlot, contains the widget displayed in a button's single slot */
-UCLASS()
-class UMG_API USizeBoxSlot : public UPanelSlot
+UCLASS(MinimalAPI)
+class USizeBoxSlot : public UPanelSlot
 {
 	GENERATED_UCLASS_BODY()
 
@@ -41,29 +41,29 @@ public:
 
 public:
 
-	FMargin GetPadding() const;
+	UMG_API FMargin GetPadding() const;
 
 	UFUNCTION(BlueprintCallable, Category="Layout|SizeBox Slot")
-	void SetPadding(FMargin InPadding);
+	UMG_API void SetPadding(FMargin InPadding);
 
-	EHorizontalAlignment GetHorizontalAlignment() const;
+	UMG_API EHorizontalAlignment GetHorizontalAlignment() const;
 
 	UFUNCTION(BlueprintCallable, Category="Layout|SizeBox Slot")
-	void SetHorizontalAlignment(EHorizontalAlignment InHorizontalAlignment);
+	UMG_API void SetHorizontalAlignment(EHorizontalAlignment InHorizontalAlignment);
 
-	EVerticalAlignment GetVerticalAlignment() const;
+	UMG_API EVerticalAlignment GetVerticalAlignment() const;
 
 	UFUNCTION(BlueprintCallable, Category="Layout|SizeBox Slot")
-	void SetVerticalAlignment(EVerticalAlignment InVerticalAlignment);
+	UMG_API void SetVerticalAlignment(EVerticalAlignment InVerticalAlignment);
 
 public:
 
 	//~ UPanelSlot interface
-	virtual void SynchronizeProperties() override;
+	UMG_API virtual void SynchronizeProperties() override;
 	//~ End of UPanelSlot interface
 
 	/** Builds the underlying slot for the slate SizeBox. */
-	void BuildSlot(TSharedRef<SBox> InSizeBox);
+	UMG_API void BuildSlot(TSharedRef<SBox> InSizeBox);
 
-	virtual void ReleaseSlateResources(bool bReleaseChildren) override;
+	UMG_API virtual void ReleaseSlateResources(bool bReleaseChildren) override;
 };
diff --git a/Engine/Source/Runtime/UMG/Public/Components/Slider.h b/Engine/Source/Runtime/UMG/Public/Components/Slider.h
--- a/Engine/Source/Runtime/UMG/Public/Components/Slider.h
+++ b/Engine/Source/Runtime/UMG/Public/Components/Slider.h
@@ -22,8 +22,8 @@ DECLARE_DYNAMIC_MULTICAST_DELEGATE_OneParam(FOnFloatValueChangedEvent, float, Va
  *
  * * No Children
  */
-UCLASS()
-class UMG_API USlider : public UWidget
+UCLASS(MinimalAPI)
+class USlider : public UWidget
 {
 	GENERATED_UCLASS_BODY()
 
@@ -119,87 +119,87 @@ public:
 
 	/** Gets the current value of the slider. */
 	UFUNCTION(BlueprintCallable, Category="Behavior")
-	float GetValue() const;
+	UMG_API float GetValue() const;
 
 	/** Get the current value scaled from 0 to 1 */
 	UFUNCTION(BlueprintCallable, Category = "Behavior")
-	float GetNormalizedValue() const;
+	UMG_API float GetNormalizedValue() const;
 
 	/** Sets the current value of the slider. */
 	UFUNCTION(BlueprintCallable, Category="Behavior")
-	void SetValue(float InValue);
+	UMG_API void SetValue(float InValue);
 
 	/** Gets the minimum value of the slider. */
-	float GetMinValue() const;
+	UMG_API float GetMinValue() const;
 
 	/** Sets the minimum value of the slider. */
 	UFUNCTION(BlueprintCallable, Category = "Behavior")
-	void SetMinValue(float InValue);
+	UMG_API void SetMinValue(float InValue);
 
 	/** Gets the maximum value of the slider. */
-	float GetMaxValue() const;
+	UMG_API float GetMaxValue() const;
 
 	/** Sets the maximum value of the slider. */
 	UFUNCTION(BlueprintCallable, Category = "Behavior")
-	void SetMaxValue(float InValue);
+	UMG_API void SetMaxValue(float InValue);
 
 	/** Get the style used by the widget. */
-	const FSliderStyle& GetWidgetStyle() const;
+	UMG_API const FSliderStyle& GetWidgetStyle() const;
 
 	/** Set the style used by the widget. */
-	void SetWidgetStyle(const FSliderStyle& InStyle);
+	UMG_API void SetWidgetStyle(const FSliderStyle& InStyle);
 
 	/** Getg the slider's orientation. */
-	EOrientation GetOrientation() const;
+	UMG_API EOrientation GetOrientation() const;
 
 	/** Sets the slider's orientation. */
-	void SetOrientation(EOrientation InOrientation);
+	UMG_API void SetOrientation(EOrientation InOrientation);
 
 	/** Gets if the slidable area should be indented to fit the handle. */
-	bool HasIndentHandle() const;
+	UMG_API bool HasIndentHandle() const;
 
 	/** Sets if the slidable area should be indented to fit the handle. */
 	UFUNCTION(BlueprintCallable, Category="Behavior")
-	void SetIndentHandle(bool InValue);
+	UMG_API void SetIndentHandle(bool InValue);
 
 	/** Returns true when the handle is fixed. */
-	bool IsLocked() const;
+	UMG_API bool IsLocked() const;
 
 	/** Sets the handle to be interactive or fixed. */
 	UFUNCTION(BlueprintCallable, Category="Behavior")
-	void SetLocked(bool InValue);
+	UMG_API void SetLocked(bool InValue);
 
 	/** Gets the amount to adjust the value by. */
-	float GetStepSize() const;
+	UMG_API float GetStepSize() const;
 
 	/** Sets the amount to adjust the value by, when using a controller or keyboard. */
 	UFUNCTION(BlueprintCallable, Category="Behavior")
-	void SetStepSize(float InValue);
+	UMG_API void SetStepSize(float InValue);
 
 	/** Gets the color of the slider bar. */
-	FLinearColor GetSliderBarColor() const;
+	UMG_API FLinearColor GetSliderBarColor() const;
 
 	/** Sets the color of the slider bar. */
 	UFUNCTION(BlueprintCallable, Category="Appearance")
-	void SetSliderBarColor(FLinearColor InValue);
+	UMG_API void SetSliderBarColor(FLinearColor InValue);
 
 	/** Gets the color of the handle bar */
-	FLinearColor GetSliderHandleColor() const;
+	UMG_API FLinearColor GetSliderHandleColor() const;
 
 	/** Sets the color of the handle bar */
 	UFUNCTION(BlueprintCallable, Category="Appearance")
-	void SetSliderHandleColor(FLinearColor InValue);
+	UMG_API void SetSliderHandleColor(FLinearColor InValue);
 	
 	// UWidget interface
-	virtual void SynchronizeProperties() override;
+	UMG_API virtual void SynchronizeProperties() override;
 	// End of UWidget interface
 
 	// UVisual interface
-	virtual void ReleaseSlateResources(bool bReleaseChildren) override;
+	UMG_API virtual void ReleaseSlateResources(bool bReleaseChildren) override;
 	// End of UVisual interface
 
 #if WITH_EDITOR
-	virtual const FText GetPaletteCategory() override;
+	UMG_API virtual const FText GetPaletteCategory() override;
 #endif
 
 protected:
@@ -207,17 +207,17 @@ protected:
 	TSharedPtr<SSlider> MySlider;
 
 	// UWidget interface
-	virtual TSharedRef<SWidget> RebuildWidget() override;
+	UMG_API virtual TSharedRef<SWidget> RebuildWidget() override;
 	// End of UWidget interface
 
-	void HandleOnValueChanged(float InValue);
-	void HandleOnMouseCaptureBegin();
-	void HandleOnMouseCaptureEnd();
-	void HandleOnControllerCaptureBegin();
-	void HandleOnControllerCaptureEnd();
+	UMG_API void HandleOnValueChanged(float InValue);
+	UMG_API void HandleOnMouseCaptureBegin();
+	UMG_API void HandleOnMouseCaptureEnd();
+	UMG_API void HandleOnControllerCaptureBegin();
+	UMG_API void HandleOnControllerCaptureEnd();
 
 #if WITH_ACCESSIBILITY
-	virtual TSharedPtr<SWidget> GetAccessibleWidget() const override;
+	UMG_API virtual TSharedPtr<SWidget> GetAccessibleWidget() const override;
 #endif
 
 	PROPERTY_BINDING_IMPLEMENTATION(float, Value);
diff --git a/Engine/Source/Runtime/UMG/Public/Components/Spacer.h b/Engine/Source/Runtime/UMG/Public/Components/Spacer.h
--- a/Engine/Source/Runtime/UMG/Public/Components/Spacer.h
+++ b/Engine/Source/Runtime/UMG/Public/Components/Spacer.h
@@ -15,8 +15,8 @@ class SSpacer;
  *
  * * No Children
  */
-UCLASS()
-class UMG_API USpacer : public UWidget
+UCLASS(MinimalAPI)
+class USpacer : public UWidget
 {
 	GENERATED_UCLASS_BODY()
 
@@ -30,24 +30,24 @@ public:
 public:
 
 	/** */
-	FVector2D GetSize() const;
+	UMG_API FVector2D GetSize() const;
 
 	/** Sets the size of the spacer */
 	UFUNCTION(BlueprintCallable, Category="Widget")
-	void SetSize(FVector2D InSize);
+	UMG_API void SetSize(FVector2D InSize);
 	
 	//~ UWidget interface
-	virtual void SynchronizeProperties() override;
-	virtual void ReleaseSlateResources(bool bReleaseChildren) override;
+	UMG_API virtual void SynchronizeProperties() override;
+	UMG_API virtual void ReleaseSlateResources(bool bReleaseChildren) override;
 	//~ End of UVisual interface
 
 #if WITH_EDITOR
-	virtual const FText GetPaletteCategory() override;
+	UMG_API virtual const FText GetPaletteCategory() override;
 #endif
 
 protected:
 	//~ UWidget interface
-	virtual TSharedRef<SWidget> RebuildWidget() override;
+	UMG_API virtual TSharedRef<SWidget> RebuildWidget() override;
 	//~ End of UWidget interface
 
 protected:
diff --git a/Engine/Source/Runtime/UMG/Public/Components/SpinBox.h b/Engine/Source/Runtime/UMG/Public/Components/SpinBox.h
--- a/Engine/Source/Runtime/UMG/Public/Components/SpinBox.h
+++ b/Engine/Source/Runtime/UMG/Public/Components/SpinBox.h
@@ -15,8 +15,8 @@
 /**
  * A numerical entry box that allows for direct entry of the number or allows the user to click and slide the number.
  */
-UCLASS()
-class UMG_API USpinBox : public UWidget
+UCLASS(MinimalAPI)
+class USpinBox : public UWidget
 {
 	GENERATED_UCLASS_BODY()
 
@@ -126,173 +126,173 @@ public:
 
 	/** Get the current value of the spin box. */
 	UFUNCTION(BlueprintCallable, Category="Behavior")
-	float GetValue() const;
+	UMG_API float GetValue() const;
 
 	/** Set the value of the spin box. */
 	UFUNCTION(BlueprintCallable, Category="Behavior")
-	void SetValue(float NewValue);
+	UMG_API void SetValue(float NewValue);
 
 public:
 
 	/** Set the style for the spin box. */
-	const FSpinBoxStyle& GetWidgetStyle() const;
+	UMG_API const FSpinBoxStyle& GetWidgetStyle() const;
 
 	/** Get the style for the spin box. */
-	void SetWidgetStyle(const FSpinBoxStyle& InWidgetStyle);
+	UMG_API void SetWidgetStyle(const FSpinBoxStyle& InWidgetStyle);
 
 	/** Get the current Min Fractional Digits for the spin box. */
 	UFUNCTION(BlueprintCallable, BlueprintGetter, Category = "Behavior")
-	int32 GetMinFractionalDigits() const;
+	UMG_API int32 GetMinFractionalDigits() const;
 
 	/** Set the Min Fractional Digits for the spin box. */
 	UFUNCTION(BlueprintCallable, BlueprintSetter, Category = "Behavior")
-	void SetMinFractionalDigits(int32 NewValue);
+	UMG_API void SetMinFractionalDigits(int32 NewValue);
 
 	/** Get the current Max Fractional Digits for the spin box. */
 	UFUNCTION(BlueprintCallable, BlueprintGetter, Category = "Behavior")
-	int32 GetMaxFractionalDigits() const;
+	UMG_API int32 GetMaxFractionalDigits() const;
 
 	/** Set the Max Fractional Digits for the spin box. */
 	UFUNCTION(BlueprintCallable, BlueprintSetter, Category = "Behavior")
-	void SetMaxFractionalDigits(int32 NewValue);
+	UMG_API void SetMaxFractionalDigits(int32 NewValue);
 
 	/** Get whether the spin box uses delta snap on type. */
 	UFUNCTION(BlueprintCallable, BlueprintGetter, Category = "Behavior")
-	bool GetAlwaysUsesDeltaSnap() const;
+	UMG_API bool GetAlwaysUsesDeltaSnap() const;
 
 	/** Set whether the spin box uses delta snap on type. */
 	UFUNCTION(BlueprintCallable, BlueprintSetter, Category = "Behavior")
-	void SetAlwaysUsesDeltaSnap(bool bNewValue);
+	UMG_API void SetAlwaysUsesDeltaSnap(bool bNewValue);
 
 	/** Get whether the spin box uses the slider feature. */
-	bool GetEnableSlider() const;
+	UMG_API bool GetEnableSlider() const;
 
 	/** Set whether the spin box uses the slider feature. */
-	void SetEnableSlider(bool bNewValue);
+	UMG_API void SetEnableSlider(bool bNewValue);
 
 	/** Get the current delta for the spin box. */
 	UFUNCTION(BlueprintCallable, BlueprintGetter, Category = "Behavior")
-	float GetDelta() const;
+	UMG_API float GetDelta() const;
 
 	/** Set the delta for the spin box. */
 	UFUNCTION(BlueprintCallable, BlueprintSetter, Category = "Behavior")
-	void SetDelta(float NewValue);
+	UMG_API void SetDelta(float NewValue);
 
 	/** Get the current slider exponent for the spin box. */
-	float GetSliderExponent() const;
+	UMG_API float GetSliderExponent() const;
 
 	/** Set the slider exponent for the spin box. */
-	void SetSliderExponent(float NewValue);
+	UMG_API void SetSliderExponent(float NewValue);
 
 	/**  Get the font color and opacity that overrides the style font. */
-	const FSlateFontInfo& GetFont() const;
+	UMG_API const FSlateFontInfo& GetFont() const;
 
 	/** Set the font color and opacity that overrides the style font. */
-	void SetFont(const FSlateFontInfo& InFont);
+	UMG_API void SetFont(const FSlateFontInfo& InFont);
 
 	/** Get the justification for value text. */
-	const ETextJustify::Type GetJustification() const;
+	UMG_API const ETextJustify::Type GetJustification() const;
 
 	/** Set the justification for value text. */
-	void SetJustification(ETextJustify::Type InJustification);
+	UMG_API void SetJustification(ETextJustify::Type InJustification);
 
 	/** Get the minimum width of the spin box. */
-	float GetMinDesiredWidth() const;
+	UMG_API float GetMinDesiredWidth() const;
 
 	/** Set the minimum width of the spin box. */
-	void SetMinDesiredWidth(float NewValue);
+	UMG_API void SetMinDesiredWidth(float NewValue);
 
 	/** Get whether the keyboard focus is removed from the spin box when the value is committed. */
-	bool GetClearKeyboardFocusOnCommit() const;
+	UMG_API bool GetClearKeyboardFocusOnCommit() const;
 
 	/** Set whether the keyboard focus is removed from the spin box when the value is committed. */
-	void SetClearKeyboardFocusOnCommit(bool bNewValue);
+	UMG_API void SetClearKeyboardFocusOnCommit(bool bNewValue);
 
 	/** Get whether to select the text in the spin box when the value is committed. */
-	bool GetSelectAllTextOnCommit() const;
+	UMG_API bool GetSelectAllTextOnCommit() const;
 
 	/** Set whether to select the text in the spin box when the value is committed. */
-	void SetSelectAllTextOnCommit(bool bNewValue);
+	UMG_API void SetSelectAllTextOnCommit(bool bNewValue);
 
 	/** Get the current minimum value that can be manually set in the spin box. */
 	UFUNCTION(BlueprintCallable, Category="Behavior")
-	float GetMinValue() const;
+	UMG_API float GetMinValue() const;
 
 	/** Set the minimum value that can be manually set in the spin box. */
 	UFUNCTION(BlueprintCallable, Category = "Behavior")
-	void SetMinValue(float NewValue);
+	UMG_API void SetMinValue(float NewValue);
 
 	/** Clear the minimum value that can be manually set in the spin box. */
 	UFUNCTION(BlueprintCallable, Category = "Behavior")
-	void ClearMinValue();
+	UMG_API void ClearMinValue();
 
 	/** Get the current maximum value that can be manually set in the spin box. */
 	UFUNCTION(BlueprintCallable, Category = "Behavior")
-	float GetMaxValue() const;
+	UMG_API float GetMaxValue() const;
 
 	/** Set the maximum value that can be manually set in the spin box. */
 	UFUNCTION(BlueprintCallable, Category = "Behavior")
-	void SetMaxValue(float NewValue);
+	UMG_API void SetMaxValue(float NewValue);
 
 	/** Clear the maximum value that can be manually set in the spin box. */
 	UFUNCTION(BlueprintCallable, Category = "Behavior")
-	void ClearMaxValue();
+	UMG_API void ClearMaxValue();
 
 	/** Get the current minimum value that can be specified using the slider. */
 	UFUNCTION(BlueprintCallable, Category = "Behavior")
-	float GetMinSliderValue() const;
+	UMG_API float GetMinSliderValue() const;
 
 	/** Set the minimum value that can be specified using the slider. */
 	UFUNCTION(BlueprintCallable, Category = "Behavior")
-	void SetMinSliderValue(float NewValue);
+	UMG_API void SetMinSliderValue(float NewValue);
 
 	/** Clear the minimum value that can be specified using the slider. */
 	UFUNCTION(BlueprintCallable, Category = "Behavior")
-	void ClearMinSliderValue();
+	UMG_API void ClearMinSliderValue();
 
 	/** Get the current maximum value that can be specified using the slider. */
 	UFUNCTION(BlueprintCallable, Category = "Behavior")
-	float GetMaxSliderValue() const;
+	UMG_API float GetMaxSliderValue() const;
 
 	/** Set the maximum value that can be specified using the slider. */
 	UFUNCTION(BlueprintCallable, Category = "Behavior")
-	void SetMaxSliderValue(float NewValue);
+	UMG_API void SetMaxSliderValue(float NewValue);
 
 	/** Clear the maximum value that can be specified using the slider. */
 	UFUNCTION(BlueprintCallable, Category = "Behavior")
-	void ClearMaxSliderValue();
+	UMG_API void ClearMaxSliderValue();
 
 	/**  */
 	UFUNCTION(BlueprintCallable, Category = "Appearance")
-	void SetForegroundColor(FSlateColor InForegroundColor);
+	UMG_API void SetForegroundColor(FSlateColor InForegroundColor);
 
 	/** Get the foreground color of the spin box. */
-	FSlateColor GetForegroundColor() const;
+	UMG_API FSlateColor GetForegroundColor() const;
 
 public:
 
 	//~ Begin UWidget Interface
-	virtual void SynchronizeProperties() override;
+	UMG_API virtual void SynchronizeProperties() override;
 	//~ End UWidget Interface
 
 	//~ Begin UVisual Interface
-	virtual void ReleaseSlateResources(bool bReleaseChildren) override;
+	UMG_API virtual void ReleaseSlateResources(bool bReleaseChildren) override;
 	//~ End UVisual Interface
 
 #if WITH_EDITOR
-	virtual const FText GetPaletteCategory() override;
+	UMG_API virtual const FText GetPaletteCategory() override;
 #endif
 	//~ End UWidget Interface
 
 protected:
 	//~ Begin UWidget Interface
-	virtual TSharedRef<SWidget> RebuildWidget() override;
+	UMG_API virtual TSharedRef<SWidget> RebuildWidget() override;
 	// End of UWidget
 
-	void HandleOnValueChanged(float InValue);
-	void HandleOnValueCommitted(float InValue, ETextCommit::Type CommitMethod);
-	void HandleOnBeginSliderMovement();
-	void HandleOnEndSliderMovement(float InValue);
+	UMG_API void HandleOnValueChanged(float InValue);
+	UMG_API void HandleOnValueCommitted(float InValue, ETextCommit::Type CommitMethod);
+	UMG_API void HandleOnBeginSliderMovement();
+	UMG_API void HandleOnEndSliderMovement(float InValue);
 
 protected:
 	/** Whether the optional MinValue attribute of the widget is set */
diff --git a/Engine/Source/Runtime/UMG/Public/Components/StackBox.h b/Engine/Source/Runtime/UMG/Public/Components/StackBox.h
--- a/Engine/Source/Runtime/UMG/Public/Components/StackBox.h
+++ b/Engine/Source/Runtime/UMG/Public/Components/StackBox.h
@@ -17,8 +17,8 @@ class SStackBox;
  * * Many Children
  * * Flows Vertical or Horizontal
  */
-UCLASS(meta = (ShortTooltip = "A layout panel for automatically laying child widgets out vertically or horizontally"))
-class UMG_API UStackBox : public UPanelWidget
+UCLASS(meta = (ShortTooltip = "A layout panel for automatically laying child widgets out vertically or horizontally"), MinimalAPI)
+class UStackBox : public UPanelWidget
 {
 	GENERATED_BODY()
 
@@ -31,34 +31,34 @@ private:
 
 public:
 	/** Get the orientation of the stack box. */
-	EOrientation GetOrientation() const;
+	UMG_API EOrientation GetOrientation() const;
 	/** Set the orientation of the stack box. The existing elements will be rearranged. */
-	void SetOrientation(EOrientation InType);
+	UMG_API void SetOrientation(EOrientation InType);
 
 	/** Adds a new child widget to the container. */
 	UFUNCTION(BlueprintCallable, Category="Panel")
-	UStackBoxSlot* AddChildToStackBox(UWidget* Content);
+	UMG_API UStackBoxSlot* AddChildToStackBox(UWidget* Content);
 
 	/** Replace the widget at the given index it with a different widget. */
 	UFUNCTION(BlueprintCallable, Category = "Panel")
-	bool ReplaceStackBoxChildAt(int32 Index, UWidget* Content);
+	UMG_API bool ReplaceStackBoxChildAt(int32 Index, UWidget* Content);
 
 #if WITH_EDITOR
 	//~ UWidget interface
-	virtual const FText GetPaletteCategory() override;
+	UMG_API virtual const FText GetPaletteCategory() override;
 	//~ End UWidget interface
 #endif
 
 protected:
 	//~ UPanelWidget
-	virtual UClass* GetSlotClass() const override;
-	virtual void OnSlotAdded(UPanelSlot* Slot) override;
-	virtual void OnSlotRemoved(UPanelSlot* Slot) override;
-	virtual TSharedRef<SWidget> RebuildWidget() override;
-	virtual void SynchronizeProperties() override;
+	UMG_API virtual UClass* GetSlotClass() const override;
+	UMG_API virtual void OnSlotAdded(UPanelSlot* Slot) override;
+	UMG_API virtual void OnSlotRemoved(UPanelSlot* Slot) override;
+	UMG_API virtual TSharedRef<SWidget> RebuildWidget() override;
+	UMG_API virtual void SynchronizeProperties() override;
 	//~ End UPanelWidget
 
-	virtual void ReleaseSlateResources(bool bReleaseChildren) override;
+	UMG_API virtual void ReleaseSlateResources(bool bReleaseChildren) override;
 
 private:
 	TSharedPtr<SStackBox> MyBox;
diff --git a/Engine/Source/Runtime/UMG/Public/Components/StackBoxSlot.h b/Engine/Source/Runtime/UMG/Public/Components/StackBoxSlot.h
--- a/Engine/Source/Runtime/UMG/Public/Components/StackBoxSlot.h
+++ b/Engine/Source/Runtime/UMG/Public/Components/StackBoxSlot.h
@@ -13,8 +13,8 @@
 class SPanel;
 
 /** The Slot for the UStackBox, contains the widget that is flowed vertically or horizontally. */
-UCLASS()
-class UMG_API UStackBoxSlot : public UPanelSlot
+UCLASS(MinimalAPI)
+class UStackBoxSlot : public UPanelSlot
 {
 	GENERATED_BODY()
 
@@ -40,33 +40,33 @@ private:
 	SStackBox::FSlot* Slot = nullptr;
 
 public:
-	FMargin GetPadding() const;
-	void SetPadding(FMargin InPadding);
+	UMG_API FMargin GetPadding() const;
+	UMG_API void SetPadding(FMargin InPadding);
 
-	FSlateChildSize GetSize() const;
-	void SetSize(FSlateChildSize InSize);
+	UMG_API FSlateChildSize GetSize() const;
+	UMG_API void SetSize(FSlateChildSize InSize);
 
-	EHorizontalAlignment GetHorizontalAlignment() const;
-	void SetHorizontalAlignment(EHorizontalAlignment InHorizontalAlignment);
+	UMG_API EHorizontalAlignment GetHorizontalAlignment() const;
+	UMG_API void SetHorizontalAlignment(EHorizontalAlignment InHorizontalAlignment);
 
-	EVerticalAlignment GetVerticalAlignment() const;
-	void SetVerticalAlignment(EVerticalAlignment InVerticalAlignment);
+	UMG_API EVerticalAlignment GetVerticalAlignment() const;
+	UMG_API void SetVerticalAlignment(EVerticalAlignment InVerticalAlignment);
 
 public:
 	//~ UPanelSlot interface
-	virtual void SynchronizeProperties() override;
+	UMG_API virtual void SynchronizeProperties() override;
 	//~ End of UPanelSlot interface
 
-	virtual void ReleaseSlateResources(bool bReleaseChildren) override;
+	UMG_API virtual void ReleaseSlateResources(bool bReleaseChildren) override;
 
 	/** Builds the underlying FSlot for the Slate layout panel. */
-	void BuildSlot(TSharedRef<SStackBox> InBox);
+	UMG_API void BuildSlot(TSharedRef<SStackBox> InBox);
 
 	/** Replace the slot content. */
-	void ReplaceContent(UWidget* Content);
+	UMG_API void ReplaceContent(UWidget* Content);
 
 #if WITH_EDITOR
-	virtual bool NudgeByDesigner(const FVector2D& NudgeDirection, const TOptional<int32>& GridSnapSize) override;
-	virtual void SynchronizeFromTemplate(const UPanelSlot* const TemplateSlot) override;
+	UMG_API virtual bool NudgeByDesigner(const FVector2D& NudgeDirection, const TOptional<int32>& GridSnapSize) override;
+	UMG_API virtual void SynchronizeFromTemplate(const UPanelSlot* const TemplateSlot) override;
 #endif //WITH_EDITOR
 };
diff --git a/Engine/Source/Runtime/UMG/Public/Components/TextBlock.h b/Engine/Source/Runtime/UMG/Public/Components/TextBlock.h
--- a/Engine/Source/Runtime/UMG/Public/Components/TextBlock.h
+++ b/Engine/Source/Runtime/UMG/Public/Components/TextBlock.h
@@ -19,8 +19,8 @@ class UMaterialInterface;
  * * No Children
  * * Text
  */
-UCLASS(meta=(DisplayName="Text"))
-class UMG_API UTextBlock : public UTextLayoutWidget
+UCLASS(meta=(DisplayName="Text"), MinimalAPI)
+class UTextBlock : public UTextLayoutWidget
 {
 	GENERATED_UCLASS_BODY()
 
@@ -92,7 +92,7 @@ public:
 	 * @return The widget text
 	 */
 	UFUNCTION(BlueprintCallable, Category="Widget", meta=(DisplayName="GetText (Text)"))
-	FText GetText() const;
+	UMG_API FText GetText() const;
 
 	/**
 	 * Directly sets the widget text.
@@ -100,10 +100,10 @@ public:
 	 * @param InText The text to assign to the widget
 	 */
 	UFUNCTION(BlueprintCallable, Category="Widget", meta=(DisplayName="SetText (Text)"))
-	virtual void SetText(FText InText);
+	UMG_API virtual void SetText(FText InText);
 
 	/** @return the color and opacity of the text in this text block. */
-	FSlateColor GetColorAndOpacity() const;
+	UMG_API FSlateColor GetColorAndOpacity() const;
 
 	/**  
 	 * Sets the color and opacity of the text in this text block
@@ -111,7 +111,7 @@ public:
 	 * @param InColorAndOpacity		The new text color and opacity
 	 */
 	UFUNCTION(BlueprintCallable, Category="Appearance")
-	void SetColorAndOpacity(FSlateColor InColorAndOpacity);
+	UMG_API void SetColorAndOpacity(FSlateColor InColorAndOpacity);
 
 	/**
 	* Sets the opacity of the text in this text block
@@ -119,10 +119,10 @@ public:
 	* @param InOpacity		The new text opacity
 	*/
 	UFUNCTION(BlueprintCallable, Category="Appearance")
-	void SetOpacity(float InOpacity);
+	UMG_API void SetOpacity(float InOpacity);
 
 	/** @return the color and opacity of the text drop shadow. */
-	FLinearColor GetShadowColorAndOpacity() const;
+	UMG_API FLinearColor GetShadowColorAndOpacity() const;
 
 	/**  
 	 * Sets the color and opacity of the text drop shadow
@@ -131,10 +131,10 @@ public:
 	 * @param InShadowColorAndOpacity		The new drop shadow color and opacity
 	 */
 	UFUNCTION(BlueprintCallable, Category="Appearance")
-	void SetShadowColorAndOpacity(FLinearColor InShadowColorAndOpacity);
+	UMG_API void SetShadowColorAndOpacity(FLinearColor InShadowColorAndOpacity);
 
 	/** @return the offset that the text drop shadow should be drawn at. */
-	FVector2D GetShadowOffset() const;
+	UMG_API FVector2D GetShadowOffset() const;
 
 	/**  
 	 * Sets the offset that the text drop shadow should be drawn at
@@ -142,10 +142,10 @@ public:
 	 * @param InShadowOffset		The new offset
 	 */
 	UFUNCTION(BlueprintCallable, Category="Appearance")
-	void SetShadowOffset(FVector2D InShadowOffset);
+	UMG_API void SetShadowOffset(FVector2D InShadowOffset);
 
 	/** @return the font info for this text block. */
-	const FSlateFontInfo& GetFont() const;
+	UMG_API const FSlateFontInfo& GetFont() const;
 
 	/**
 	 * Dynamically set the font info for this text block
@@ -153,10 +153,10 @@ public:
 	 * @param InFontInfo The new font info
 	 */
 	UFUNCTION(BlueprintCallable, Category="Appearance")
-	void SetFont(FSlateFontInfo InFontInfo);
+	UMG_API void SetFont(FSlateFontInfo InFontInfo);
 
 	/** @return the brush to use to strike through text.  */
-	const FSlateBrush& GetStrikeBrush() const;
+	UMG_API const FSlateBrush& GetStrikeBrush() const;
 
 	/**
 	 * Dynamically set the strike brush for this text block
@@ -164,10 +164,10 @@ public:
 	 * @param InStrikeBrush The new brush to use to strike through text
 	 */
 	UFUNCTION(BlueprintCallable, Category="Appearance")
-	void SetStrikeBrush(FSlateBrush InStrikeBrush);
+	UMG_API void SetStrikeBrush(FSlateBrush InStrikeBrush);
 
 	/** @return the minimum desired width for this text block */
-	float GetMinDesiredWidth() const;
+	UMG_API float GetMinDesiredWidth() const;
 
 	/**
 	 *  Set the minimum desired width for this text block
@@ -175,7 +175,7 @@ public:
 	 *  @param InMinDesiredWidth new minimum desired width
 	 */
 	UFUNCTION(BlueprintCallable, Category="Appearance")
-	void SetMinDesiredWidth(float InMinDesiredWidth);
+	UMG_API void SetMinDesiredWidth(float InMinDesiredWidth);
 	
 	/**
 	 * Set the auto wrap for this text block.
@@ -183,10 +183,10 @@ public:
 	 * @param InAutoTextWrap to turn wrap on or off.
 	 */
 	UFUNCTION(BlueprintCallable, Category="Appearance")
-	void SetAutoWrapText(bool InAutoTextWrap);
+	UMG_API void SetAutoWrapText(bool InAutoTextWrap);
 
 	/** @return the text transformation policy for this text block. */
-	ETextTransformPolicy GetTextTransformPolicy() const;
+	UMG_API ETextTransformPolicy GetTextTransformPolicy() const;
 
 	/**
 	 * Set the text transformation policy for this text block.
@@ -194,10 +194,10 @@ public:
 	 * @param InTransformPolicy the new text transformation policy.
 	 */
 	UFUNCTION(BlueprintCallable, Category="Appearance")
-	void SetTextTransformPolicy(ETextTransformPolicy InTransformPolicy);
+	UMG_API void SetTextTransformPolicy(ETextTransformPolicy InTransformPolicy);
 
 	/** @return the text overflow policy for this text block. */
-	ETextOverflowPolicy GetTextOverflowPolicy() const;
+	UMG_API ETextOverflowPolicy GetTextOverflowPolicy() const;
 
 	/**
 	 * Set the text overflow policy for this text block.
@@ -205,21 +205,21 @@ public:
 	 * @param InOverflowPolicy the new text overflow policy.
 	 */
 	UFUNCTION(BlueprintCallable, Category="Appearance")
-	void SetTextOverflowPolicy(ETextOverflowPolicy InOverflowPolicy);
+	UMG_API void SetTextOverflowPolicy(ETextOverflowPolicy InOverflowPolicy);
 
 	UFUNCTION(BlueprintCallable, Category="Appearance")
-	void SetFontMaterial(UMaterialInterface* InMaterial);
+	UMG_API void SetFontMaterial(UMaterialInterface* InMaterial);
 
 	UFUNCTION(BlueprintCallable, Category="Appearance")
-	void SetFontOutlineMaterial(UMaterialInterface* InMaterial);
+	UMG_API void SetFontOutlineMaterial(UMaterialInterface* InMaterial);
 
 	/**  */
 	UFUNCTION(BlueprintCallable, Category="Appearance")
-	UMaterialInstanceDynamic* GetDynamicFontMaterial();
+	UMG_API UMaterialInstanceDynamic* GetDynamicFontMaterial();
 
 	/**  */
 	UFUNCTION(BlueprintCallable, Category="Appearance")
-	UMaterialInstanceDynamic* GetDynamicOutlineMaterial();
+	UMG_API UMaterialInstanceDynamic* GetDynamicOutlineMaterial();
 
 protected:
 	virtual void OnTextChanged()
@@ -232,27 +232,27 @@ protected:
 
 public:
 	//~ Begin UTextLayoutWidget Interface
-	virtual void SetJustification(ETextJustify::Type InJustification) override;
+	UMG_API virtual void SetJustification(ETextJustify::Type InJustification) override;
 	//~ End UTextLayoutWidget Interface
 
 	//~ Begin UWidget Interface
-	virtual void SynchronizeProperties() override;
+	UMG_API virtual void SynchronizeProperties() override;
 	//~ End UWidget Interface
 
 	//~ Begin UVisual Interface
-	virtual void ReleaseSlateResources(bool bReleaseChildren) override;
+	UMG_API virtual void ReleaseSlateResources(bool bReleaseChildren) override;
 	//~ End UVisual Interface
 
 #if WITH_EDITOR
 	//~ Begin UWidget Interface
-	virtual const FText GetPaletteCategory() override;
-	virtual void OnCreationFromPalette() override;
-	virtual bool CanEditChange(const FProperty* InProperty) const override;
+	UMG_API virtual const FText GetPaletteCategory() override;
+	UMG_API virtual void OnCreationFromPalette() override;
+	UMG_API virtual bool CanEditChange(const FProperty* InProperty) const override;
 	//~ End UWidget Interface
 
-	virtual FString GetLabelMetadata() const override;
+	UMG_API virtual FString GetLabelMetadata() const override;
 
-	void HandleTextCommitted(const FText& InText, ETextCommit::Type CommitteType);
+	UMG_API void HandleTextCommitted(const FText& InText, ETextCommit::Type CommitteType);
 #endif
 
 protected:
@@ -265,17 +265,17 @@ protected:
 	bool bSimpleTextMode;
 
 	//~ Begin UWidget Interface
-	virtual TSharedRef<SWidget> RebuildWidget() override;
-	virtual void OnBindingChanged(const FName& Property) override;
+	UMG_API virtual TSharedRef<SWidget> RebuildWidget() override;
+	UMG_API virtual void OnBindingChanged(const FName& Property) override;
 	//~ End UWidget Interface
 
 	/** Get the text that should be displayed in the internal Slate widget (allows flags to mutate the display text without modifying the persistent designer property data) */
-	virtual TAttribute<FText> GetDisplayText();
+	UMG_API virtual TAttribute<FText> GetDisplayText();
 
-	EVisibility GetTextWarningImageVisibility() const;
+	UMG_API EVisibility GetTextWarningImageVisibility() const;
 
 #if WITH_ACCESSIBILITY
-	virtual TSharedPtr<SWidget> GetAccessibleWidget() const override;
+	UMG_API virtual TSharedPtr<SWidget> GetAccessibleWidget() const override;
 #endif
 
 protected:
diff --git a/Engine/Source/Runtime/UMG/Public/Components/TextWidgetTypes.h b/Engine/Source/Runtime/UMG/Public/Components/TextWidgetTypes.h
--- a/Engine/Source/Runtime/UMG/Public/Components/TextWidgetTypes.h
+++ b/Engine/Source/Runtime/UMG/Public/Components/TextWidgetTypes.h
@@ -17,11 +17,11 @@ enum class ETextShapingMethod : uint8;
  * Contains the common options that should be exposed for the underlying Slate widget.
  */
 USTRUCT(BlueprintType)
-struct UMG_API FShapedTextOptions
+struct FShapedTextOptions
 {
 	GENERATED_USTRUCT_BODY()
 
-	FShapedTextOptions();
+	UMG_API FShapedTextOptions();
 
 	/** Synchronize the properties with the given widget. A template as the Slate widgets conform to the same API, but don't derive from a common base. */
 	template <typename TWidgetType>
@@ -53,8 +53,8 @@ struct UMG_API FShapedTextOptions
  * Base class for all widgets that use a text layout.
  * Contains the common options that should be exposed for the underlying Slate widget.
  */
-UCLASS(Abstract, BlueprintType)
-class UMG_API UTextLayoutWidget : public UWidget
+UCLASS(Abstract, BlueprintType, MinimalAPI)
+class UTextLayoutWidget : public UWidget
 {
 	GENERATED_UCLASS_BODY()
 
diff --git a/Engine/Source/Runtime/UMG/Public/Components/Throbber.h b/Engine/Source/Runtime/UMG/Public/Components/Throbber.h
--- a/Engine/Source/Runtime/UMG/Public/Components/Throbber.h
+++ b/Engine/Source/Runtime/UMG/Public/Components/Throbber.h
@@ -15,8 +15,8 @@ class USlateBrushAsset;
 /**
  * A Throbber widget that shows several zooming circles in a row.
  */
-UCLASS()
-class UMG_API UThrobber : public UWidget
+UCLASS(MinimalAPI)
+class UThrobber : public UWidget
 {
 	GENERATED_UCLASS_BODY()
 
@@ -51,42 +51,42 @@ public:
 
 	/** Sets how many pieces there are */
 	UFUNCTION(BlueprintCallable, Category="Appearance")
-	void SetNumberOfPieces(int32 InNumberOfPieces);
-	int32 GetNumberOfPieces() const;
+	UMG_API void SetNumberOfPieces(int32 InNumberOfPieces);
+	UMG_API int32 GetNumberOfPieces() const;
 
 	/** Sets whether the pieces animate horizontally. */
 	UFUNCTION(BlueprintCallable, Category="Appearance")
-	void SetAnimateHorizontally(bool bInAnimateHorizontally);
-	bool IsAnimateHorizontally() const;
+	UMG_API void SetAnimateHorizontally(bool bInAnimateHorizontally);
+	UMG_API bool IsAnimateHorizontally() const;
 
 	/** Sets whether the pieces animate vertically. */
 	UFUNCTION(BlueprintCallable, Category="Appearance")
-	void SetAnimateVertically(bool bInAnimateVertically);
-	bool IsAnimateVertically() const;
+	UMG_API void SetAnimateVertically(bool bInAnimateVertically);
+	UMG_API bool IsAnimateVertically() const;
 
 	/** Sets whether the pieces animate their opacity. */
 	UFUNCTION(BlueprintCallable, Category = "Appearance")
-	void SetAnimateOpacity(bool bInAnimateOpacity);
-	bool IsAnimateOpacity() const;
+	UMG_API void SetAnimateOpacity(bool bInAnimateOpacity);
+	UMG_API bool IsAnimateOpacity() const;
 
-	void SetImage(const FSlateBrush& Brush);
-	const FSlateBrush& GetImage() const;
+	UMG_API void SetImage(const FSlateBrush& Brush);
+	UMG_API const FSlateBrush& GetImage() const;
 
 	//~ Begin UWidget Interface
-	virtual void SynchronizeProperties() override;
+	UMG_API virtual void SynchronizeProperties() override;
 	//~ End UWidget Interface
 
 	//~ Begin UVisual Interface
-	virtual void ReleaseSlateResources(bool bReleaseChildren) override;
+	UMG_API virtual void ReleaseSlateResources(bool bReleaseChildren) override;
 	//~ End UVisual Interface
 
 #if WITH_EDITOR
-	virtual const FText GetPaletteCategory() override;
+	UMG_API virtual const FText GetPaletteCategory() override;
 #endif
 
 protected:
 	//~ Begin UWidget Interface
-	virtual TSharedRef<SWidget> RebuildWidget() override;
+	UMG_API virtual TSharedRef<SWidget> RebuildWidget() override;
 	//~ End UWidget Interface
 
 private:
diff --git a/Engine/Source/Runtime/UMG/Public/Components/TileView.h b/Engine/Source/Runtime/UMG/Public/Components/TileView.h
--- a/Engine/Source/Runtime/UMG/Public/Components/TileView.h
+++ b/Engine/Source/Runtime/UMG/Public/Components/TileView.h
@@ -11,23 +11,23 @@
  *
  * To make a widget usable as an entry in a TileView, it must inherit from the IUserObjectListEntry interface.
  */
-UCLASS()
-class UMG_API UTileView : public UListView
+UCLASS(MinimalAPI)
+class UTileView : public UListView
 {
 	GENERATED_BODY()
 
 public:
-	UTileView(const FObjectInitializer& ObjectInitializer);
+	UMG_API UTileView(const FObjectInitializer& ObjectInitializer);
 
-	virtual void ReleaseSlateResources(bool bReleaseChildren) override;
+	UMG_API virtual void ReleaseSlateResources(bool bReleaseChildren) override;
 
 	/** Sets the height of every tile entry */
 	UFUNCTION(BlueprintCallable, Category = TileView)
-	void SetEntryHeight(float NewHeight);
+	UMG_API void SetEntryHeight(float NewHeight);
 
 	/** Sets the width of every tile entry */
 	UFUNCTION(BlueprintCallable, Category = TileView)
-	void SetEntryWidth(float NewWidth);
+	UMG_API void SetEntryWidth(float NewWidth);
 
 	/** Gets the height of tile entries */
 	UFUNCTION(BlueprintPure, Category = TileView)
@@ -38,11 +38,11 @@ public:
 	float GetEntryWidth() const { return EntryWidth; }
 
 protected:
-	virtual TSharedRef<STableViewBase> RebuildListWidget() override;
-	virtual FMargin GetDesiredEntryPadding(UObject* Item) const override;
+	UMG_API virtual TSharedRef<STableViewBase> RebuildListWidget() override;
+	UMG_API virtual FMargin GetDesiredEntryPadding(UObject* Item) const override;
 
-	float GetTotalEntryHeight() const;
-	float GetTotalEntryWidth() const;
+	UMG_API float GetTotalEntryHeight() const;
+	UMG_API float GetTotalEntryWidth() const;
 
 	/** STileView construction helper - useful if using a custom STileView subclass */
 	template <template<typename> class TileViewT = STileView>
@@ -82,7 +82,7 @@ protected:
 private:
 	/** Returns whether the TileView is left, right or center aligned. */
 	UFUNCTION()
-	bool IsAligned() const;
+	UMG_API bool IsAligned() const;
 
 protected:
 	/** The height of each tile */
diff --git a/Engine/Source/Runtime/UMG/Public/Components/TreeView.h b/Engine/Source/Runtime/UMG/Public/Components/TreeView.h
--- a/Engine/Source/Runtime/UMG/Public/Components/TreeView.h
+++ b/Engine/Source/Runtime/UMG/Public/Components/TreeView.h
@@ -16,26 +16,26 @@ DECLARE_DYNAMIC_MULTICAST_DELEGATE_TwoParams(FOnItemExpansionChangedDynamic, UOb
  *
  * To make a widget usable as an entry in a TreeView, it must inherit from the IUserObjectListEntry interface.
  */
-UCLASS()
-class UMG_API UTreeView : public UListView
+UCLASS(MinimalAPI)
+class UTreeView : public UListView
 {
 	GENERATED_BODY()
 
 public:
-	UTreeView(const FObjectInitializer& ObjectInitializer);
-	virtual void ReleaseSlateResources(bool bReleaseChildren) override;
+	UMG_API UTreeView(const FObjectInitializer& ObjectInitializer);
+	UMG_API virtual void ReleaseSlateResources(bool bReleaseChildren) override;
 	
 	/** Attempts to expand/collapse the given item (only relevant if the item has children) */
 	UFUNCTION(BlueprintCallable, Category = TreeView)
-	void SetItemExpansion(UObject* Item, bool bExpandItem);
+	UMG_API void SetItemExpansion(UObject* Item, bool bExpandItem);
 
 	/** Expands all items with children */
 	UFUNCTION(BlueprintCallable, Category = TreeView)
-	void ExpandAll();
+	UMG_API void ExpandAll();
 
 	/** Collapses all currently expanded items */
 	UFUNCTION(BlueprintCallable, Category = TreeView)
-	void CollapseAll();
+	UMG_API void CollapseAll();
 
 	template <typename ObjectT>
 	void SetOnGetItemChildren(ObjectT* InUserObject, typename TSlateDelegates<UObject*>::FOnGetChildren::TMethodPtr<ObjectT> InMethodPtr)
@@ -50,11 +50,11 @@ public:
 	}
 
 protected:
-	virtual TSharedRef<STableViewBase> RebuildListWidget() override;
-	virtual void OnItemClickedInternal(UObject* ListItem) override;
+	UMG_API virtual TSharedRef<STableViewBase> RebuildListWidget() override;
+	UMG_API virtual void OnItemClickedInternal(UObject* ListItem) override;
 
-	virtual void OnItemExpansionChangedInternal(UObject* Item, bool bIsExpanded) override;
-	virtual void OnGetChildrenInternal(UObject* Item, TArray<UObject*>& OutChildren) const override;
+	UMG_API virtual void OnItemExpansionChangedInternal(UObject* Item, bool bIsExpanded) override;
+	UMG_API virtual void OnGetChildrenInternal(UObject* Item, TArray<UObject*>& OutChildren) const override;
 
 	/** STreeView construction helper - useful if using a custom STreeView subclass */
 	template <template<typename> class TreeViewT = STreeView>
diff --git a/Engine/Source/Runtime/UMG/Public/Components/UniformGridPanel.h b/Engine/Source/Runtime/UMG/Public/Components/UniformGridPanel.h
--- a/Engine/Source/Runtime/UMG/Public/Components/UniformGridPanel.h
+++ b/Engine/Source/Runtime/UMG/Public/Components/UniformGridPanel.h
@@ -15,8 +15,8 @@ class UUniformGridSlot;
 /**
  * A panel that evenly divides up available space between all of its children.
  */
-UCLASS()
-class UMG_API UUniformGridPanel : public UPanelWidget
+UCLASS(MinimalAPI)
+class UUniformGridPanel : public UPanelWidget
 {
 	GENERATED_UCLASS_BODY()
 
@@ -40,49 +40,49 @@ public:
 public:
 
 	/** */
-	FMargin GetSlotPadding() const;
+	UMG_API FMargin GetSlotPadding() const;
 
 	/**  */
 	UFUNCTION(BlueprintCallable, Category="Child Layout")
-	void SetSlotPadding(FMargin InSlotPadding);
+	UMG_API void SetSlotPadding(FMargin InSlotPadding);
 
 	/** */
-	float GetMinDesiredSlotWidth() const;
+	UMG_API float GetMinDesiredSlotWidth() const;
 
 	/**  */
 	UFUNCTION(BlueprintCallable, Category="Child Layout")
-	void SetMinDesiredSlotWidth(float InMinDesiredSlotWidth);
+	UMG_API void SetMinDesiredSlotWidth(float InMinDesiredSlotWidth);
 
 	/** */
-	float GetMinDesiredSlotHeight() const;
+	UMG_API float GetMinDesiredSlotHeight() const;
 
 	/**  */
 	UFUNCTION(BlueprintCallable, Category="Child Layout")
-	void SetMinDesiredSlotHeight(float InMinDesiredSlotHeight);
+	UMG_API void SetMinDesiredSlotHeight(float InMinDesiredSlotHeight);
 
 public:
 
 	/**  */
 	UFUNCTION(BlueprintCallable, Category="Widget")
-	UUniformGridSlot* AddChildToUniformGrid(UWidget* Content, int32 InRow = 0, int32 InColumn = 0);
+	UMG_API UUniformGridSlot* AddChildToUniformGrid(UWidget* Content, int32 InRow = 0, int32 InColumn = 0);
 
 public:
 
 	//~ UWidget interface
-	virtual void SynchronizeProperties() override;
-	virtual void ReleaseSlateResources(bool bReleaseChildren) override;
+	UMG_API virtual void SynchronizeProperties() override;
+	UMG_API virtual void ReleaseSlateResources(bool bReleaseChildren) override;
 
 #if WITH_EDITOR
-	virtual const FText GetPaletteCategory() override;
+	UMG_API virtual const FText GetPaletteCategory() override;
 #endif
 	//~ End of UWidget interface
 
 protected:
 
 	// UPanelWidget
-	virtual UClass* GetSlotClass() const override;
-	virtual void OnSlotAdded(UPanelSlot* Slot) override;
-	virtual void OnSlotRemoved(UPanelSlot* Slot) override;
+	UMG_API virtual UClass* GetSlotClass() const override;
+	UMG_API virtual void OnSlotAdded(UPanelSlot* Slot) override;
+	UMG_API virtual void OnSlotRemoved(UPanelSlot* Slot) override;
 	// End UPanelWidget
 
 protected:
@@ -91,6 +91,6 @@ protected:
 
 protected:
 	//~ UWidget interface
-	virtual TSharedRef<SWidget> RebuildWidget() override;
+	UMG_API virtual TSharedRef<SWidget> RebuildWidget() override;
 	//~ End of UWidget interface
 };
diff --git a/Engine/Source/Runtime/UMG/Public/Components/UniformGridSlot.h b/Engine/Source/Runtime/UMG/Public/Components/UniformGridSlot.h
--- a/Engine/Source/Runtime/UMG/Public/Components/UniformGridSlot.h
+++ b/Engine/Source/Runtime/UMG/Public/Components/UniformGridSlot.h
@@ -14,8 +14,8 @@
  * A slot for UUniformGridPanel, these slots all share the same size as the largest slot
  * in the grid.
  */
-UCLASS()
-class UMG_API UUniformGridSlot : public UPanelSlot
+UCLASS(MinimalAPI)
+class UUniformGridSlot : public UPanelSlot
 {
 	GENERATED_UCLASS_BODY()
 
@@ -44,47 +44,47 @@ public:
 public:
 
 	/** Gets the row index of the slot */
-	int32 GetRow() const;
+	UMG_API int32 GetRow() const;
 
 	/** Sets the row index of the slot, this determines what cell the slot is in the panel */
 	UFUNCTION(BlueprintCallable, Category="Layout|Uniform Grid Slot")
-	void SetRow(int32 InRow);
+	UMG_API void SetRow(int32 InRow);
 
 	/** Gets the column index of the slot. */
-	int32 GetColumn() const;
+	UMG_API int32 GetColumn() const;
 
 	/** Sets the column index of the slot, this determines what cell the slot is in the panel */
 	UFUNCTION(BlueprintCallable, Category="Layout|Uniform Grid Slot")
-	void SetColumn(int32 InColumn);
+	UMG_API void SetColumn(int32 InColumn);
 
 	/**  */
-	EHorizontalAlignment GetHorizontalAlignment() const;
+	UMG_API EHorizontalAlignment GetHorizontalAlignment() const;
 
 	/**  */
 	UFUNCTION(BlueprintCallable, Category="Layout|Uniform Grid Slot")
-	void SetHorizontalAlignment(EHorizontalAlignment InHorizontalAlignment);
+	UMG_API void SetHorizontalAlignment(EHorizontalAlignment InHorizontalAlignment);
 
 	/**  */
-	EVerticalAlignment GetVerticalAlignment() const;
+	UMG_API EVerticalAlignment GetVerticalAlignment() const;
 
 	/**  */
 	UFUNCTION(BlueprintCallable, Category="Layout|Uniform Grid Slot")
-	void SetVerticalAlignment(EVerticalAlignment InVerticalAlignment);
+	UMG_API void SetVerticalAlignment(EVerticalAlignment InVerticalAlignment);
 
 public:
 
 	//~ UPanelSlot interface
-	virtual void SynchronizeProperties() override;
+	UMG_API virtual void SynchronizeProperties() override;
 	//~ End of UPanelSlot interface
 
 	/** Builds the underlying FSlot for the Slate layout panel. */
-	void BuildSlot(TSharedRef<SUniformGridPanel> GridPanel);
+	UMG_API void BuildSlot(TSharedRef<SUniformGridPanel> GridPanel);
 
-	virtual void ReleaseSlateResources(bool bReleaseChildren) override;
+	UMG_API virtual void ReleaseSlateResources(bool bReleaseChildren) override;
 
 #if WITH_EDITOR
-	virtual bool NudgeByDesigner(const FVector2D& NudgeDirection, const TOptional<int32>& GridSnapSize) override;
-	virtual void SynchronizeFromTemplate(const UPanelSlot* const TemplateSlot) override;
+	UMG_API virtual bool NudgeByDesigner(const FVector2D& NudgeDirection, const TOptional<int32>& GridSnapSize) override;
+	UMG_API virtual void SynchronizeFromTemplate(const UPanelSlot* const TemplateSlot) override;
 #endif //WITH_EDITOR
 
 private:
diff --git a/Engine/Source/Runtime/UMG/Public/Components/VerticalBox.h b/Engine/Source/Runtime/UMG/Public/Components/VerticalBox.h
--- a/Engine/Source/Runtime/UMG/Public/Components/VerticalBox.h
+++ b/Engine/Source/Runtime/UMG/Public/Components/VerticalBox.h
@@ -18,8 +18,8 @@ class UVerticalBoxSlot;
  * * Many Children
  * * Flows Vertical
  */
-UCLASS(meta = (ShortTooltip = "A layout panel for automatically laying child widgets out vertically"))
-class UMG_API UVerticalBox : public UPanelWidget
+UCLASS(meta = (ShortTooltip = "A layout panel for automatically laying child widgets out vertically"), MinimalAPI)
+class UVerticalBox : public UPanelWidget
 {
 	GENERATED_UCLASS_BODY()
 
@@ -27,23 +27,23 @@ public:
 
 	/**  */
 	UFUNCTION(BlueprintCallable, Category="Panel")
-	UVerticalBoxSlot* AddChildToVerticalBox(UWidget* Content);
+	UMG_API UVerticalBoxSlot* AddChildToVerticalBox(UWidget* Content);
 
 #if WITH_EDITOR
 	// UWidget interface
-	virtual const FText GetPaletteCategory() override;
+	UMG_API virtual const FText GetPaletteCategory() override;
 	// End UWidget interface
 #endif
 
 protected:
 
 	// UPanelWidget
-	virtual UClass* GetSlotClass() const override;
-	virtual void OnSlotAdded(UPanelSlot* Slot) override;
-	virtual void OnSlotRemoved(UPanelSlot* Slot) override;
+	UMG_API virtual UClass* GetSlotClass() const override;
+	UMG_API virtual void OnSlotAdded(UPanelSlot* Slot) override;
+	UMG_API virtual void OnSlotRemoved(UPanelSlot* Slot) override;
 	// End UPanelWidget
 
-	virtual void ReleaseSlateResources(bool bReleaseChildren) override;
+	UMG_API virtual void ReleaseSlateResources(bool bReleaseChildren) override;
 
 protected:
 
@@ -51,6 +51,6 @@ protected:
 
 protected:
 	// UWidget interface
-	virtual TSharedRef<SWidget> RebuildWidget() override;
+	UMG_API virtual TSharedRef<SWidget> RebuildWidget() override;
 	// End of UWidget interface
 };
diff --git a/Engine/Source/Runtime/UMG/Public/Components/VerticalBoxSlot.h b/Engine/Source/Runtime/UMG/Public/Components/VerticalBoxSlot.h
--- a/Engine/Source/Runtime/UMG/Public/Components/VerticalBoxSlot.h
+++ b/Engine/Source/Runtime/UMG/Public/Components/VerticalBoxSlot.h
@@ -13,8 +13,8 @@
 #include "VerticalBoxSlot.generated.h"
 
 /** The Slot for the UVerticalBox, contains the widget that is flowed vertically */
-UCLASS()
-class UMG_API UVerticalBoxSlot : public UPanelSlot
+UCLASS(MinimalAPI)
+class UVerticalBoxSlot : public UPanelSlot
 {
 	GENERATED_UCLASS_BODY()
 
@@ -47,39 +47,39 @@ private:
 
 public:
 
-	FMargin GetPadding() const;
+	UMG_API FMargin GetPadding() const;
 
 	UFUNCTION(BlueprintCallable, Category="Layout|Vertical Box Slot")
-	void SetPadding(FMargin InPadding);
+	UMG_API void SetPadding(FMargin InPadding);
 
-	FSlateChildSize GetSize() const;
+	UMG_API FSlateChildSize GetSize() const;
 
 	UFUNCTION(BlueprintCallable, Category="Layout|Vertical Box Slot")
-	void SetSize(FSlateChildSize InSize);
+	UMG_API void SetSize(FSlateChildSize InSize);
 
-	EHorizontalAlignment GetHorizontalAlignment() const;
+	UMG_API EHorizontalAlignment GetHorizontalAlignment() const;
 
 	UFUNCTION(BlueprintCallable, Category="Layout|Vertical Box Slot")
-	void SetHorizontalAlignment(EHorizontalAlignment InHorizontalAlignment);
+	UMG_API void SetHorizontalAlignment(EHorizontalAlignment InHorizontalAlignment);
 
-	EVerticalAlignment GetVerticalAlignment() const;
+	UMG_API EVerticalAlignment GetVerticalAlignment() const;
 
 	UFUNCTION(BlueprintCallable, Category="Layout|Vertical Box Slot")
-	void SetVerticalAlignment(EVerticalAlignment InVerticalAlignment);
+	UMG_API void SetVerticalAlignment(EVerticalAlignment InVerticalAlignment);
 
 public:
 
 	// UPanelSlot interface
-	virtual void SynchronizeProperties() override;
+	UMG_API virtual void SynchronizeProperties() override;
 	// End of UPanelSlot interface
 
-	virtual void ReleaseSlateResources(bool bReleaseChildren) override;
+	UMG_API virtual void ReleaseSlateResources(bool bReleaseChildren) override;
 
 	/** Builds the underlying FSlot for the Slate layout panel. */
-	void BuildSlot(TSharedRef<SVerticalBox> InVerticalBox);
+	UMG_API void BuildSlot(TSharedRef<SVerticalBox> InVerticalBox);
 
 #if WITH_EDITOR
-	virtual bool NudgeByDesigner(const FVector2D& NudgeDirection, const TOptional<int32>& GridSnapSize) override;
-	virtual void SynchronizeFromTemplate(const UPanelSlot* const TemplateSlot) override;
+	UMG_API virtual bool NudgeByDesigner(const FVector2D& NudgeDirection, const TOptional<int32>& GridSnapSize) override;
+	UMG_API virtual void SynchronizeFromTemplate(const UPanelSlot* const TemplateSlot) override;
 #endif //WITH_EDITOR
 };
