[Engine]

diff --git a/Engine/Source/Runtime/UMG/Public/Components/Viewport.h b/Engine/Source/Runtime/UMG/Public/Components/Viewport.h
--- a/Engine/Source/Runtime/UMG/Public/Components/Viewport.h
+++ b/Engine/Source/Runtime/UMG/Public/Components/Viewport.h
@@ -24,13 +24,13 @@ class SAutoRefreshViewport;
 /**
 * Stores the transformation data for the viewport camera
 */
-struct UMG_API FUMGViewportCameraTransform
+struct FUMGViewportCameraTransform
 {
 public:
-	FUMGViewportCameraTransform();
+	UMG_API FUMGViewportCameraTransform();
 
 	/** Sets the transform's location */
-	void SetLocation(const FVector& Position);
+	UMG_API void SetLocation(const FVector& Position);
 
 	/** Sets the transform's rotation */
 	void SetRotation(const FRotator& Rotation)
@@ -69,19 +69,19 @@ public:
 	* @param InDesiredLocation	The location to transition to
 	* @param bInstant			If the desired location should be set instantly rather than transitioned to over time
 	*/
-	void TransitionToLocation(const FVector& InDesiredLocation, bool bInstant);
+	UMG_API void TransitionToLocation(const FVector& InDesiredLocation, bool bInstant);
 
 	/**
 	* Updates any current location transitions
 	*
 	* @return true if there is currently a transition
 	*/
-	bool UpdateTransition();
+	UMG_API bool UpdateTransition();
 
 	/**
 	* Computes a matrix to use for viewport location and rotation when orbiting
 	*/
-	FMatrix ComputeOrbitMatrix() const;
+	UMG_API FMatrix ComputeOrbitMatrix() const;
 private:
 	/** The time when a transition to the desired location began */
 	double TransitionStartTime;
@@ -100,35 +100,35 @@ private:
 };
 
 
-class UMG_API FUMGViewportClient : public FCommonViewportClient, public FViewElementDrawer
+class FUMGViewportClient : public FCommonViewportClient, public FViewElementDrawer
 {
 public:
-	FUMGViewportClient(FPreviewScene* InPreviewScene = nullptr);
-	virtual ~FUMGViewportClient();
+	UMG_API FUMGViewportClient(FPreviewScene* InPreviewScene = nullptr);
+	UMG_API virtual ~FUMGViewportClient();
 
 	using FViewElementDrawer::Draw;
 
 	// FViewportClient interface
-	virtual UWorld* GetWorld() const override;
-	virtual void Draw(FViewport* InViewport, FCanvas* Canvas) override;
+	UMG_API virtual UWorld* GetWorld() const override;
+	UMG_API virtual void Draw(FViewport* InViewport, FCanvas* Canvas) override;
 
 	// FUMGViewportClient
 
-	virtual void Tick(float InDeltaTime);
+	UMG_API virtual void Tick(float InDeltaTime);
 
-	virtual FSceneView* CalcSceneView(FSceneViewFamily* ViewFamily);
+	UMG_API virtual FSceneView* CalcSceneView(FSceneViewFamily* ViewFamily);
 
 
 	/**
 	 * @return The scene being rendered in this viewport
 	 */
-	virtual FSceneInterface* GetScene() const;
+	UMG_API virtual FSceneInterface* GetScene() const;
 
 
-	bool IsAspectRatioConstrained() const;
+	UMG_API bool IsAspectRatioConstrained() const;
 
-	void SetBackgroundColor(FLinearColor InBackgroundColor);
-	FLinearColor GetBackgroundColor() const;
+	UMG_API void SetBackgroundColor(FLinearColor InBackgroundColor);
+	UMG_API FLinearColor GetBackgroundColor() const;
 
 	/** Sets the location of the viewport's camera */
 	void SetViewLocation(const FVector& NewLocation)
@@ -195,7 +195,7 @@ public:
 	}
 
 	/** @return The number of units per pixel displayed in this viewport */
-	float GetOrthoUnitsPerPixel(const FViewport* Viewport) const;
+	UMG_API float GetOrthoUnitsPerPixel(const FViewport* Viewport) const;
 
 	void SetEngineShowFlags(FEngineShowFlags InEngineShowFlags)
 	{
@@ -226,8 +226,8 @@ protected:
 /**
  * 
  */
-UCLASS(Experimental)
-class UMG_API UViewport : public UContentWidget
+UCLASS(Experimental, MinimalAPI)
+class UViewport : public UContentWidget
 {
 	GENERATED_UCLASS_BODY()
 
@@ -236,42 +236,42 @@ class UMG_API UViewport : public UContentWidget
 	FLinearColor BackgroundColor;
 
 	UFUNCTION(BlueprintCallable, Category="Viewport")
-	UWorld* GetViewportWorld() const;
+	UMG_API UWorld* GetViewportWorld() const;
 
 	UFUNCTION(BlueprintCallable, Category="Camera")
-	FVector GetViewLocation() const;
+	UMG_API FVector GetViewLocation() const;
 
 	UFUNCTION(BlueprintCallable, Category="Camera")
-	void SetViewLocation(FVector Location);
+	UMG_API void SetViewLocation(FVector Location);
 
 	UFUNCTION(BlueprintCallable, Category="Camera")
-	FRotator GetViewRotation() const;
+	UMG_API FRotator GetViewRotation() const;
 
 	UFUNCTION(BlueprintCallable, Category="Camera")
-	void SetViewRotation(FRotator Rotation);
+	UMG_API void SetViewRotation(FRotator Rotation);
 
 	UFUNCTION(BlueprintCallable, Category="Viewport")
-	AActor* Spawn(TSubclassOf<AActor> ActorClass);
+	UMG_API AActor* Spawn(TSubclassOf<AActor> ActorClass);
 
-	void SetBackgroundColor(const FLinearColor& InColor);
+	UMG_API void SetBackgroundColor(const FLinearColor& InColor);
 
-	const FLinearColor& GetBackgroundColor() const;
+	UMG_API const FLinearColor& GetBackgroundColor() const;
 
 	// UWidget interface
-	virtual void SynchronizeProperties() override;
+	UMG_API virtual void SynchronizeProperties() override;
 	// End of UWidget interface
 
-	virtual void ReleaseSlateResources(bool bReleaseChildren) override;
+	UMG_API virtual void ReleaseSlateResources(bool bReleaseChildren) override;
 
 #if WITH_EDITOR
-	virtual const FText GetPaletteCategory() override;
+	UMG_API virtual const FText GetPaletteCategory() override;
 #endif
 
 protected:
 
 	// UPanelWidget
-	virtual void OnSlotAdded(UPanelSlot* Slot) override;
-	virtual void OnSlotRemoved(UPanelSlot* Slot) override;
+	UMG_API virtual void OnSlotAdded(UPanelSlot* Slot) override;
+	UMG_API virtual void OnSlotRemoved(UPanelSlot* Slot) override;
 	// End UPanelWidget
 
 protected:
@@ -279,7 +279,7 @@ protected:
 
 protected:
 	// UWidget interface
-	virtual TSharedRef<SWidget> RebuildWidget() override;
+	UMG_API virtual TSharedRef<SWidget> RebuildWidget() override;
 	// End of UWidget interface
 
 	/** Show flags for the engine for this viewport */
diff --git a/Engine/Source/Runtime/UMG/Public/Components/Visual.h b/Engine/Source/Runtime/UMG/Public/Components/Visual.h
--- a/Engine/Source/Runtime/UMG/Public/Components/Visual.h
+++ b/Engine/Source/Runtime/UMG/Public/Components/Visual.h
@@ -8,17 +8,17 @@
 #include "Visual.generated.h"
 
 /** The base class for elements in UMG: slots and widgets. */
-UCLASS(DefaultToInstanced)
-class UMG_API UVisual : public UObject
+UCLASS(DefaultToInstanced, MinimalAPI)
+class UVisual : public UObject
 {
 	GENERATED_UCLASS_BODY()
 
 public:
-	virtual void ReleaseSlateResources(bool bReleaseChildren);
+	UMG_API virtual void ReleaseSlateResources(bool bReleaseChildren);
 
 	//~ Begin UObject Interface
-	virtual void BeginDestroy() override;
-	virtual bool NeedsLoadForServer() const override;
+	UMG_API virtual void BeginDestroy() override;
+	UMG_API virtual bool NeedsLoadForServer() const override;
 	//~ End UObject Interface
 
 private:
diff --git a/Engine/Source/Runtime/UMG/Public/Components/Widget.h b/Engine/Source/Runtime/UMG/Public/Components/Widget.h
--- a/Engine/Source/Runtime/UMG/Public/Components/Widget.h
+++ b/Engine/Source/Runtime/UMG/Public/Components/Widget.h
@@ -211,8 +211,8 @@ public:
 /**
  * This is the base class for all wrapped Slate controls that are exposed to UObjects.
  */
-UCLASS(Abstract, BlueprintType, Blueprintable, CustomFieldNotify)
-class UMG_API UWidget : public UVisual, public INotifyFieldValueChanged
+UCLASS(Abstract, BlueprintType, Blueprintable, CustomFieldNotify, MinimalAPI)
+class UWidget : public UVisual, public INotifyFieldValueChanged
 {
 	GENERATED_UCLASS_BODY()
 
@@ -483,131 +483,131 @@ public:
 
 public:
 	/** */
-	const FWidgetTransform& GetRenderTransform() const;
+	UMG_API const FWidgetTransform& GetRenderTransform() const;
 
 	/** */
 	UFUNCTION(BlueprintCallable, Category="Widget|Transform")
-	void SetRenderTransform(FWidgetTransform InTransform);
+	UMG_API void SetRenderTransform(FWidgetTransform InTransform);
 
 	/** */
 	UFUNCTION(BlueprintCallable, Category="Widget|Transform")
-	void SetRenderScale(FVector2D Scale);
+	UMG_API void SetRenderScale(FVector2D Scale);
 
 	/** */
 	UFUNCTION(BlueprintCallable, Category="Widget|Transform")
-	void SetRenderShear(FVector2D Shear);
+	UMG_API void SetRenderShear(FVector2D Shear);
 
 	/** */
 	UFUNCTION(BlueprintCallable, Category="Widget|Transform")
-	void SetRenderTransformAngle(float Angle);
+	UMG_API void SetRenderTransformAngle(float Angle);
 	
 	/** */
 	UFUNCTION(BlueprintCallable, Category = "Widget|Transform")
-	float GetRenderTransformAngle() const;
+	UMG_API float GetRenderTransformAngle() const;
 	
 	/** */
 	UFUNCTION(BlueprintCallable, Category="Widget|Transform")
-	void SetRenderTranslation(FVector2D Translation);
+	UMG_API void SetRenderTranslation(FVector2D Translation);
 
 	/** */
-	FVector2D GetRenderTransformPivot() const;
+	UMG_API FVector2D GetRenderTransformPivot() const;
 
 	/** */
 	UFUNCTION(BlueprintCallable, Category="Widget|Transform")
-	void SetRenderTransformPivot(FVector2D Pivot);
+	UMG_API void SetRenderTransformPivot(FVector2D Pivot);
 
 	/** Gets the flow direction preference of the widget */
-	EFlowDirectionPreference GetFlowDirectionPreference() const;
+	UMG_API EFlowDirectionPreference GetFlowDirectionPreference() const;
 
 	/** Sets the flow direction preference of the widget */
-	void SetFlowDirectionPreference(EFlowDirectionPreference FlowDirection);
+	UMG_API void SetFlowDirectionPreference(EFlowDirectionPreference FlowDirection);
 
 	/** Gets the current enabled status of the widget */
 	UFUNCTION(BlueprintCallable, Category="Widget")
-	bool GetIsEnabled() const;
+	UMG_API bool GetIsEnabled() const;
 
 	/** Sets the current enabled status of the widget */
 	UFUNCTION(BlueprintCallable, Category="Widget")
-	virtual void SetIsEnabled(bool bInIsEnabled);
+	UMG_API virtual void SetIsEnabled(bool bInIsEnabled);
 
 	/* @return true if the widget was added to the viewport using AddToViewport or AddToPlayerScreen. */
 	UFUNCTION(BlueprintPure, BlueprintCosmetic, Category = "Appearance")
-	bool IsInViewport() const;
+	UMG_API bool IsInViewport() const;
 
 	/** @return the tooltip text for the widget. */
-	FText GetToolTipText() const;
+	UMG_API FText GetToolTipText() const;
 
 	/** Sets the tooltip text for the widget. */
 	UFUNCTION(BlueprintCallable, Category="Widget")
-	void SetToolTipText(const FText& InToolTipText);
+	UMG_API void SetToolTipText(const FText& InToolTipText);
 
 	/** @return the custom widget as the tooltip of the widget. */
-	UWidget* GetToolTip() const;
+	UMG_API UWidget* GetToolTip() const;
 
 	/** Sets a custom widget as the tooltip of the widget. */
 	UFUNCTION(BlueprintCallable, Category="Widget")
-	void SetToolTip(UWidget* Widget);
+	UMG_API void SetToolTip(UWidget* Widget);
 
 	/** Sets the cursor to show over the widget. */
-	EMouseCursor::Type GetCursor() const;
+	UMG_API EMouseCursor::Type GetCursor() const;
 
 	/** Sets the cursor to show over the widget. */
 	UFUNCTION(BlueprintCallable, Category="Widget")
-	void SetCursor(EMouseCursor::Type InCursor);
+	UMG_API void SetCursor(EMouseCursor::Type InCursor);
 
 	/** Resets the cursor to use on the widget, removing any customization for it. */
 	UFUNCTION(BlueprintCallable, Category="Widget")
-	void ResetCursor();
+	UMG_API void ResetCursor();
 	
 	/** Returns true if the widget is Visible, HitTestInvisible or SelfHitTestInvisible and the Render Opacity is greater than 0. */
 	UFUNCTION(BlueprintCallable, Category="Widget")
-	bool IsRendered() const;
+	UMG_API bool IsRendered() const;
 
 	/** Returns true if the widget is Visible, HitTestInvisible or SelfHitTestInvisible. */
 	UFUNCTION(BlueprintCallable, Category="Widget")
-	bool IsVisible() const;
+	UMG_API bool IsVisible() const;
 
 	/** Gets the current visibility of the widget. */
 	UFUNCTION(BlueprintCallable, Category="Widget")
-	ESlateVisibility GetVisibility() const;
+	UMG_API ESlateVisibility GetVisibility() const;
 
 	/** Sets the visibility of the widget. */
 	UFUNCTION(BlueprintCallable, Category="Widget")
-	virtual void SetVisibility(ESlateVisibility InVisibility);
+	UMG_API virtual void SetVisibility(ESlateVisibility InVisibility);
 
 protected:
-	void SetVisibilityInternal(ESlateVisibility InVisibility);
+	UMG_API void SetVisibilityInternal(ESlateVisibility InVisibility);
 
 public:
 	/** Gets the current visibility of the widget. */
 	UFUNCTION(BlueprintCallable, Category="Widget")
-	float GetRenderOpacity() const;
+	UMG_API float GetRenderOpacity() const;
 
 	/** Sets the visibility of the widget. */
 	UFUNCTION(BlueprintCallable, Category="Widget")
-	void SetRenderOpacity(float InOpacity);
+	UMG_API void SetRenderOpacity(float InOpacity);
 
 	/** Gets the clipping state of this widget. */
 	UFUNCTION(BlueprintCallable, Category = "Widget")
-	EWidgetClipping GetClipping() const;
+	UMG_API EWidgetClipping GetClipping() const;
 
 	/** Sets the clipping state of this widget. */
 	UFUNCTION(BlueprintCallable, Category = "Widget")
-	void SetClipping(EWidgetClipping InClipping);
+	UMG_API void SetClipping(EWidgetClipping InClipping);
 	
 	/** Gets the pixel snapping method of this widget. */
-	EWidgetPixelSnapping GetPixelSnapping() const;
+	UMG_API EWidgetPixelSnapping GetPixelSnapping() const;
 
 	/** Sets the pixel snapping method of this widget. */
-	void SetPixelSnapping(EWidgetPixelSnapping InPixelSnapping);
+	UMG_API void SetPixelSnapping(EWidgetPixelSnapping InPixelSnapping);
 
 	/** Sets the forced volatility of the widget. */
 	UFUNCTION(BlueprintCallable, Category="Widget")
-	void ForceVolatile(bool bForce);
+	UMG_API void ForceVolatile(bool bForce);
 
 	/** Returns true if the widget is currently being hovered by a pointer device */
 	UFUNCTION(BlueprintCallable, Category="Widget")
-	virtual bool IsHovered() const;
+	UMG_API virtual bool IsHovered() const;
 
 	/**
 	 * Checks to see if this widget currently has the keyboard focus
@@ -615,14 +615,14 @@ public:
 	 * @return  True if this widget has keyboard focus
 	 */
 	UFUNCTION(BlueprintCallable, Category="Widget")
-	bool HasKeyboardFocus() const;
+	UMG_API bool HasKeyboardFocus() const;
 
 	/**
 	 * Checks to see if this widget is the current mouse captor
 	 * @return  True if this widget has captured the mouse
 	 */
 	UFUNCTION(BlueprintCallable, Category="Widget")
-	bool HasMouseCapture() const;
+	UMG_API bool HasMouseCapture() const;
 
 	/**
 	 * Checks to see if this widget is the current mouse captor
@@ -631,35 +631,35 @@ public:
 	 *	@return  True if this widget has captured the mouse with given user and pointer
 	 */
 	UFUNCTION(BlueprintCallable, Category = "Widget")
-	bool HasMouseCaptureByUser(int32 UserIndex, int32 PointerIndex = -1) const;
+	UMG_API bool HasMouseCaptureByUser(int32 UserIndex, int32 PointerIndex = -1) const;
 
 	/** Sets the focus to this widget. */
 	UFUNCTION(BlueprintCallable, Category="Widget")
-	void SetKeyboardFocus();
+	UMG_API void SetKeyboardFocus();
 
 	/** Returns true if this widget is focused by a specific user. */
 	UFUNCTION(BlueprintCallable, Category="Widget")
-	bool HasUserFocus(APlayerController* PlayerController) const;
+	UMG_API bool HasUserFocus(APlayerController* PlayerController) const;
 
 	/** Returns true if this widget is focused by any user. */
 	UFUNCTION(BlueprintCallable, Category="Widget")
-	bool HasAnyUserFocus() const;
+	UMG_API bool HasAnyUserFocus() const;
 
 	/** Returns true if any descendant widget is focused by any user. */
 	UFUNCTION(BlueprintCallable, Category="Widget", meta=(DisplayName="HasAnyUserFocusedDescendants"))
-	bool HasFocusedDescendants() const;
+	UMG_API bool HasFocusedDescendants() const;
 
 	/** Returns true if any descendant widget is focused by a specific user. */
 	UFUNCTION(BlueprintCallable, Category="Widget")
-	bool HasUserFocusedDescendants(APlayerController* PlayerController) const;
+	UMG_API bool HasUserFocusedDescendants(APlayerController* PlayerController) const;
 	
 	/** Sets the focus to this widget for the owning user */
 	UFUNCTION(BlueprintCallable, Category="Widget")
-	void SetFocus();
+	UMG_API void SetFocus();
 
 	/** Sets the focus to this widget for a specific user (if setting focus for the owning user, prefer SetFocus()) */
 	UFUNCTION(BlueprintCallable, Category="Widget")
-	void SetUserFocus(APlayerController* PlayerController);
+	UMG_API void SetUserFocus(APlayerController* PlayerController);
 
 	/**
 	 * Forces a pre-pass.  A pre-pass caches the desired size of the widget hierarchy owned by this widget.  
@@ -667,14 +667,14 @@ public:
 	 * pre-pass if you are adding child widgets this frame and want them to immediately be visible this frame.
 	 */
 	UFUNCTION(BlueprintCallable, Category="Widget")
-	void ForceLayoutPrepass();
+	UMG_API void ForceLayoutPrepass();
 
 	/**
 	 * Invalidates the widget from the view of a layout caching widget that may own this widget.
 	 * will force the owning widget to redraw and cache children on the next paint pass.
 	 */
 	UFUNCTION(BlueprintCallable, Category="Widget")
-	void InvalidateLayoutAndVolatility();
+	UMG_API void InvalidateLayoutAndVolatility();
 
 	/**
 	 * Gets the widgets desired size.
@@ -684,7 +684,7 @@ public:
 	 * @return The widget's desired size
 	 */
 	UFUNCTION(BlueprintCallable, Category="Widget")
-	FVector2D GetDesiredSize() const;
+	UMG_API FVector2D GetDesiredSize() const;
 
 	/**
 	 *	Sets the widget navigation rules for all directions. This can only be called on widgets that are in a widget tree.
@@ -692,7 +692,7 @@ public:
 	 *	@param WidgetToFocus When using the Explicit rule, focus on this widget
 	 */
 	UFUNCTION(BlueprintCallable, Category = "Widget")
-	void SetAllNavigationRules(EUINavigationRule Rule, FName WidgetToFocus);
+	UMG_API void SetAllNavigationRules(EUINavigationRule Rule, FName WidgetToFocus);
 
 	/**
 	 *	Sets the widget navigation rules for a specific direction. This can only be called on widgets that are in a widget tree.
@@ -702,7 +702,7 @@ public:
 	 */
 	UE_DEPRECATED(4.23, "SetNavigationRule is deprecated. Please use either SetNavigationRuleBase or SetNavigationRuleExplicit or SetNavigationRuleCustom or SetNavigationRuleCustomBoundary.")
 	UFUNCTION(BlueprintCallable, Category = "Widget", meta = (DeprecatedFunction, DeprecatedMessage = "Please use either SetNavigationRuleBase or SetNavigationRuleExplicit or SetNavigationRuleCustom or SetNavigationRuleCustomBoundary."))
-	void SetNavigationRule(EUINavigation Direction, EUINavigationRule Rule, FName WidgetToFocus);
+	UMG_API void SetNavigationRule(EUINavigation Direction, EUINavigationRule Rule, FName WidgetToFocus);
 
 	/**
 	 *	Sets the widget navigation rules for a specific direction. This can only be called on widgets that are in a widget tree. This works only for non Explicit, non Custom and non CustomBoundary Rules.
@@ -710,7 +710,7 @@ public:
 	 *	@param Rule The rule to use when navigation is taking place
 	 */
 	UFUNCTION(BlueprintCallable, Category = "Widget")
-	void SetNavigationRuleBase(EUINavigation Direction, EUINavigationRule Rule);
+	UMG_API void SetNavigationRuleBase(EUINavigation Direction, EUINavigationRule Rule);
 
 	/**
 	 *	Sets the widget navigation rules for a specific direction. This can only be called on widgets that are in a widget tree. This works only for Explicit Rule.
@@ -718,7 +718,7 @@ public:
 	 *	@param InWidget Focus on this widget instance
 	 */
 	UFUNCTION(BlueprintCallable, Category = "Widget")
-	void SetNavigationRuleExplicit(EUINavigation Direction, UWidget* InWidget);
+	UMG_API void SetNavigationRuleExplicit(EUINavigation Direction, UWidget* InWidget);
 
 	/**
 	 *	Sets the widget navigation rules for a specific direction. This can only be called on widgets that are in a widget tree. This works only for Custom Rule.
@@ -726,7 +726,7 @@ public:
 	 *	@param InCustomDelegate Custom Delegate that will be called
 	 */
 	UFUNCTION(BlueprintCallable, Category = "Widget")
-	void SetNavigationRuleCustom(EUINavigation Direction, FCustomWidgetNavigationDelegate InCustomDelegate);
+	UMG_API void SetNavigationRuleCustom(EUINavigation Direction, FCustomWidgetNavigationDelegate InCustomDelegate);
 
 	/**
 	 *	Sets the widget navigation rules for a specific direction. This can only be called on widgets that are in a widget tree. This works only for CustomBoundary Rule.
@@ -734,18 +734,18 @@ public:
 	 *	@param InCustomDelegate Custom Delegate that will be called
 	 */
 	UFUNCTION(BlueprintCallable, Category = "Widget")
-	void SetNavigationRuleCustomBoundary(EUINavigation Direction, FCustomWidgetNavigationDelegate InCustomDelegate);
+	UMG_API void SetNavigationRuleCustomBoundary(EUINavigation Direction, FCustomWidgetNavigationDelegate InCustomDelegate);
 
 	/** Gets the parent widget */
 	UFUNCTION(BlueprintCallable, Category="Widget")
-	class UPanelWidget* GetParent() const;
+	UMG_API class UPanelWidget* GetParent() const;
 
 	/**
 	 * Removes the widget from its parent widget.  If this widget was added to the player's screen or the viewport
 	 * it will also be removed from those containers.
 	 */
 	UFUNCTION(BlueprintCallable, Category="Widget")
-	virtual void RemoveFromParent();
+	UMG_API virtual void RemoveFromParent();
 
 	/**
 	 * Gets the last geometry used to Tick the widget.  This data may not exist yet if this call happens prior to 
@@ -757,32 +757,32 @@ public:
 	 * being used to advise how to layout a dependent object the current frame.
 	 */
 	UFUNCTION(BlueprintCallable, Category="Widget")
-	const FGeometry& GetCachedGeometry() const;
+	UMG_API const FGeometry& GetCachedGeometry() const;
 
 	UFUNCTION(BlueprintCallable, Category="Widget")
-	const FGeometry& GetTickSpaceGeometry() const;
+	UMG_API const FGeometry& GetTickSpaceGeometry() const;
 
 	UFUNCTION(BlueprintCallable, Category="Widget")
-	const FGeometry& GetPaintSpaceGeometry() const;
+	UMG_API const FGeometry& GetPaintSpaceGeometry() const;
 
 	//~ Begin INotifyFieldValueChanged Interface
 public:
-	virtual FDelegateHandle AddFieldValueChangedDelegate(UE::FieldNotification::FFieldId InFieldId, FFieldValueChangedDelegate InNewDelegate) override final;
-	virtual bool RemoveFieldValueChangedDelegate(UE::FieldNotification::FFieldId InFieldId, FDelegateHandle InHandle) override final;
-	virtual int32 RemoveAllFieldValueChangedDelegates(const void* InUserObject) override final;
-	virtual int32 RemoveAllFieldValueChangedDelegates(UE::FieldNotification::FFieldId InFieldId, const void* InUserObject) override final;
-	virtual void BroadcastFieldValueChanged(UE::FieldNotification::FFieldId InFieldId) override final;
+	UMG_API virtual FDelegateHandle AddFieldValueChangedDelegate(UE::FieldNotification::FFieldId InFieldId, FFieldValueChangedDelegate InNewDelegate) override final;
+	UMG_API virtual bool RemoveFieldValueChangedDelegate(UE::FieldNotification::FFieldId InFieldId, FDelegateHandle InHandle) override final;
+	UMG_API virtual int32 RemoveAllFieldValueChangedDelegates(const void* InUserObject) override final;
+	UMG_API virtual int32 RemoveAllFieldValueChangedDelegates(UE::FieldNotification::FFieldId InFieldId, const void* InUserObject) override final;
+	UMG_API virtual void BroadcastFieldValueChanged(UE::FieldNotification::FFieldId InFieldId) override final;
 	//~ End INotifyFieldValueChanged Interface
 
 	UFUNCTION(BlueprintCallable, Category = "FieldNotify", meta = (DisplayName = "Add Field Value Changed Delegate", ScriptName = "AddFieldValueChangedDelegate"))
-	void K2_AddFieldValueChangedDelegate(FFieldNotificationId FieldId, FFieldValueChangedDynamicDelegate Delegate);
+	UMG_API void K2_AddFieldValueChangedDelegate(FFieldNotificationId FieldId, FFieldValueChangedDynamicDelegate Delegate);
 
 	UFUNCTION(BlueprintCallable, Category = "FieldNotify", meta = (DisplayName = "Remove Field Value Changed Delegate", ScriptName="RemoveFieldValueChangedDelegate"))
-	void K2_RemoveFieldValueChangedDelegate(FFieldNotificationId FieldId, FFieldValueChangedDynamicDelegate Delegate);
+	UMG_API void K2_RemoveFieldValueChangedDelegate(FFieldNotificationId FieldId, FFieldValueChangedDynamicDelegate Delegate);
 
 protected:
 	UFUNCTION(BlueprintCallable, Category="FieldNotify", meta = (DisplayName="Broadcast Field Value Changed", ScriptName="BroadcastFieldValueChanged"))
-	void K2_BroadcastFieldValueChanged(FFieldNotificationId FieldId);
+	UMG_API void K2_BroadcastFieldValueChanged(FFieldNotificationId FieldId);
 
 public:
 	/**
@@ -792,7 +792,7 @@ public:
 	 * you must also take care to not call TakeWidget before calling TakeDerivedWidget, as that would put the wrong
 	 * expected wrapper around the resulting widget being constructed.
 	 */
-	TSharedRef<SWidget> TakeWidget();
+	UMG_API TSharedRef<SWidget> TakeWidget();
 
 	/**
 	 * Gets the underlying slate widget or constructs it if it doesn't exist.
@@ -821,22 +821,22 @@ public:
 	}
 	
 private:
-	TSharedRef<SWidget> TakeWidget_Private( ConstructMethodType ConstructMethod );
+	UMG_API TSharedRef<SWidget> TakeWidget_Private( ConstructMethodType ConstructMethod );
 
 public:
 
 	/** Gets the last created widget does not recreate the gc container for the widget if one is needed. */
-	TSharedPtr<SWidget> GetCachedWidget() const;
+	UMG_API TSharedPtr<SWidget> GetCachedWidget() const;
 
 	/** Gets the last created widget does not recreate the gc container for the widget if one is needed. */
-	bool IsConstructed() const;
+	UMG_API bool IsConstructed() const;
 
 	/**
 	 * Gets the game instance associated with this UI.
 	 * @return a pointer to the owning game instance
 	 */
 	UFUNCTION(BlueprintCallable, BlueprintCosmetic, Category = "Widget")
-	UGameInstance* GetGameInstance() const;
+	UMG_API UGameInstance* GetGameInstance() const;
 
 	/**
 	 * Gets the game instance associated with this UI.
@@ -853,7 +853,7 @@ public:
 	 * @return The player controller that owns the UI.
 	 */
 	UFUNCTION(BlueprintCallable, BlueprintCosmetic, Category = "Widget")
-	virtual APlayerController* GetOwningPlayer() const;
+	UMG_API virtual APlayerController* GetOwningPlayer() const;
 
 	/**
 	 * Gets the player controller associated with this UI cast to the template type.
@@ -870,7 +870,7 @@ public:
 	 * @return The owning local player.
 	 */
 	UFUNCTION(BlueprintCallable, BlueprintCosmetic, Category = "Widget")
-	virtual ULocalPlayer* GetOwningLocalPlayer() const;
+	UMG_API virtual ULocalPlayer* GetOwningLocalPlayer() const;
 	
 	/**
 	 * Gets the local player associated with this UI cast to the template type.
@@ -888,7 +888,7 @@ public:
 	  * accessibility is dsabled or the underlying accessible widget is invalid.
 	 */
 	UFUNCTION(BlueprintCallable, Category = "Widget")
-	FText GetAccessibleText() const;
+	UMG_API FText GetAccessibleText() const;
 
 	/**
 	 * Gets the accessible summary text from the underlying Slate accessible widget.
@@ -896,21 +896,21 @@ public:
 	  * accessibility is dsabled or the underlying accessible widget is invalid.
 	 */
 	UFUNCTION(BlueprintCallable, Category = "Widget")
-	FText GetAccessibleSummaryText() const;
+	UMG_API FText GetAccessibleSummaryText() const;
 	
 	/**
 	 * Applies all properties to the native widget if possible.  This is called after a widget is constructed.
 	 * It can also be called by the editor to update modified state, so ensure all initialization to a widgets
 	 * properties are performed here, or the property and visual state may become unsynced.
 	 */
-	virtual void SynchronizeProperties();
+	UMG_API virtual void SynchronizeProperties();
 
 	/**
 	 * Called by the owning user widget after the slate widget has been created.  After the entire widget tree
 	 * has been initialized, any widget reference that was needed to support navigating to another widget will
 	 * now be initialized and ready for usage.
 	 */
-	void BuildNavigation();
+	UMG_API void BuildNavigation();
 
 #if WITH_EDITOR
 	/** Returns if the widget is currently being displayed in the designer, it may want to display different data. */
@@ -920,7 +920,7 @@ public:
 	}
 
 	/** Sets the designer flags on the widget. */
-	virtual void SetDesignerFlags(EWidgetDesignFlags NewFlags);
+	UMG_API virtual void SetDesignerFlags(EWidgetDesignFlags NewFlags);
 
 	/** Gets the designer flags currently set on the widget. */
 	FORCEINLINE EWidgetDesignFlags GetDesignerFlags() const
@@ -946,13 +946,13 @@ public:
 	}
 
 	/** Sets the friendly name of the widget to display in the editor */
-	void SetDisplayLabel(const FString& DisplayLabel);
+	UMG_API void SetDisplayLabel(const FString& DisplayLabel);
 
 	/** Returns the category name of the widget */
-	const FString& GetCategoryName() const;
+	UMG_API const FString& GetCategoryName() const;
 
 	/** Sets the category name of the widget */
-	void SetCategoryName(const FString& InValue);
+	UMG_API void SetCategoryName(const FString& InValue);
 
 	/**
 	 * Called at the end of Widget Blueprint compilation.
@@ -962,7 +962,7 @@ public:
 	virtual void ValidateCompiledDefaults(class IWidgetCompilerLog& CompileLog) const {}
 
 	/** Mark this object as modified, also mark the slot as modified. */
-	virtual bool Modify(bool bAlwaysMarkDirty = true) override;
+	UMG_API virtual bool Modify(bool bAlwaysMarkDirty = true) override;
 #else
 	FORCEINLINE bool IsDesignTime() const { return false; }
 	FORCEINLINE bool IsPreviewTime() const { return false; }
@@ -972,10 +972,10 @@ public:
 	 * Recurses up the list of parents and returns true if this widget is a descendant of the PossibleParent
 	 * @return true if this widget is a child of the PossibleParent
 	 */
-	bool IsChildOf(UWidget* PossibleParent);
+	UMG_API bool IsChildOf(UWidget* PossibleParent);
 
 	/**  */
-	bool AddBinding(FDelegateProperty* DelegateProperty, UObject* SourceObject, const FDynamicPropertyPath& BindingPath);
+	UMG_API bool AddBinding(FDelegateProperty* DelegateProperty, UObject* SourceObject, const FDynamicPropertyPath& BindingPath);
 
 	/**
 	 * Add a post-state-changed listener to this widget, will fire after a state changed and all related side effects are resolved.
@@ -985,26 +985,26 @@ public:
 	 * @param ListenerDelegate Delegate to fire when state changes
 	 * @param bBroadcastCurrentState true if we should trigger this delegate once on registration with the current widget state (Does not globally broadcast).
 	 */
-	FDelegateHandle RegisterPostStateListener(const FOnWidgetStateBroadcast::FDelegate& ListenerDelegate, bool bBroadcastCurrentState = true);
+	UMG_API FDelegateHandle RegisterPostStateListener(const FOnWidgetStateBroadcast::FDelegate& ListenerDelegate, bool bBroadcastCurrentState = true);
 
 	/**
 	 * Remove a post-state-changed  listener from this widget, resets state bitfield if no other state listeners exist
 	 *
 	 * @param ListenerDelegate Delegate to remove
 	 */
-	void UnregisterPostStateListener(const FDelegateHandle& ListenerDelegate);
+	UMG_API void UnregisterPostStateListener(const FDelegateHandle& ListenerDelegate);
 
-	static TSubclassOf<UPropertyBinding> FindBinderClassForDestination(FProperty* Property);
+	static UMG_API TSubclassOf<UPropertyBinding> FindBinderClassForDestination(FProperty* Property);
 
 	//~ Begin UObject
-	virtual UWorld* GetWorld() const override;
-	virtual void BeginDestroy() override;
-	virtual void FinishDestroy() override;
-	virtual void PreSave(FObjectPreSaveContext ObjectSaveContext) override;
+	UMG_API virtual UWorld* GetWorld() const override;
+	UMG_API virtual void BeginDestroy() override;
+	UMG_API virtual void FinishDestroy() override;
+	UMG_API virtual void PreSave(FObjectPreSaveContext ObjectSaveContext) override;
 	//~ End UObject
 
 	//~ Begin UVisual
-	virtual void ReleaseSlateResources(bool bReleaseChildren) override;
+	UMG_API virtual void ReleaseSlateResources(bool bReleaseChildren) override;
 	//~ End UVisual
 
 	FORCEINLINE bool CanSafelyRouteEvent()
@@ -1025,22 +1025,22 @@ public:
 #if WITH_EDITOR
 
 	/** Is the label generated or provided by the user? */
-	bool IsGeneratedName() const;
+	UMG_API bool IsGeneratedName() const;
 
 	/** Get Label Metadata, which may be as simple as a bit of string data to help identify an anonymous text block. */
-	virtual FString GetLabelMetadata() const;
+	UMG_API virtual FString GetLabelMetadata() const;
 
 	/** Gets the label to display to the user for this widget. */
-	FText GetLabelText() const;
+	UMG_API FText GetLabelText() const;
 
 	/** Gets the label to display to the user for this widget, including any extra metadata like the text string for text. */
-	FText GetLabelTextWithMetadata() const;
+	UMG_API FText GetLabelTextWithMetadata() const;
 
 	/** Gets the palette category of the widget */
-	virtual const FText GetPaletteCategory();
+	UMG_API virtual const FText GetPaletteCategory();
 
 	/** Called by the palette after constructing a new widget. */
-	void CreatedFromPalette();
+	UMG_API void CreatedFromPalette();
 
 	/**
 	 * Called after constructing a new widget from the palette.
@@ -1052,15 +1052,15 @@ public:
 	virtual void ConnectEditorData() { }
 
 	// UObject interface
-	virtual void PostEditChangeProperty(struct FPropertyChangedEvent& PropertyChangedEvent) override;
+	UMG_API virtual void PostEditChangeProperty(struct FPropertyChangedEvent& PropertyChangedEvent) override;
 	// End of UObject interface
 
 	/** Is the widget visible in the designer?  If this widget is 'hidden in the designer' or a parent is, this widget will also return false here. */
-	bool IsVisibleInDesigner() const;
+	UMG_API bool IsVisibleInDesigner() const;
 
 	// Begin Designer contextual events
-	void SelectByDesigner();
-	void DeselectByDesigner();
+	UMG_API void SelectByDesigner();
+	UMG_API void DeselectByDesigner();
 
 	virtual void OnDesignerChanged(const FDesignerChangedEventArgs& EventArgs) { }
 
@@ -1077,25 +1077,25 @@ public:
 
 	// Utility methods
 	//@TODO UMG: Should move elsewhere
-	static EVisibility ConvertSerializedVisibilityToRuntime(ESlateVisibility Input);
-	static ESlateVisibility ConvertRuntimeToSerializedVisibility(const EVisibility& Input);
+	static UMG_API EVisibility ConvertSerializedVisibilityToRuntime(ESlateVisibility Input);
+	static UMG_API ESlateVisibility ConvertRuntimeToSerializedVisibility(const EVisibility& Input);
 
-	static FSizeParam ConvertSerializedSizeParamToRuntime(const FSlateChildSize& Input);
+	static UMG_API FSizeParam ConvertSerializedSizeParamToRuntime(const FSlateChildSize& Input);
 
-	static UWidget* FindChildContainingDescendant(UWidget* Root, UWidget* Descendant);
+	static UMG_API UWidget* FindChildContainingDescendant(UWidget* Root, UWidget* Descendant);
 
-	static FString GetDefaultFontName();
+	static UMG_API FString GetDefaultFontName();
 
 protected:
 #if WITH_EDITOR
 	// This is an implementation detail that allows us to show and hide the widget in the designer
 	// regardless of the actual visibility state set by the user.
-	EVisibility GetVisibilityInDesigner() const;
+	UMG_API EVisibility GetVisibilityInDesigner() const;
 
-	bool IsEditorWidget() const;
+	UMG_API bool IsEditorWidget() const;
 #endif
 
-	virtual void OnBindingChanged(const FName& Property);
+	UMG_API virtual void OnBindingChanged(const FName& Property);
 
 	/**
 	 * Broadcast a binary state post change
@@ -1103,37 +1103,37 @@ protected:
 	 * @param StateChange bitfield marking states that should be changed
 	 * @param bInValue true if marked states should be enabled, false otherwise
 	 */
-	void BroadcastBinaryPostStateChange(const FWidgetStateBitfield& StateChange, bool bInValue);
+	UMG_API void BroadcastBinaryPostStateChange(const FWidgetStateBitfield& StateChange, bool bInValue);
 
 	/**
 	 * Broadcast an enum state post change
 	 *
 	 * @param StateChange bitfield marking states that should be changed
 	 */
-	void BroadcastEnumPostStateChange(const FWidgetStateBitfield& StateChange);
+	UMG_API void BroadcastEnumPostStateChange(const FWidgetStateBitfield& StateChange);
 
 protected:
-	UObject* GetSourceAssetOrClass() const;
+	UMG_API UObject* GetSourceAssetOrClass() const;
 
 	/** Function implemented by all subclasses of UWidget is called when the underlying SWidget needs to be constructed. */
-	virtual TSharedRef<SWidget> RebuildWidget();
+	UMG_API virtual TSharedRef<SWidget> RebuildWidget();
 
 	/** Function called after the underlying SWidget is constructed. */
-	virtual void OnWidgetRebuilt();
+	UMG_API virtual void OnWidgetRebuilt();
 	
 #if WITH_EDITOR
-	virtual TSharedRef<SWidget> RebuildDesignWidget(TSharedRef<SWidget> Content);
+	UMG_API virtual TSharedRef<SWidget> RebuildDesignWidget(TSharedRef<SWidget> Content);
 
-	TSharedRef<SWidget> CreateDesignerOutline(TSharedRef<SWidget> Content) const;
+	UMG_API TSharedRef<SWidget> CreateDesignerOutline(TSharedRef<SWidget> Content) const;
 #endif
 
-	void UpdateRenderTransform();
+	UMG_API void UpdateRenderTransform();
 
 	/** Gets the base name used to generate the display label/name of this widget. */
-	FText GetDisplayNameBase() const;
+	UMG_API FText GetDisplayNameBase() const;
 
 	/** Copy all accessible properties to the AccessibleWidgetData object */
-	void SynchronizeAccessibleData();
+	UMG_API void SynchronizeAccessibleData();
 
 protected:
 	//TODO UMG Consider moving conversion functions into another class.
@@ -1154,11 +1154,11 @@ protected:
 		return FSlateColor(InLinearColor.Get());
 	}
 
-	void SetNavigationRuleInternal(EUINavigation Direction, EUINavigationRule Rule, FName WidgetToFocus = NAME_None, UWidget* InWidget = nullptr, FCustomWidgetNavigationDelegate InCustomDelegate = FCustomWidgetNavigationDelegate());
+	UMG_API void SetNavigationRuleInternal(EUINavigation Direction, EUINavigationRule Rule, FName WidgetToFocus = NAME_None, UWidget* InWidget = nullptr, FCustomWidgetNavigationDelegate InCustomDelegate = FCustomWidgetNavigationDelegate());
 
 #if WITH_ACCESSIBILITY
 	/** Gets the widget that accessibility properties should synchronize to. */
-	virtual TSharedPtr<SWidget> GetAccessibleWidget() const;
+	UMG_API virtual TSharedPtr<SWidget> GetAccessibleWidget() const;
 #endif
 
 protected:
@@ -1181,7 +1181,7 @@ protected:
 	UPROPERTY(Transient)
 	TArray<TObjectPtr<UPropertyBinding>> NativeBindings;
 
-	static TArray<TSubclassOf<UPropertyBinding>> BinderClasses;
+	static UMG_API TArray<TSubclassOf<UPropertyBinding>> BinderClasses;
 
 private:
 	TBitArray<> EnabledFieldNotifications;
@@ -1205,7 +1205,7 @@ private:
 
 #if !(UE_BUILD_SHIPPING || UE_BUILD_TEST)
 
-	void VerifySynchronizeProperties();
+	UMG_API void VerifySynchronizeProperties();
 
 	/** Did we route the synchronize properties call? */
 	bool bRoutedSynchronizeProperties;
diff --git a/Engine/Source/Runtime/UMG/Public/Components/WidgetComponent.h b/Engine/Source/Runtime/UMG/Public/Components/WidgetComponent.h
--- a/Engine/Source/Runtime/UMG/Public/Components/WidgetComponent.h
+++ b/Engine/Source/Runtime/UMG/Public/Components/WidgetComponent.h
@@ -90,67 +90,67 @@ enum class ETickMode : uint8
  * TintColorAndOpacity [Vector]
  * OpacityFromTexture [Scalar]
  */
-UCLASS(Blueprintable, ClassGroup="UserInterface", hidecategories=(Object,Activation,"Components|Activation",Sockets,Base,Lighting,LOD,Mesh), editinlinenew, meta=(BlueprintSpawnableComponent) )
-class UMG_API UWidgetComponent : public UMeshComponent
+UCLASS(Blueprintable, ClassGroup="UserInterface", hidecategories=(Object,Activation,"Components|Activation",Sockets,Base,Lighting,LOD,Mesh), editinlinenew, meta=(BlueprintSpawnableComponent) , MinimalAPI)
+class UWidgetComponent : public UMeshComponent
 {
 	GENERATED_UCLASS_BODY()
 
 public:
 	//UObject interface
-	virtual void Serialize(FArchive& Ar) override;
-	virtual bool CanBeInCluster() const override;
+	UMG_API virtual void Serialize(FArchive& Ar) override;
+	UMG_API virtual bool CanBeInCluster() const override;
 	//~ End UObject Interface
 
 	/** UActorComponent Interface */
-	virtual void BeginPlay() override;
-	virtual void EndPlay(const EEndPlayReason::Type EndPlayReason) override;
+	UMG_API virtual void BeginPlay() override;
+	UMG_API virtual void EndPlay(const EEndPlayReason::Type EndPlayReason) override;
 
 	/* UPrimitiveComponent Interface */
-	virtual FPrimitiveSceneProxy* CreateSceneProxy() override;
-	virtual FBoxSphereBounds CalcBounds(const FTransform & LocalToWorld) const override;
-	virtual UBodySetup* GetBodySetup() override;
-	virtual FCollisionShape GetCollisionShape(float Inflation) const override;
-	virtual void OnRegister() override;
-	virtual void OnUnregister() override;
-	virtual void DestroyComponent(bool bPromoteChildren = false) override;
-	UMaterialInterface* GetMaterial(int32 MaterialIndex) const override;
-	virtual void SetMaterial(int32 ElementIndex, UMaterialInterface* Material) override;
-	int32 GetNumMaterials() const override;
-	virtual void TickComponent(float DeltaTime, enum ELevelTick TickType, FActorComponentTickFunction *ThisTickFunction) override;
-
-	virtual TStructOnScope<FActorComponentInstanceData> GetComponentInstanceData() const override;
-	void ApplyComponentInstanceData(struct FWidgetComponentInstanceData* ComponentInstanceData);
-	virtual void GetUsedMaterials(TArray<UMaterialInterface*>& OutMaterials, bool bGetDebugMaterials = false) const override;
+	UMG_API virtual FPrimitiveSceneProxy* CreateSceneProxy() override;
+	UMG_API virtual FBoxSphereBounds CalcBounds(const FTransform & LocalToWorld) const override;
+	UMG_API virtual UBodySetup* GetBodySetup() override;
+	UMG_API virtual FCollisionShape GetCollisionShape(float Inflation) const override;
+	UMG_API virtual void OnRegister() override;
+	UMG_API virtual void OnUnregister() override;
+	UMG_API virtual void DestroyComponent(bool bPromoteChildren = false) override;
+	UMG_API UMaterialInterface* GetMaterial(int32 MaterialIndex) const override;
+	UMG_API virtual void SetMaterial(int32 ElementIndex, UMaterialInterface* Material) override;
+	UMG_API int32 GetNumMaterials() const override;
+	UMG_API virtual void TickComponent(float DeltaTime, enum ELevelTick TickType, FActorComponentTickFunction *ThisTickFunction) override;
+
+	UMG_API virtual TStructOnScope<FActorComponentInstanceData> GetComponentInstanceData() const override;
+	UMG_API void ApplyComponentInstanceData(struct FWidgetComponentInstanceData* ComponentInstanceData);
+	UMG_API virtual void GetUsedMaterials(TArray<UMaterialInterface*>& OutMaterials, bool bGetDebugMaterials = false) const override;
 
 #if WITH_EDITOR
-	virtual bool CanEditChange(const FProperty* InProperty) const override;
-	virtual void PostEditChangeProperty(FPropertyChangedEvent& PropertyChangedEvent) override;
+	UMG_API virtual bool CanEditChange(const FProperty* InProperty) const override;
+	UMG_API virtual void PostEditChangeProperty(FPropertyChangedEvent& PropertyChangedEvent) override;
 #endif
 
 	/** Ensures the user widget is initialized */
-	virtual void InitWidget();
+	UMG_API virtual void InitWidget();
 
 	/** Release resources associated with the widget. */
-	virtual void ReleaseResources();
+	UMG_API virtual void ReleaseResources();
 
 	/** Ensures the 3d window is created its size and content. */
-	virtual void UpdateWidget();
+	UMG_API virtual void UpdateWidget();
 
 	/** Ensure the render target is initialized and updates it if needed. */
-	virtual void UpdateRenderTarget(FIntPoint DesiredRenderTargetSize);
+	UMG_API virtual void UpdateRenderTarget(FIntPoint DesiredRenderTargetSize);
 
 	/** 
 	* Ensures the body setup is initialized and updates it if needed.
 	* @param bDrawSizeChanged Whether the draw size of this component has changed since the last update call.
 	*/
-	void UpdateBodySetup( bool bDrawSizeChanged = false );
+	UMG_API void UpdateBodySetup( bool bDrawSizeChanged = false );
 
 	/**
 	 * Converts a world-space hit result to a hit location on the widget
 	 * @param HitResult The hit on this widget in the world
 	 * @param (Out) The transformed 2D hit location on the widget
 	 */
-	virtual void GetLocalHitLocation(FVector WorldHitLocation, FVector2D& OutLocalHitLocation) const;
+	UMG_API virtual void GetLocalHitLocation(FVector WorldHitLocation, FVector2D& OutLocalHitLocation) const;
 
 	/**
 	 * When using EWidgetGeometryMode::Cylinder, continues the trace from the front face
@@ -158,7 +158,7 @@ public:
 	 * 
 	 * @returns two hit locations FVector is in world space and a FVector2D is in widget-space.
 	 */
-	TTuple<FVector, FVector2D> GetCylinderHitLocation(FVector WorldHitLocation, FVector WorldHitDirection) const;
+	UMG_API TTuple<FVector, FVector2D> GetCylinderHitLocation(FVector WorldHitLocation, FVector WorldHitDirection) const;
 
 	/** Gets the last local location that was hit */
 	FVector2D GetLastLocalHitLocation() const
@@ -171,46 +171,46 @@ public:
 
 	/** Returns the user widget object displayed by this component */
 	UFUNCTION(BlueprintCallable, Category=UserInterface, meta=(UnsafeDuringActorConstruction=true))
-	UUserWidget* GetUserWidgetObject() const;
+	UMG_API UUserWidget* GetUserWidgetObject() const;
 
 	/** Returns the Slate widget that was assigned to this component, if any */
-	const TSharedPtr<SWidget>& GetSlateWidget() const;
+	UMG_API const TSharedPtr<SWidget>& GetSlateWidget() const;
 
 	/** Returns the list of widgets with their geometry and the cursor position transformed into this Widget component's space. */
-	TArray<FWidgetAndPointer> GetHitWidgetPath(FVector WorldHitLocation, bool bIgnoreEnabledStatus, float CursorRadius = 0.0f);
+	UMG_API TArray<FWidgetAndPointer> GetHitWidgetPath(FVector WorldHitLocation, bool bIgnoreEnabledStatus, float CursorRadius = 0.0f);
 
 	/** Returns the list of widgets with their geometry and the cursor position transformed into this Widget space. The widget space is expressed as a Vector2D. */
-	TArray<FWidgetAndPointer> GetHitWidgetPath(FVector2D WidgetSpaceHitCoordinate, bool bIgnoreEnabledStatus, float CursorRadius = 0.0f);
+	UMG_API TArray<FWidgetAndPointer> GetHitWidgetPath(FVector2D WidgetSpaceHitCoordinate, bool bIgnoreEnabledStatus, float CursorRadius = 0.0f);
 
 	/** Returns the render target to which the user widget is rendered */
 	UFUNCTION(BlueprintCallable, Category=UserInterface)
-	UTextureRenderTarget2D* GetRenderTarget() const;
+	UMG_API UTextureRenderTarget2D* GetRenderTarget() const;
 
 	/** Returns the dynamic material instance used to render the user widget */
 	UFUNCTION(BlueprintCallable, Category=UserInterface)
-	UMaterialInstanceDynamic* GetMaterialInstance() const;
+	UMG_API UMaterialInstanceDynamic* GetMaterialInstance() const;
 
 	/** Returns the window containing the user widget content */
-	TSharedPtr<SWindow> GetSlateWindow() const;
+	UMG_API TSharedPtr<SWindow> GetSlateWindow() const;
 
 	/**  
 	 *  Gets the widget that is used by this Widget Component. It will be null if a Slate Widget was set using SetSlateWidget function.
 	 */ 
 	UFUNCTION(BlueprintCallable, Category=UserInterface)
-	virtual UUserWidget* GetWidget() const;
+	UMG_API virtual UUserWidget* GetWidget() const;
 
 	/**  
 	 *  Sets the widget to use directly. This function will keep track of the widget till the next time it's called
 	 *	with either a newer widget or a nullptr
 	 */ 
 	UFUNCTION(BlueprintCallable, Category=UserInterface)
-	virtual void SetWidget(UUserWidget* Widget);
+	UMG_API virtual void SetWidget(UUserWidget* Widget);
 
 	/**  
 	 *  Sets a Slate widget to be rendered.  You can use this to draw native Slate widgets using a WidgetComponent, instead
 	 *  of drawing user widgets.
 	 */ 
-	virtual void SetSlateWidget( const TSharedPtr<SWidget>& InSlateWidget);
+	UMG_API virtual void SetSlateWidget( const TSharedPtr<SWidget>& InSlateWidget);
 
 	/**
 	 * Sets the local player that owns this widget component.  Setting the owning player controls
@@ -218,7 +218,7 @@ public:
 	 * forwards the owning player to the actual UserWidget that is spawned.
 	 */
 	UFUNCTION(BlueprintCallable, Category=UserInterface)
-	void SetOwnerPlayer(ULocalPlayer* LocalPlayer);
+	UMG_API void SetOwnerPlayer(ULocalPlayer* LocalPlayer);
 
 	/** @see bManuallyRedraw */
 	UFUNCTION(BlueprintCallable, Category = UserInterface)
@@ -229,37 +229,37 @@ public:
 
 	/** @see bManuallyRedraw */
 	UFUNCTION(BlueprintCallable, Category=UserInterface)
-	void SetManuallyRedraw(bool bUseManualRedraw);
+	UMG_API void SetManuallyRedraw(bool bUseManualRedraw);
 
 	/** Gets the local player that owns this widget component. */
 	UFUNCTION(BlueprintCallable, Category=UserInterface)
-	ULocalPlayer* GetOwnerPlayer() const;
+	UMG_API ULocalPlayer* GetOwnerPlayer() const;
 
 	/** Returns the "specified" draw size of the quad in the world */
 	UFUNCTION(BlueprintCallable, Category=UserInterface)
-	FVector2D GetDrawSize() const;
+	UMG_API FVector2D GetDrawSize() const;
 
 	/** Returns the "actual" draw size of the quad in the world */
 	UFUNCTION(BlueprintCallable, Category = UserInterface)
-	FVector2D GetCurrentDrawSize() const;
+	UMG_API FVector2D GetCurrentDrawSize() const;
 
 	/** Sets the draw size of the quad in the world */
 	UFUNCTION(BlueprintCallable, Category=UserInterface)
-	void SetDrawSize(FVector2D Size);
+	UMG_API void SetDrawSize(FVector2D Size);
 
 	/** Requests that the widget be redrawn.  */
 	UFUNCTION(BlueprintCallable, Category=UserInterface, meta = (DeprecatedFunction, DeprecationMessage = "Use RequestRenderUpdate instead"))
-	virtual void RequestRedraw();
+	UMG_API virtual void RequestRedraw();
 
 	/** Requests that the widget have it's render target updated, if TickMode is disabled, this will force a tick to happen to update the render target. */
 	UFUNCTION(BlueprintCallable, Category = UserInterface)
-	virtual void RequestRenderUpdate();
+	UMG_API virtual void RequestRenderUpdate();
 
 	/** Gets the blend mode for the widget. */
 	EWidgetBlendMode GetBlendMode() const { return BlendMode; }
 
 	/** Sets the blend mode to use for this widget */
-	void SetBlendMode( const EWidgetBlendMode NewBlendMode );
+	UMG_API void SetBlendMode( const EWidgetBlendMode NewBlendMode );
 
 	/** Gets whether the widget is two-sided or not */
 	UFUNCTION(BlueprintCallable, Category = Rendering)
@@ -270,7 +270,7 @@ public:
 
 	/** Sets whether the widget is two-sided or not */
 	UFUNCTION(BlueprintCallable, Category = Rendering)
-	void SetTwoSided( const bool bWantTwoSided );
+	UMG_API void SetTwoSided( const bool bWantTwoSided );
 
 	/** Gets whether the widget ticks when offscreen or not */
 	UFUNCTION(BlueprintCallable, Category = Animation)
@@ -288,14 +288,14 @@ public:
 
 	/** Sets the background color and opacityscale for this widget */
 	UFUNCTION(BlueprintCallable, Category=UserInterface)
-	void SetBackgroundColor( const FLinearColor NewBackgroundColor );
+	UMG_API void SetBackgroundColor( const FLinearColor NewBackgroundColor );
 
 	/** Sets the tint color and opacity scale for this widget */
 	UFUNCTION(BlueprintCallable, Category=UserInterface)
-	void SetTintColorAndOpacity( const FLinearColor NewTintColorAndOpacity );
+	UMG_API void SetTintColorAndOpacity( const FLinearColor NewTintColorAndOpacity );
 
 	/** Sets how much opacity from the UI widget's texture alpha is used when rendering to the viewport (0.0-1.0) */
-	void SetOpacityFromTexture( const float NewOpacityFromTexture );
+	UMG_API void SetOpacityFromTexture( const float NewOpacityFromTexture );
 
 	/** Returns the pivot point where the UI is rendered about the origin. */
 	UFUNCTION(BlueprintCallable, Category = UserInterface)
@@ -322,13 +322,13 @@ public:
 	void SetRedrawTime(float InRedrawTime) { RedrawTime = InRedrawTime; }
 
 	/** Get the fake window we create for widgets displayed in the world. */
-	TSharedPtr< SWindow > GetVirtualWindow() const;
+	UMG_API TSharedPtr< SWindow > GetVirtualWindow() const;
 	
 	/** Updates the dynamic parameters on the material instance, without re-creating it */
-	void UpdateMaterialInstanceParameters();
+	UMG_API void UpdateMaterialInstanceParameters();
 
 	/** Sets the widget class used to generate the widget for this component */
-	void SetWidgetClass(TSubclassOf<UUserWidget> InWidgetClass);
+	UMG_API void SetWidgetClass(TSubclassOf<UUserWidget> InWidgetClass);
 
 	UFUNCTION(BlueprintCallable, Category = UserInterface)
 	EWidgetSpace GetWidgetSpace() const { return Space; }
@@ -382,7 +382,7 @@ public:
 
 	/** @see bWindowFocusable */
 	UFUNCTION(BlueprintCallable, Category = UserInterface)
-	void SetWindowFocusable(bool bInWindowFocusable);
+	UMG_API void SetWindowFocusable(bool bInWindowFocusable);
 
 	/** Gets the visibility of the virtual window created to host the widget focusable. */
 	UFUNCTION(BlueprintCallable, Category = UserInterface)
@@ -393,44 +393,44 @@ public:
 
 	/** Sets the visibility of the virtual window created to host the widget focusable. */
 	UFUNCTION(BlueprintCallable, Category = UserInterface)
-	void SetWindowVisibility(EWindowVisibility InVisibility);
+	UMG_API void SetWindowVisibility(EWindowVisibility InVisibility);
 
 	/** Sets the Tick mode of the Widget Component.*/
 	UFUNCTION(BlueprintCallable, Category = UserInterface)
-	void SetTickMode(ETickMode InTickMode);
+	UMG_API void SetTickMode(ETickMode InTickMode);
 
 	/** Returns true if the the Slate window is visible and that the widget is also visible, false otherwise. */
 	UFUNCTION(BlueprintCallable, Category = UserInterface)
-	bool IsWidgetVisible() const;
+	UMG_API bool IsWidgetVisible() const;
 
 	/** Hook to allow this component modify the local position of the widget after it has been projected from world space to screen space. */
 	virtual FVector2D ModifyProjectedLocalPosition(const FGeometry& ViewportGeometry, const FVector2D& LocalPosition) { return LocalPosition; }
 
 protected:
-	void OnLevelRemovedFromWorld(ULevel* InLevel, UWorld* InWorld);
+	UMG_API void OnLevelRemovedFromWorld(ULevel* InLevel, UWorld* InWorld);
 
 	/** Just because the user attempts to receive hardware input does not mean it's possible. */
-	bool CanReceiveHardwareInput() const;
+	UMG_API bool CanReceiveHardwareInput() const;
 
-	void RegisterHitTesterWithViewport(TSharedPtr<SViewport> ViewportWidget);
-	void UnregisterHitTesterWithViewport(TSharedPtr<SViewport> ViewportWidget);
+	UMG_API void RegisterHitTesterWithViewport(TSharedPtr<SViewport> ViewportWidget);
+	UMG_API void UnregisterHitTesterWithViewport(TSharedPtr<SViewport> ViewportWidget);
 
-	void RegisterWindow();
-	void UnregisterWindow();
-	void RemoveWidgetFromScreen();
+	UMG_API void RegisterWindow();
+	UMG_API void UnregisterWindow();
+	UMG_API void RemoveWidgetFromScreen();
 
 	/** Allows subclasses to control if the widget should be drawn.  Called right before we draw the widget. */
-	virtual bool ShouldDrawWidget() const;
+	UMG_API virtual bool ShouldDrawWidget() const;
 
 	/** Draws the current widget to the render target if possible. */
-	virtual void DrawWidgetToRenderTarget(float DeltaTime);
+	UMG_API virtual void DrawWidgetToRenderTarget(float DeltaTime);
 
 	/** Returns the width of the widget component taking GeometryMode into account. */
-	float ComputeComponentWidth() const;
+	UMG_API float ComputeComponentWidth() const;
 
-	void UpdateMaterialInstance();
+	UMG_API void UpdateMaterialInstance();
 
-	virtual void OnHiddenInGameChanged() override;
+	UMG_API virtual void OnHiddenInGameChanged() override;
 
 protected:
 	/** The coordinate space in which to render the widget */
@@ -469,7 +469,7 @@ protected:
 	double LastWidgetRenderTime;
 
 	/** Returns current absolute time, respecting TimingPolicy. */
-	double GetCurrentTime() const;
+	UMG_API double GetCurrentTime() const;
 
 	/**
 	 * The actual draw size, this changes based on DrawSize - or the desired size of the widget if
@@ -624,13 +624,13 @@ protected:
 	FVector2D LastLocalHitLocation;
 
 	/** The hit tester to use for this component */
-	static TSharedPtr<class FWidget3DHitTester> WidgetHitTester;
+	static UMG_API TSharedPtr<class FWidget3DHitTester> WidgetHitTester;
 
 	/** Helper class for drawing widgets to a render target. */
 	class FWidgetRenderer* WidgetRenderer;
 
 private: 
-	bool ShouldReenableComponentTickWhenWidgetBecomesVisible() const;
+	UMG_API bool ShouldReenableComponentTickWhenWidgetBecomesVisible() const;
 
 	/** The User Widget object displayed and managed by this component */
 	UPROPERTY(Transient, DuplicateTransient)
@@ -642,15 +642,15 @@ private:
 	/** The slate widget currently being drawn. */
 	TWeakPtr<SWidget> CurrentSlateWidget;
 
-	static EVisibility ConvertWindowVisibilityToVisibility(EWindowVisibility visibility);
+	static UMG_API EVisibility ConvertWindowVisibilityToVisibility(EWindowVisibility visibility);
 
-	void OnWidgetVisibilityChanged(ESlateVisibility InVisibility);
+	UMG_API void OnWidgetVisibilityChanged(ESlateVisibility InVisibility);
 
 	/** When using Screen space, this will update the Widget on Screen. **/
-	void UpdateWidgetOnScreen();
+	UMG_API void UpdateWidgetOnScreen();
 
 	/** When using Screen space, this code will add the widget to the screen. **/
-	void AddWidgetToScreen(ULocalPlayer* TargetPlayer);
+	UMG_API void AddWidgetToScreen(ULocalPlayer* TargetPlayer);
 
 	/** Set to true after a draw of an empty component.*/
 	bool bRenderCleared;
diff --git a/Engine/Source/Runtime/UMG/Public/Components/WidgetInteractionComponent.h b/Engine/Source/Runtime/UMG/Public/Components/WidgetInteractionComponent.h
--- a/Engine/Source/Runtime/UMG/Public/Components/WidgetInteractionComponent.h
+++ b/Engine/Source/Runtime/UMG/Public/Components/WidgetInteractionComponent.h
@@ -55,8 +55,8 @@ DECLARE_DYNAMIC_MULTICAST_DELEGATE_TwoParams(FOnHoveredWidgetChanged, UWidgetCom
  * to show as if the mouse were moving on top of it.  You'll then tell the component to simulate key presses, 
  * like Left Mouse, down and up, to simulate a mouse click.
  */
-UCLASS(ClassGroup="UserInterface", meta=(BlueprintSpawnableComponent) )
-class UMG_API UWidgetInteractionComponent : public USceneComponent
+UCLASS(ClassGroup="UserInterface", meta=(BlueprintSpawnableComponent) , MinimalAPI)
+class UWidgetInteractionComponent : public USceneComponent
 {
 	GENERATED_BODY()
 
@@ -69,13 +69,13 @@ public:
 	FOnHoveredWidgetChanged OnHoveredWidgetChanged;
 
 public:
-	UWidgetInteractionComponent(const FObjectInitializer& ObjectInitializer = FObjectInitializer::Get());
+	UMG_API UWidgetInteractionComponent(const FObjectInitializer& ObjectInitializer = FObjectInitializer::Get());
 
 	// Begin ActorComponent interface
-	virtual void OnComponentCreated() override;
-	virtual void Activate(bool bReset = false) override;
-	virtual void Deactivate() override;
-	virtual void TickComponent(float DeltaTime, enum ELevelTick TickType, FActorComponentTickFunction* ThisTickFunction) override;
+	UMG_API virtual void OnComponentCreated() override;
+	UMG_API virtual void Activate(bool bReset = false) override;
+	UMG_API virtual void Deactivate() override;
+	UMG_API virtual void TickComponent(float DeltaTime, enum ELevelTick TickType, FActorComponentTickFunction* ThisTickFunction) override;
 	// End UActorComponent
 	
 	/**
@@ -84,7 +84,7 @@ public:
 	 * send other keys to signal widgets to take special actions if they're under the cursor.
 	 */
 	UFUNCTION(BlueprintCallable, Category="Interaction")
-	virtual void PressPointerKey(FKey Key);
+	UMG_API virtual void PressPointerKey(FKey Key);
 	
 	/**
 	 * Releases a key as if the mouse/pointer were the source of it.  Normally you would just use
@@ -92,7 +92,7 @@ public:
 	 * send other keys to signal widgets to take special actions if they're under the cursor.
 	 */
 	UFUNCTION(BlueprintCallable, Category="Interaction")
-	virtual void ReleasePointerKey(FKey Key);
+	UMG_API virtual void ReleasePointerKey(FKey Key);
 
 	/**
 	 * Press a key as if it had come from the keyboard.  Avoid using this for 'a-z|A-Z', things like
@@ -100,88 +100,88 @@ public:
 	 * send to the widget.  So for those cases you should use SendKeyChar.
 	 */
 	UFUNCTION(BlueprintCallable, Category="Interaction")
-	virtual bool PressKey(FKey Key, bool bRepeat = false);
+	UMG_API virtual bool PressKey(FKey Key, bool bRepeat = false);
 	
 	/**
 	 * Releases a key as if it had been released by the keyboard.
 	 */
 	UFUNCTION(BlueprintCallable, Category="Interaction")
-	virtual bool ReleaseKey(FKey Key);
+	UMG_API virtual bool ReleaseKey(FKey Key);
 
 	/**
 	 * Does both the press and release of a simulated keyboard key.
 	 */
 	UFUNCTION(BlueprintCallable, Category="Interaction")
-	virtual bool PressAndReleaseKey(FKey Key);
+	UMG_API virtual bool PressAndReleaseKey(FKey Key);
 
 	/**
 	 * Transmits a list of characters to a widget by simulating a OnKeyChar event for each key listed in
 	 * the string.
 	 */
 	UFUNCTION(BlueprintCallable, Category="Interaction")
-	virtual bool SendKeyChar(FString Characters, bool bRepeat = false);
+	UMG_API virtual bool SendKeyChar(FString Characters, bool bRepeat = false);
 	
 	/**
 	 * Sends a scroll wheel event to the widget under the last hit result.
 	 */
 	UFUNCTION(BlueprintCallable, Category="Interaction")
-	virtual void ScrollWheel(float ScrollDelta);
+	UMG_API virtual void ScrollWheel(float ScrollDelta);
 	
 	/**
 	 * Get the currently hovered widget component.
 	 */
 	UFUNCTION(BlueprintCallable, Category="Interaction")
-	UWidgetComponent* GetHoveredWidgetComponent() const;
+	UMG_API UWidgetComponent* GetHoveredWidgetComponent() const;
 
 	/**
 	 * Returns true if a widget under the hit result is interactive.  e.g. Slate widgets 
 	 * that return true for IsInteractable().
 	 */
 	UFUNCTION(BlueprintCallable, Category="Interaction")
-	bool IsOverInteractableWidget() const;
+	UMG_API bool IsOverInteractableWidget() const;
 
 	/**
 	 * Returns true if a widget under the hit result is focusable.  e.g. Slate widgets that 
 	 * return true for SupportsKeyboardFocus().
 	 */
 	UFUNCTION(BlueprintCallable, Category="Interaction")
-	bool IsOverFocusableWidget() const;
+	UMG_API bool IsOverFocusableWidget() const;
 
 	/**
 	 * Returns true if a widget under the hit result is has a visibility that makes it hit test 
 	 * visible.  e.g. Slate widgets that return true for GetVisibility().IsHitTestVisible().
 	 */
 	UFUNCTION(BlueprintCallable, Category="Interaction")
-	bool IsOverHitTestVisibleWidget() const;
+	UMG_API bool IsOverHitTestVisibleWidget() const;
 
 	/**
 	 * Gets the widget path for the slate widgets under the last hit result.
 	 */
-	const FWeakWidgetPath& GetHoveredWidgetPath() const;
+	UMG_API const FWeakWidgetPath& GetHoveredWidgetPath() const;
 
 	/**
 	 * Gets the last hit result generated by the component.  Returns the custom hit result if that was set.
 	 */
 	UFUNCTION(BlueprintCallable, Category="Interaction")
-	const FHitResult& GetLastHitResult() const;
+	UMG_API const FHitResult& GetLastHitResult() const;
 
 	/**
 	 * Gets the last hit location on the widget in 2D, local pixel units of the render target.
 	 */
 	UFUNCTION(BlueprintCallable, Category="Interaction")
-	FVector2D Get2DHitLocation() const;
+	UMG_API FVector2D Get2DHitLocation() const;
 
 	/**
 	 * Set custom hit result.  This is only taken into account if InteractionSource is set to EWidgetInteractionSource::Custom.
 	 */
 	UFUNCTION(BlueprintCallable, Category="Interaction")
-	void SetCustomHitResult(const FHitResult& HitResult);
+	UMG_API void SetCustomHitResult(const FHitResult& HitResult);
 
 	/**
 	 * Set the focus target of the virtual user managed by this component
 	 */
 	UFUNCTION(BlueprintCallable, Category = "Interaction")
-	void SetFocus(UWidget* FocusWidget);
+	UMG_API void SetFocus(UWidget* FocusWidget);
 
 protected:
 	/**
@@ -270,13 +270,13 @@ public:
 protected:
 
 	// Gets the key and char codes for sending keys for the platform.
-	void GetKeyAndCharCodes(const FKey& Key, bool& bHasKeyCode, uint32& KeyCode, bool& bHasCharCode, uint32& CharCode);
+	UMG_API void GetKeyAndCharCodes(const FKey& Key, bool& bHasKeyCode, uint32& KeyCode, bool& bHasCharCode, uint32& CharCode);
 
 	/** Is it safe for this interaction component to run?  Might not be in a server situation with no slate application. */
-	bool CanSendInput();
+	UMG_API bool CanSendInput();
 
 	/** Performs the simulation of pointer movement.  Does not run if bEnableHitTesting is set to false. */
-	void SimulatePointerMovement();
+	UMG_API void SimulatePointerMovement();
 
 	struct FWidgetTraceResult
 	{
@@ -301,19 +301,19 @@ protected:
 	};
 
 	/** Gets the WidgetPath for the widget being hovered over based on the hit result. */
-	virtual FWidgetPath FindHoveredWidgetPath(const FWidgetTraceResult& TraceResult)  const;
+	UMG_API virtual FWidgetPath FindHoveredWidgetPath(const FWidgetTraceResult& TraceResult)  const;
 
 	/** Performs the trace and gets the hit result under the specified InteractionSource */
-	virtual FWidgetTraceResult PerformTrace() const;
+	UMG_API virtual FWidgetTraceResult PerformTrace() const;
 
 	/**
 	 * Gets the list of components to ignore during hit testing.  Which is everything that is a parent/sibling of this 
 	 * component that's not a Widget Component.  This is so traces don't get blocked by capsules and such around the player.
 	 */
-	void GetRelatedComponentsToIgnoreInAutomaticHitTesting(TArray<UPrimitiveComponent*>& IgnorePrimitives) const;
+	UMG_API void GetRelatedComponentsToIgnoreInAutomaticHitTesting(TArray<UPrimitiveComponent*>& IgnorePrimitives) const;
 
 	/** Returns true if the inteaction component can interact with the supplied widget component */
-	bool CanInteractWithComponent(UWidgetComponent* Component) const;
+	UMG_API bool CanInteractWithComponent(UWidgetComponent* Component) const;
 
 protected:
 
diff --git a/Engine/Source/Runtime/UMG/Public/Components/WidgetSwitcher.h b/Engine/Source/Runtime/UMG/Public/Components/WidgetSwitcher.h
--- a/Engine/Source/Runtime/UMG/Public/Components/WidgetSwitcher.h
+++ b/Engine/Source/Runtime/UMG/Public/Components/WidgetSwitcher.h
@@ -12,8 +12,8 @@
 /**
  * A widget switcher is like a tab control, but without tabs. At most one widget is visible at time.
  */
-UCLASS()
-class UMG_API UWidgetSwitcher : public UPanelWidget
+UCLASS(MinimalAPI)
+class UWidgetSwitcher : public UPanelWidget
 {
 	GENERATED_UCLASS_BODY()
 
@@ -27,58 +27,58 @@ public:
 
 	/** Gets the number of widgets that this switcher manages. */
 	UFUNCTION(BlueprintCallable, Category="Switcher")
-	int32 GetNumWidgets() const;
+	UMG_API int32 GetNumWidgets() const;
 
 	/** Gets the slot index of the currently active widget */
 	UFUNCTION(BlueprintCallable, Category="Switcher")
-	int32 GetActiveWidgetIndex() const;
+	UMG_API int32 GetActiveWidgetIndex() const;
 
 	/** Activates the widget at the specified index. */
 	UFUNCTION(BlueprintCallable, Category="Switcher")
-	virtual void SetActiveWidgetIndex( int32 Index );
+	UMG_API virtual void SetActiveWidgetIndex( int32 Index );
 
 	/** Activates the widget and makes it the active index. */
 	UFUNCTION(BlueprintCallable, Category="Switcher")
-	virtual void SetActiveWidget(UWidget* Widget);
+	UMG_API virtual void SetActiveWidget(UWidget* Widget);
 
 	/** Get a widget at the provided index */
 	UFUNCTION( BlueprintCallable, Category = "Switcher" )
-	UWidget* GetWidgetAtIndex( int32 Index ) const;
+	UMG_API UWidget* GetWidgetAtIndex( int32 Index ) const;
 	
 	/** Get the reference of the currently active widget */
 	UFUNCTION(BlueprintCallable, Category = "Switcher")
-	UWidget* GetActiveWidget() const;
+	UMG_API UWidget* GetActiveWidget() const;
 	
 	// UWidget interface
-	virtual void SynchronizeProperties() override;
-	virtual void ReleaseSlateResources(bool bReleaseChildren) override;
+	UMG_API virtual void SynchronizeProperties() override;
+	UMG_API virtual void ReleaseSlateResources(bool bReleaseChildren) override;
 	// End of UWidget interface
 
 #if WITH_EDITOR
-	virtual const FText GetPaletteCategory() override;
-	virtual void OnDescendantSelectedByDesigner(UWidget* DescendantWidget) override;
-	virtual void OnDescendantDeselectedByDesigner(UWidget* DescendantWidget) override;
+	UMG_API virtual const FText GetPaletteCategory() override;
+	UMG_API virtual void OnDescendantSelectedByDesigner(UWidget* DescendantWidget) override;
+	UMG_API virtual void OnDescendantDeselectedByDesigner(UWidget* DescendantWidget) override;
 
 	// UObject interface
-	virtual void PostEditChangeProperty(struct FPropertyChangedEvent& PropertyChangedEvent) override;
+	UMG_API virtual void PostEditChangeProperty(struct FPropertyChangedEvent& PropertyChangedEvent) override;
 	// End of UObject interface
 #endif
 
 protected:
 
 	// UPanelWidget
-	virtual UClass* GetSlotClass() const override;
-	virtual void OnSlotAdded(UPanelSlot* Slot) override;
-	virtual void OnSlotRemoved(UPanelSlot* Slot) override;
+	UMG_API virtual UClass* GetSlotClass() const override;
+	UMG_API virtual void OnSlotAdded(UPanelSlot* Slot) override;
+	UMG_API virtual void OnSlotRemoved(UPanelSlot* Slot) override;
 	// End UPanelWidget
 
-	void SetActiveWidgetIndexForSlateWidget();
+	UMG_API void SetActiveWidgetIndexForSlateWidget();
 protected:
 
 	TSharedPtr<class SWidgetSwitcher> MyWidgetSwitcher;
 
 protected:
 	// UWidget interface
-	virtual TSharedRef<SWidget> RebuildWidget() override;
+	UMG_API virtual TSharedRef<SWidget> RebuildWidget() override;
 	// End of UWidget interface
 };
diff --git a/Engine/Source/Runtime/UMG/Public/Components/WidgetSwitcherSlot.h b/Engine/Source/Runtime/UMG/Public/Components/WidgetSwitcherSlot.h
--- a/Engine/Source/Runtime/UMG/Public/Components/WidgetSwitcherSlot.h
+++ b/Engine/Source/Runtime/UMG/Public/Components/WidgetSwitcherSlot.h
@@ -14,8 +14,8 @@
 class UWidget;
 
 /** The Slot for the UWidgetSwitcher, contains the widget that is flowed vertically */
-UCLASS()
-class UMG_API UWidgetSwitcherSlot : public UPanelSlot
+UCLASS(MinimalAPI)
+class UWidgetSwitcherSlot : public UPanelSlot
 {
 	GENERATED_UCLASS_BODY()
 
@@ -44,32 +44,32 @@ public:
 
 public:
 
-	FMargin GetPadding() const;
+	UMG_API FMargin GetPadding() const;
 
 	UFUNCTION(BlueprintCallable, Category="Layout|Widget Switcher Slot")
-	void SetPadding(FMargin InPadding);
+	UMG_API void SetPadding(FMargin InPadding);
 
-	EHorizontalAlignment GetHorizontalAlignment() const;
+	UMG_API EHorizontalAlignment GetHorizontalAlignment() const;
 
 	UFUNCTION(BlueprintCallable, Category="Layout|Widget Switcher Slot")
-	void SetHorizontalAlignment(EHorizontalAlignment InHorizontalAlignment);
+	UMG_API void SetHorizontalAlignment(EHorizontalAlignment InHorizontalAlignment);
 
-	EVerticalAlignment GetVerticalAlignment() const;
+	UMG_API EVerticalAlignment GetVerticalAlignment() const;
 
 	UFUNCTION(BlueprintCallable, Category="Layout|Widget Switcher Slot")
-	void SetVerticalAlignment(EVerticalAlignment InVerticalAlignment);
+	UMG_API void SetVerticalAlignment(EVerticalAlignment InVerticalAlignment);
 
 public:
 
 	//~ UPanelSlot interface
-	virtual void SynchronizeProperties() override;
+	UMG_API virtual void SynchronizeProperties() override;
 	//~ End of UPanelSlot interface
 
-	virtual void ReleaseSlateResources(bool bReleaseChildren) override;
+	UMG_API virtual void ReleaseSlateResources(bool bReleaseChildren) override;
 
 	/** Builds the underlying FSlot for the Slate layout panel. */
-	void BuildSlot(TSharedRef<SWidgetSwitcher> InWidgetSwitcher);
+	UMG_API void BuildSlot(TSharedRef<SWidgetSwitcher> InWidgetSwitcher);
 
 	/** Sets the content of this slot, removing existing content if needed. */
-	void SetContent(UWidget* NewContent);
+	UMG_API void SetContent(UWidget* NewContent);
 };
diff --git a/Engine/Source/Runtime/UMG/Public/Components/WindowTitleBarArea.h b/Engine/Source/Runtime/UMG/Public/Components/WindowTitleBarArea.h
--- a/Engine/Source/Runtime/UMG/Public/Components/WindowTitleBarArea.h
+++ b/Engine/Source/Runtime/UMG/Public/Components/WindowTitleBarArea.h
@@ -17,8 +17,8 @@ class SWindowTitleBarArea;
 * A panel for defining a region of the UI that should allow users to drag the window on desktop platforms.
 */
 
-UCLASS()
-class UMG_API UWindowTitleBarArea : public UContentWidget
+UCLASS(MinimalAPI)
+class UWindowTitleBarArea : public UContentWidget
 {
 	GENERATED_UCLASS_BODY()
 
@@ -37,49 +37,49 @@ public:
 public:
 
 	UFUNCTION(BlueprintCallable, Category = "Appearance")
-	void SetPadding(FMargin InPadding);
+	UMG_API void SetPadding(FMargin InPadding);
 
 	UFUNCTION(BlueprintCallable, Category = "Appearance")
-	void SetHorizontalAlignment(EHorizontalAlignment InHorizontalAlignment);
+	UMG_API void SetHorizontalAlignment(EHorizontalAlignment InHorizontalAlignment);
 
 	UFUNCTION(BlueprintCallable, Category = "Appearance")
-	void SetVerticalAlignment(EVerticalAlignment InVerticalAlignment);
+	UMG_API void SetVerticalAlignment(EVerticalAlignment InVerticalAlignment);
 
-	void SetWindowButtonsEnabled(bool InWindowButtonsEnabled);
+	UMG_API void SetWindowButtonsEnabled(bool InWindowButtonsEnabled);
 
-	bool IsWindowButtonsEnabled() const;
+	UMG_API bool IsWindowButtonsEnabled() const;
 
-	void SetDoubleClickTogglesFullscreen(bool InDoubleClickTogglesFullscreen);
+	UMG_API void SetDoubleClickTogglesFullscreen(bool InDoubleClickTogglesFullscreen);
 
-	bool IsDoubleClickTogglesFullscreen() const;
+	UMG_API bool IsDoubleClickTogglesFullscreen() const;
 
 public:
 
 	//~ Begin UVisual Interface
-	virtual void ReleaseSlateResources(bool bReleaseChildren) override;
+	UMG_API virtual void ReleaseSlateResources(bool bReleaseChildren) override;
 	//~ End UVisual Interface
 
 	//~ Begin UObject Interface
-	virtual void PostLoad() override;
+	UMG_API virtual void PostLoad() override;
 	//~ End UObject Interface
 
 	//~ Begin UWidget Interface
 #if WITH_EDITOR	
-	virtual const FText GetPaletteCategory() override;
+	UMG_API virtual const FText GetPaletteCategory() override;
 #endif
 	//~ End UWidget Interface
 
 protected:
 	//~ Begin UWidget Interface
-	virtual TSharedRef<SWidget> RebuildWidget() override;
+	UMG_API virtual TSharedRef<SWidget> RebuildWidget() override;
 	//~ End UWidget Interface
 
 protected:
 
 	// UPanelWidget
-	virtual UClass* GetSlotClass() const override;
-	virtual void OnSlotAdded(UPanelSlot* Slot) override;
-	virtual void OnSlotRemoved(UPanelSlot* Slot) override;
+	UMG_API virtual UClass* GetSlotClass() const override;
+	UMG_API virtual void OnSlotAdded(UPanelSlot* Slot) override;
+	UMG_API virtual void OnSlotRemoved(UPanelSlot* Slot) override;
 	// End UPanelWidget
 
 protected:
diff --git a/Engine/Source/Runtime/UMG/Public/Components/WindowTitleBarAreaSlot.h b/Engine/Source/Runtime/UMG/Public/Components/WindowTitleBarAreaSlot.h
--- a/Engine/Source/Runtime/UMG/Public/Components/WindowTitleBarAreaSlot.h
+++ b/Engine/Source/Runtime/UMG/Public/Components/WindowTitleBarAreaSlot.h
@@ -13,27 +13,27 @@
 class SWindowTitleBarArea;
 
 /** The Slot for the UWindowTitleBarArea */
-UCLASS()
-class UMG_API UWindowTitleBarAreaSlot : public UPanelSlot
+UCLASS(MinimalAPI)
+class UWindowTitleBarAreaSlot : public UPanelSlot
 {
 	GENERATED_UCLASS_BODY()
 
 public:
 
-	FMargin GetPadding() const;
+	UMG_API FMargin GetPadding() const;
 
 	UFUNCTION(BlueprintCallable, Category="Layout|WindowTitleBarArea Slot")
-	void SetPadding(FMargin InPadding);
+	UMG_API void SetPadding(FMargin InPadding);
 
-	EHorizontalAlignment GetHorizontalAlignment() const;
+	UMG_API EHorizontalAlignment GetHorizontalAlignment() const;
 
 	UFUNCTION(BlueprintCallable, Category="Layout|WindowTitleBarArea Slot")
-	void SetHorizontalAlignment(EHorizontalAlignment InHorizontalAlignment);
+	UMG_API void SetHorizontalAlignment(EHorizontalAlignment InHorizontalAlignment);
 
-	EVerticalAlignment GetVerticalAlignment() const;
+	UMG_API EVerticalAlignment GetVerticalAlignment() const;
 
 	UFUNCTION(BlueprintCallable, Category = "Layout|WindowTitleBarArea Slot")
-	void SetVerticalAlignment(EVerticalAlignment InVerticalAlignment);
+	UMG_API void SetVerticalAlignment(EVerticalAlignment InVerticalAlignment);
 
 protected:
 
@@ -55,13 +55,13 @@ protected:
 public:
 
 	//~ UPanelSlot interface
-	virtual void SynchronizeProperties() override;
+	UMG_API virtual void SynchronizeProperties() override;
 	//~ End of UPanelSlot interface
 
 	/** Builds the underlying FSlot for the Slate layout panel. */
-	void BuildSlot(TSharedRef<SWindowTitleBarArea> WindowTitleBarArea);
+	UMG_API void BuildSlot(TSharedRef<SWindowTitleBarArea> WindowTitleBarArea);
 
-	virtual void ReleaseSlateResources(bool bReleaseChildren) override;
+	UMG_API virtual void ReleaseSlateResources(bool bReleaseChildren) override;
 
 private:
 
diff --git a/Engine/Source/Runtime/UMG/Public/Components/WrapBox.h b/Engine/Source/Runtime/UMG/Public/Components/WrapBox.h
--- a/Engine/Source/Runtime/UMG/Public/Components/WrapBox.h
+++ b/Engine/Source/Runtime/UMG/Public/Components/WrapBox.h
@@ -18,8 +18,8 @@ class UWrapBoxSlot;
  * * Flows
  * * Wraps
  */
-UCLASS()
-class UMG_API UWrapBox : public UPanelWidget
+UCLASS(MinimalAPI)
+class UWrapBox : public UPanelWidget
 {
 	GENERATED_UCLASS_BODY()
 
@@ -51,61 +51,61 @@ public:
 
 public:
 	/** Gets the inner slot padding goes between slots sharing borders */
-	FVector2D GetInnerSlotPadding() const;
+	UMG_API FVector2D GetInnerSlotPadding() const;
 
 	/** Sets the inner slot padding goes between slots sharing borders */
 	UFUNCTION(BlueprintCallable, Category="Content Layout")
-	void SetInnerSlotPadding(FVector2D InPadding);
+	UMG_API void SetInnerSlotPadding(FVector2D InPadding);
 
 	/** */
-	float GetWrapSize() const;
+	UMG_API float GetWrapSize() const;
 
 	/** */
-	void SetWrapSize(float InWrapSize);
+	UMG_API void SetWrapSize(float InWrapSize);
 
 	/** */
-	bool UseExplicitWrapSize() const;
+	UMG_API bool UseExplicitWrapSize() const;
 
 	/** */
-	void SetExplicitWrapSize(bool bInExplicitWrapSize);
+	UMG_API void SetExplicitWrapSize(bool bInExplicitWrapSize);
 
 	/** */
-	EHorizontalAlignment GetHorizontalAlignment() const;
+	UMG_API EHorizontalAlignment GetHorizontalAlignment() const;
 
 	/** */
 	UFUNCTION(BlueprintCallable, Category="Content Layout")
-	void SetHorizontalAlignment(EHorizontalAlignment InHorizontalAlignment);
+	UMG_API void SetHorizontalAlignment(EHorizontalAlignment InHorizontalAlignment);
 
 	/** */
-	EOrientation GetOrientation() const;
+	UMG_API EOrientation GetOrientation() const;
 
 	/** */
-	void SetOrientation(EOrientation InOrientation);
+	UMG_API void SetOrientation(EOrientation InOrientation);
 
 public:
 	UFUNCTION(BlueprintCallable, Category="Panel")
-	UWrapBoxSlot* AddChildToWrapBox(UWidget* Content);
+	UMG_API UWrapBoxSlot* AddChildToWrapBox(UWidget* Content);
 
 #if WITH_EDITOR
 	//~ UWidget interface
-	virtual const FText GetPaletteCategory() override;
+	UMG_API virtual const FText GetPaletteCategory() override;
 	//~ End UWidget interface
 #endif
 
 protected:
 
 	//~ UPanelWidget
-	virtual UClass* GetSlotClass() const override;
-	virtual void OnSlotAdded(UPanelSlot* Slot) override;
-	virtual void OnSlotRemoved(UPanelSlot* Slot) override;
+	UMG_API virtual UClass* GetSlotClass() const override;
+	UMG_API virtual void OnSlotAdded(UPanelSlot* Slot) override;
+	UMG_API virtual void OnSlotRemoved(UPanelSlot* Slot) override;
 	//~ End UPanelWidget
 
 	//~ UWidget interface
-	virtual void SynchronizeProperties() override;
+	UMG_API virtual void SynchronizeProperties() override;
 	//~ End of UWidget interface
 
 	//~ UVisual interface
-	virtual void ReleaseSlateResources(bool bReleaseChildren) override;
+	UMG_API virtual void ReleaseSlateResources(bool bReleaseChildren) override;
 	//~ End of UVisual interface
 
 protected:
@@ -114,6 +114,6 @@ protected:
 
 protected:
 	//~ UWidget interface
-	virtual TSharedRef<SWidget> RebuildWidget() override;
+	UMG_API virtual TSharedRef<SWidget> RebuildWidget() override;
 	//~ End of UWidget interface
 };
diff --git a/Engine/Source/Runtime/UMG/Public/Components/WrapBoxSlot.h b/Engine/Source/Runtime/UMG/Public/Components/WrapBoxSlot.h
--- a/Engine/Source/Runtime/UMG/Public/Components/WrapBoxSlot.h
+++ b/Engine/Source/Runtime/UMG/Public/Components/WrapBoxSlot.h
@@ -12,8 +12,8 @@
 #include "WrapBoxSlot.generated.h"
 
 /** The Slot for the UWrapBox, contains the widget that is flowed vertically */
-UCLASS()
-class UMG_API UWrapBoxSlot : public UPanelSlot
+UCLASS(MinimalAPI)
+class UWrapBoxSlot : public UPanelSlot
 {
 	GENERATED_UCLASS_BODY()
 
@@ -54,46 +54,46 @@ public:
 
 public:
 
-	FMargin GetPadding() const;
+	UMG_API FMargin GetPadding() const;
 
 	UFUNCTION(BlueprintCallable, Category="Layout|Wrap Box Slot")
-	void SetPadding(FMargin InPadding);
+	UMG_API void SetPadding(FMargin InPadding);
 
-	bool DoesFillEmptySpace() const;
+	UMG_API bool DoesFillEmptySpace() const;
 
 	UFUNCTION(BlueprintCallable, Category="Layout|Wrap Box Slot")
-	void SetFillEmptySpace(bool InbFillEmptySpace);
+	UMG_API void SetFillEmptySpace(bool InbFillEmptySpace);
 
-	float GetFillSpanWhenLessThan() const;
+	UMG_API float GetFillSpanWhenLessThan() const;
 
 	UFUNCTION(BlueprintCallable, Category="Layout|Wrap Box Slot")
-	void SetFillSpanWhenLessThan(float InFillSpanWhenLessThan);
+	UMG_API void SetFillSpanWhenLessThan(float InFillSpanWhenLessThan);
 
-	EHorizontalAlignment  GetHorizontalAlignment() const;
+	UMG_API EHorizontalAlignment  GetHorizontalAlignment() const;
 
 	UFUNCTION(BlueprintCallable, Category="Layout|Wrap Box Slot")
-	void SetHorizontalAlignment(EHorizontalAlignment InHorizontalAlignment);
+	UMG_API void SetHorizontalAlignment(EHorizontalAlignment InHorizontalAlignment);
 
-	EVerticalAlignment GetVerticalAlignment() const;
+	UMG_API EVerticalAlignment GetVerticalAlignment() const;
 
 	UFUNCTION(BlueprintCallable, Category="Layout|Wrap Box Slot")
-	void SetVerticalAlignment(EVerticalAlignment InVerticalAlignment);
+	UMG_API void SetVerticalAlignment(EVerticalAlignment InVerticalAlignment);
 
-	bool DoesForceNewLine() const;
+	UMG_API bool DoesForceNewLine() const;
 
 	UFUNCTION(BlueprintCallable, Category="Layout|Wrap Box Slot")
-	void SetNewLine(bool InForceNewLine);
+	UMG_API void SetNewLine(bool InForceNewLine);
 
 public:
 
 	// UPanelSlot interface
-	virtual void SynchronizeProperties() override;
+	UMG_API virtual void SynchronizeProperties() override;
 	// End of UPanelSlot interface
 
-	virtual void ReleaseSlateResources(bool bReleaseChildren) override;
+	UMG_API virtual void ReleaseSlateResources(bool bReleaseChildren) override;
 
 	/** Builds the underlying FSlot for the Slate layout panel. */
-	void BuildSlot(TSharedRef<SWrapBox> InWrapBox);
+	UMG_API void BuildSlot(TSharedRef<SWrapBox> InWrapBox);
 
 private:
 	//TODO UMG Slots should hold weak or shared refs to slots.
diff --git a/Engine/Source/Runtime/UMG/Public/Extensions/UserWidgetExtension.h b/Engine/Source/Runtime/UMG/Public/Extensions/UserWidgetExtension.h
--- a/Engine/Source/Runtime/UMG/Public/Extensions/UserWidgetExtension.h
+++ b/Engine/Source/Runtime/UMG/Public/Extensions/UserWidgetExtension.h
@@ -15,8 +15,8 @@ class UUserWidget;
 /**
  * UserWidgetExtension is the base class for components that define reusable behavior that can be added to different types of Widgets.
  */
-UCLASS(Abstract, DefaultToInstanced, Within=UserWidget)
-class UMG_API UUserWidgetExtension : public UObject
+UCLASS(Abstract, DefaultToInstanced, Within=UserWidget, MinimalAPI)
+class UUserWidgetExtension : public UObject
 {
 	GENERATED_BODY()
 	
diff --git a/Engine/Source/Runtime/UMG/Public/Extensions/WidgetBlueprintGeneratedClassExtension.h b/Engine/Source/Runtime/UMG/Public/Extensions/WidgetBlueprintGeneratedClassExtension.h
--- a/Engine/Source/Runtime/UMG/Public/Extensions/WidgetBlueprintGeneratedClassExtension.h
+++ b/Engine/Source/Runtime/UMG/Public/Extensions/WidgetBlueprintGeneratedClassExtension.h
@@ -11,8 +11,8 @@
 /**
  * WidgetExtension is the base class for components that define reusable behavior that can be added to different types of Widgets.
  */
-UCLASS(Abstract, DefaultToInstanced)
-class UMG_API UWidgetBlueprintGeneratedClassExtension : public UObject
+UCLASS(Abstract, DefaultToInstanced, MinimalAPI)
+class UWidgetBlueprintGeneratedClassExtension : public UObject
 {
 	GENERATED_BODY()
 	
diff --git a/Engine/Source/Runtime/UMG/Public/Slate/SMeshWidget.h b/Engine/Source/Runtime/UMG/Public/Slate/SMeshWidget.h
--- a/Engine/Source/Runtime/UMG/Public/Slate/SMeshWidget.h
+++ b/Engine/Source/Runtime/UMG/Public/Slate/SMeshWidget.h
@@ -20,7 +20,7 @@ struct FSlateBrush;
  * The Mesh's material is used.
  * Hardware instancing is supported.
  */
-class UMG_API SMeshWidget : public SLeafWidget, public FGCObject
+class SMeshWidget : public SLeafWidget, public FGCObject
 {
 public:
 	SLATE_BEGIN_ARGS(SMeshWidget)
@@ -30,17 +30,17 @@ public:
 		SLATE_ARGUMENT(USlateVectorArtData*, MeshData)
 	SLATE_END_ARGS()
 
-	void Construct(const FArguments& Args);
+	UMG_API void Construct(const FArguments& Args);
 
 	/**
 	 * Draw the InStaticMesh when this widget paints.
 	 *
 	 * @return the Index of the mesh data that was added; cache this value for use with @see FRenderRun.
 	 */
-	uint32 AddMesh(USlateVectorArtData& InMeshData);
+	UMG_API uint32 AddMesh(USlateVectorArtData& InMeshData);
 
 	/** Much like AddMesh, but also enables instancing support for this MeshId. */
-	uint32 AddMeshWithInstancing(USlateVectorArtData& InMeshData, int32 InitialBufferSize = 1);
+	UMG_API uint32 AddMeshWithInstancing(USlateVectorArtData& InMeshData, int32 InitialBufferSize = 1);
 
 	/**
 	 * Switch from static material to material instance dynamic.
@@ -49,10 +49,10 @@ public:
 	 * 
 	 * @return The MID for this Asset on which parameters can be set.
 	 */
-	UMaterialInstanceDynamic* ConvertToMID( uint32 MeshId );
+	UMG_API UMaterialInstanceDynamic* ConvertToMID( uint32 MeshId );
 
 	/** Discard any previous runs and reserve space for new render runs if needed. */
-	void ClearRuns(int32 NumRuns);
+	UMG_API void ClearRuns(int32 NumRuns);
 
 	/**
 	 * Tell the widget to draw instances of a mesh a given number of times starting at
@@ -68,25 +68,25 @@ public:
 	}
 
 	/** Enable hardware instancing */
-	void EnableInstancing(uint32 MeshId, int32 InitialSize);
+	UMG_API void EnableInstancing(uint32 MeshId, int32 InitialSize);
 
 	/** Updates the per instance buffer. Automatically enables hardware instancing. */
-	void UpdatePerInstanceBuffer(uint32 MeshId, FSlateInstanceBufferData& Data);
+	UMG_API void UpdatePerInstanceBuffer(uint32 MeshId, FSlateInstanceBufferData& Data);
 
 protected:
 	// BEGIN SLeafWidget interface
-	virtual int32 OnPaint(const FPaintArgs& Args, const FGeometry& AllottedGeometry, const FSlateRect& MyCullingRect, FSlateWindowElementList& OutDrawElements, int32 LayerId, const FWidgetStyle& InWidgetStyle, bool bParentEnabled) const override;
-	virtual FVector2D ComputeDesiredSize(float) const override;
+	UMG_API virtual int32 OnPaint(const FPaintArgs& Args, const FGeometry& AllottedGeometry, const FSlateRect& MyCullingRect, FSlateWindowElementList& OutDrawElements, int32 LayerId, const FWidgetStyle& InWidgetStyle, bool bParentEnabled) const override;
+	UMG_API virtual FVector2D ComputeDesiredSize(float) const override;
 	// END SLeafWidget interface
 
 	// ~ FGCObject
-	virtual void AddReferencedObjects(FReferenceCollector& Collector) override;
-	virtual FString GetReferencerName() const override;
+	UMG_API virtual void AddReferencedObjects(FReferenceCollector& Collector) override;
+	UMG_API virtual FString GetReferencerName() const override;
 	// ~ FGCObject
 
 protected:
-	static void PushUpdate(uint32 VectorArtId, SMeshWidget& Widget, const FVector2D& Position, float Scale, uint32 BaseAddress);
-	static void PushUpdate(uint32 VectorArtId, SMeshWidget& Widget, const FVector2D& Position, float Scale, float OptionalFloat = 0);
+	static UMG_API void PushUpdate(uint32 VectorArtId, SMeshWidget& Widget, const FVector2D& Position, float Scale, uint32 BaseAddress);
+	static UMG_API void PushUpdate(uint32 VectorArtId, SMeshWidget& Widget, const FVector2D& Position, float Scale, float OptionalFloat = 0);
 
 	struct FRenderData
 	{
diff --git a/Engine/Source/Runtime/UMG/Public/Slate/SObjectTableRow.h b/Engine/Source/Runtime/UMG/Public/Slate/SObjectTableRow.h
--- a/Engine/Source/Runtime/UMG/Public/Slate/SObjectTableRow.h
+++ b/Engine/Source/Runtime/UMG/Public/Slate/SObjectTableRow.h
@@ -16,7 +16,7 @@ public:
 	virtual UListViewBase* GetOwningListView() const = 0;
 	virtual UUserWidget* GetUserWidget() const = 0;
 
-	UMG_API static TSharedPtr<const IObjectTableRow> ObjectRowFromUserWidget(const UUserWidget* RowUserWidget)
+	static TSharedPtr<const IObjectTableRow> ObjectRowFromUserWidget(const UUserWidget* RowUserWidget)
 	{
 		TWeakPtr<const IObjectTableRow>* ObjectRow = ObjectRowsByUserWidget.Find(RowUserWidget);
 		if (ObjectRow && ObjectRow->IsValid())
diff --git a/Engine/Source/Runtime/UMG/Public/Slate/SObjectWidget.h b/Engine/Source/Runtime/UMG/Public/Slate/SObjectWidget.h
--- a/Engine/Source/Runtime/UMG/Public/Slate/SObjectWidget.h
+++ b/Engine/Source/Runtime/UMG/Public/Slate/SObjectWidget.h
@@ -27,7 +27,7 @@ template< typename ObjectType > class TAttribute;
  * The UWidget allowing it to be garbage collected.  It also forwards the slate events to the UUserWidget
  * so that it can forward them to listeners.
  */
-class UMG_API SObjectWidget : public SCompoundWidget, public FGCObject
+class SObjectWidget : public SCompoundWidget, public FGCObject
 {
 	SLATE_BEGIN_ARGS(SObjectWidget)
 	{
@@ -37,68 +37,68 @@ class UMG_API SObjectWidget : public SCompoundWidget, public FGCObject
 		SLATE_DEFAULT_SLOT(FArguments, Content)
 	SLATE_END_ARGS()
 
-	virtual ~SObjectWidget(void);
+	UMG_API virtual ~SObjectWidget(void);
 
-	void Construct(const FArguments& InArgs, UUserWidget* InWidgetObject);
+	UMG_API void Construct(const FArguments& InArgs, UUserWidget* InWidgetObject);
 
-	void ResetWidget();
+	UMG_API void ResetWidget();
 
 	// FGCObject interface
-	virtual FString GetReferencerName() const override;
-	virtual void AddReferencedObjects(FReferenceCollector& Collector) override;
+	UMG_API virtual FString GetReferencerName() const override;
+	UMG_API virtual void AddReferencedObjects(FReferenceCollector& Collector) override;
 	// End of FGCObject interface
 
 	UUserWidget* GetWidgetObject() const { return WidgetObject; }
 
-	void SetPadding(const TAttribute<FMargin>& InMargin);
+	UMG_API void SetPadding(const TAttribute<FMargin>& InMargin);
 
 	/** SWidget Tick override.  Note this will not be called if bCanTick is set to false by the UserWidget */
-	virtual void Tick( const FGeometry& AllottedGeometry, const double InCurrentTime, const float InDeltaTime ) override;
-	virtual int32 OnPaint(const FPaintArgs& Args, const FGeometry& AllottedGeometry, const FSlateRect& MyCullingRect, FSlateWindowElementList& OutDrawElements, int32 LayerId, const FWidgetStyle& InWidgetStyle, bool bParentEnabled) const override;
+	UMG_API virtual void Tick( const FGeometry& AllottedGeometry, const double InCurrentTime, const float InDeltaTime ) override;
+	UMG_API virtual int32 OnPaint(const FPaintArgs& Args, const FGeometry& AllottedGeometry, const FSlateRect& MyCullingRect, FSlateWindowElementList& OutDrawElements, int32 LayerId, const FWidgetStyle& InWidgetStyle, bool bParentEnabled) const override;
 
-	virtual FVector2D ComputeDesiredSize(float LayoutScaleMultiplier) const override;
+	UMG_API virtual FVector2D ComputeDesiredSize(float LayoutScaleMultiplier) const override;
 
-	virtual bool IsInteractable() const override;
-	virtual bool SupportsKeyboardFocus() const override;
+	UMG_API virtual bool IsInteractable() const override;
+	UMG_API virtual bool SupportsKeyboardFocus() const override;
 
-	virtual FReply OnFocusReceived(const FGeometry& MyGeometry, const FFocusEvent& InFocusEvent) override;
-	virtual void OnFocusLost(const FFocusEvent& InFocusEvent) override;
-	virtual void OnFocusChanging(const FWeakWidgetPath& PreviousFocusPath, const FWidgetPath& NewWidgetPath, const FFocusEvent& InFocusEvent) override;
+	UMG_API virtual FReply OnFocusReceived(const FGeometry& MyGeometry, const FFocusEvent& InFocusEvent) override;
+	UMG_API virtual void OnFocusLost(const FFocusEvent& InFocusEvent) override;
+	UMG_API virtual void OnFocusChanging(const FWeakWidgetPath& PreviousFocusPath, const FWidgetPath& NewWidgetPath, const FFocusEvent& InFocusEvent) override;
 	
-	virtual FReply OnKeyChar(const FGeometry& MyGeometry, const FCharacterEvent& InCharacterEvent) override;
-	virtual FReply OnPreviewKeyDown(const FGeometry& MyGeometry, const FKeyEvent& InKeyEvent) override;
-	virtual FReply OnKeyDown(const FGeometry& MyGeometry, const FKeyEvent& InKeyEvent) override;
-	virtual FReply OnKeyUp(const FGeometry& MyGeometry, const FKeyEvent& InKeyEvent) override;
-	virtual FReply OnAnalogValueChanged(const FGeometry& MyGeometry, const FAnalogInputEvent& InAnalogInputEvent) override;
+	UMG_API virtual FReply OnKeyChar(const FGeometry& MyGeometry, const FCharacterEvent& InCharacterEvent) override;
+	UMG_API virtual FReply OnPreviewKeyDown(const FGeometry& MyGeometry, const FKeyEvent& InKeyEvent) override;
+	UMG_API virtual FReply OnKeyDown(const FGeometry& MyGeometry, const FKeyEvent& InKeyEvent) override;
+	UMG_API virtual FReply OnKeyUp(const FGeometry& MyGeometry, const FKeyEvent& InKeyEvent) override;
+	UMG_API virtual FReply OnAnalogValueChanged(const FGeometry& MyGeometry, const FAnalogInputEvent& InAnalogInputEvent) override;
 	
-	virtual FReply OnMouseButtonDown(const FGeometry& MyGeometry, const FPointerEvent& MouseEvent) override;
-	virtual FReply OnPreviewMouseButtonDown(const FGeometry& MyGeometry, const FPointerEvent& MouseEvent) override;
-	virtual FReply OnMouseButtonUp(const FGeometry& MyGeometry, const FPointerEvent& MouseEvent) override;
-	virtual FReply OnMouseMove(const FGeometry& MyGeometry, const FPointerEvent& MouseEvent) override;
-	virtual void OnMouseEnter(const FGeometry& MyGeometry, const FPointerEvent& MouseEvent) override;
-	virtual void OnMouseLeave(const FPointerEvent& MouseEvent) override;
-	virtual FReply OnMouseWheel(const FGeometry& MyGeometry, const FPointerEvent& MouseEvent) override;
-	virtual FCursorReply OnCursorQuery(const FGeometry& MyGeometry, const FPointerEvent& CursorEvent) const override;
-	virtual FReply OnMouseButtonDoubleClick(const FGeometry& InMyGeometry, const FPointerEvent& InMouseEvent) override;
+	UMG_API virtual FReply OnMouseButtonDown(const FGeometry& MyGeometry, const FPointerEvent& MouseEvent) override;
+	UMG_API virtual FReply OnPreviewMouseButtonDown(const FGeometry& MyGeometry, const FPointerEvent& MouseEvent) override;
+	UMG_API virtual FReply OnMouseButtonUp(const FGeometry& MyGeometry, const FPointerEvent& MouseEvent) override;
+	UMG_API virtual FReply OnMouseMove(const FGeometry& MyGeometry, const FPointerEvent& MouseEvent) override;
+	UMG_API virtual void OnMouseEnter(const FGeometry& MyGeometry, const FPointerEvent& MouseEvent) override;
+	UMG_API virtual void OnMouseLeave(const FPointerEvent& MouseEvent) override;
+	UMG_API virtual FReply OnMouseWheel(const FGeometry& MyGeometry, const FPointerEvent& MouseEvent) override;
+	UMG_API virtual FCursorReply OnCursorQuery(const FGeometry& MyGeometry, const FPointerEvent& CursorEvent) const override;
+	UMG_API virtual FReply OnMouseButtonDoubleClick(const FGeometry& InMyGeometry, const FPointerEvent& InMouseEvent) override;
 	
-	virtual FReply OnDragDetected(const FGeometry& MyGeometry, const FPointerEvent& MouseEvent) override;
-	virtual void OnDragEnter(const FGeometry& MyGeometry, const FDragDropEvent& DragDropEvent) override;
-	virtual void OnDragLeave(const FDragDropEvent& DragDropEvent) override;
-	virtual FReply OnDragOver(const FGeometry& MyGeometry, const FDragDropEvent& DragDropEvent) override;
-	virtual FReply OnDrop(const FGeometry& MyGeometry, const FDragDropEvent& DragDropEvent) override;
+	UMG_API virtual FReply OnDragDetected(const FGeometry& MyGeometry, const FPointerEvent& MouseEvent) override;
+	UMG_API virtual void OnDragEnter(const FGeometry& MyGeometry, const FDragDropEvent& DragDropEvent) override;
+	UMG_API virtual void OnDragLeave(const FDragDropEvent& DragDropEvent) override;
+	UMG_API virtual FReply OnDragOver(const FGeometry& MyGeometry, const FDragDropEvent& DragDropEvent) override;
+	UMG_API virtual FReply OnDrop(const FGeometry& MyGeometry, const FDragDropEvent& DragDropEvent) override;
 
-	void OnDragCancelled(const FDragDropEvent& DragDropEvent, UDragDropOperation* Operation);
+	UMG_API void OnDragCancelled(const FDragDropEvent& DragDropEvent, UDragDropOperation* Operation);
 	
-	virtual FReply OnTouchGesture(const FGeometry& MyGeometry, const FPointerEvent& GestureEvent) override;
-	virtual FReply OnTouchStarted(const FGeometry& MyGeometry, const FPointerEvent& InTouchEvent) override;
-	virtual FReply OnTouchMoved(const FGeometry& MyGeometry, const FPointerEvent& InTouchEvent) override;
-	virtual FReply OnTouchEnded(const FGeometry& MyGeometry, const FPointerEvent& InTouchEvent) override;
-	virtual FReply OnMotionDetected(const FGeometry& MyGeometry, const FMotionEvent& InMotionEvent) override;
-	virtual FReply OnTouchForceChanged(const FGeometry& MyGeometry, const FPointerEvent& InTouchEvent) override;
+	UMG_API virtual FReply OnTouchGesture(const FGeometry& MyGeometry, const FPointerEvent& GestureEvent) override;
+	UMG_API virtual FReply OnTouchStarted(const FGeometry& MyGeometry, const FPointerEvent& InTouchEvent) override;
+	UMG_API virtual FReply OnTouchMoved(const FGeometry& MyGeometry, const FPointerEvent& InTouchEvent) override;
+	UMG_API virtual FReply OnTouchEnded(const FGeometry& MyGeometry, const FPointerEvent& InTouchEvent) override;
+	UMG_API virtual FReply OnMotionDetected(const FGeometry& MyGeometry, const FMotionEvent& InMotionEvent) override;
+	UMG_API virtual FReply OnTouchForceChanged(const FGeometry& MyGeometry, const FPointerEvent& InTouchEvent) override;
 
-	virtual FNavigationReply OnNavigation(const FGeometry& MyGeometry, const FNavigationEvent& InNavigationEvent) override;
+	UMG_API virtual FNavigationReply OnNavigation(const FGeometry& MyGeometry, const FNavigationEvent& InNavigationEvent) override;
 
-	virtual void OnMouseCaptureLost(const FCaptureLostEvent& CaptureLostEvent) override;
+	UMG_API virtual void OnMouseCaptureLost(const FCaptureLostEvent& CaptureLostEvent) override;
 protected:
 
 	/** The UWidget that created this SObjectWidget who needs to be kept alive. */
