UMG: Add getters for sizebox overrides
#10707

diff --git a/Engine/Source/Runtime/UMG/Private/Components/SizeBox.cpp b/Engine/Source/Runtime/UMG/Private/Components/SizeBox.cpp
--- a/Engine/Source/Runtime/UMG/Private/Components/SizeBox.cpp
+++ b/Engine/Source/Runtime/UMG/Private/Components/SizeBox.cpp
@@ -49,6 +49,11 @@ float USizeBox::GetWidthOverride() const
 	return WidthOverride;
 }
 
+bool USizeBox::IsWidthOverride() const
+{
+	return bOverride_WidthOverride;
+}
+
 void USizeBox::SetWidthOverride(float InWidthOverride)
 {
 	bOverride_WidthOverride = true;
@@ -73,6 +78,11 @@ float USizeBox::GetHeightOverride() const
 	return HeightOverride;
 }
 
+bool USizeBox::IsHeightOverride() const
+{
+	return bOverride_HeightOverride;
+}
+
 void USizeBox::SetHeightOverride(float InHeightOverride)
 {
 	bOverride_HeightOverride = true;
@@ -97,6 +107,11 @@ float USizeBox::GetMinDesiredWidth() const
 	return MinDesiredWidth;
 }
 
+bool USizeBox::IsMinDesiredWidthOverride() const
+{
+	return bOverride_MinDesiredWidth;
+}
+
 void USizeBox::SetMinDesiredWidth(float InMinDesiredWidth)
 {
 	bOverride_MinDesiredWidth = true;
@@ -121,6 +136,11 @@ float USizeBox::GetMinDesiredHeight() const
 	return MinDesiredHeight;
 }
 
+bool USizeBox::IsMinDesiredHeightOverride() const
+{
+	return bOverride_MinDesiredHeight;
+}
+
 void USizeBox::SetMinDesiredHeight(float InMinDesiredHeight)
 {
 	bOverride_MinDesiredHeight = true;
@@ -145,6 +165,11 @@ float USizeBox::GetMaxDesiredWidth() const
 	return MaxDesiredWidth;
 }
 
+bool USizeBox::IsMaxDesiredWidthOverride() const
+{
+	return bOverride_MaxDesiredWidth;
+}
+
 void USizeBox::SetMaxDesiredWidth(float InMaxDesiredWidth)
 {
 	bOverride_MaxDesiredWidth = true;
@@ -169,6 +194,11 @@ float USizeBox::GetMaxDesiredHeight() const
 	return MaxDesiredHeight;
 }
 
+bool USizeBox::IsMaxDesiredHeightOverride() const
+{
+	return bOverride_MaxDesiredHeight;
+}
+
 void USizeBox::SetMaxDesiredHeight(float InMaxDesiredHeight)
 {
 	bOverride_MaxDesiredHeight = true;
@@ -193,6 +223,11 @@ float USizeBox::GetMinAspectRatio() const
 	return MinAspectRatio;
 }
 
+bool USizeBox::IsMinAspectRatioOverride() const
+{
+	return bOverride_MinAspectRatio;
+}
+
 void USizeBox::SetMinAspectRatio(float InMinAspectRatio)
 {
 	bOverride_MinAspectRatio = true;
@@ -217,6 +252,11 @@ float USizeBox::GetMaxAspectRatio() const
 	return MaxAspectRatio;
 }
 
+bool USizeBox::IsMaxAspectRatioOverride() const
+{
+	return bOverride_MaxAspectRatio;
+}
+
 void USizeBox::SetMaxAspectRatio(float InMaxAspectRatio)
 {
 	bOverride_MaxAspectRatio = true;
diff --git a/Engine/Source/Runtime/UMG/Public/Components/SizeBox.h b/Engine/Source/Runtime/UMG/Public/Components/SizeBox.h
--- a/Engine/Source/Runtime/UMG/Public/Components/SizeBox.h
+++ b/Engine/Source/Runtime/UMG/Public/Components/SizeBox.h
@@ -113,6 +113,9 @@ public:
 	/** */
 	UMG_API float GetWidthOverride() const;
 
+	/** */
+	UMG_API bool IsWidthOverride() const;
+
 	/** When specified, ignore the content's desired size and report the WidthOverride as the Box's desired width. */
 	UFUNCTION(BlueprintCallable, Category="Layout|Size Box")
 	UMG_API void SetWidthOverride(float InWidthOverride);
@@ -123,6 +126,9 @@ public:
 	/** */
 	UMG_API float GetHeightOverride() const;
 
+	/** */
+	UMG_API bool IsHeightOverride() const;
+
 	/** When specified, ignore the content's desired size and report the HeightOverride as the Box's desired height. */
 	UFUNCTION(BlueprintCallable, Category="Layout|Size Box")
 	UMG_API void SetHeightOverride(float InHeightOverride);
@@ -133,6 +139,9 @@ public:
 	/** */
 	UMG_API float GetMinDesiredWidth() const;
 
+	/** */
+	UMG_API bool IsMinDesiredWidthOverride() const;
+
 	/** When specified, will report the MinDesiredWidth if larger than the content's desired width. */
 	UFUNCTION(BlueprintCallable, Category="Layout|Size Box")
 	UMG_API void SetMinDesiredWidth(float InMinDesiredWidth);
@@ -143,6 +152,9 @@ public:
 	/** */
 	UMG_API float GetMinDesiredHeight() const;
 
+	/** */
+	UMG_API bool IsMinDesiredHeightOverride() const;
+
 	/** When specified, will report the MinDesiredHeight if larger than the content's desired height. */
 	UFUNCTION(BlueprintCallable, Category="Layout|Size Box")
 	UMG_API void SetMinDesiredHeight(float InMinDesiredHeight);
@@ -153,6 +165,9 @@ public:
 	/** */
 	UMG_API float GetMaxDesiredWidth() const;
 
+	/** */
+	UMG_API bool IsMaxDesiredWidthOverride() const;
+
 	/** When specified, will report the MaxDesiredWidth if smaller than the content's desired width. */
 	UFUNCTION(BlueprintCallable, Category="Layout|Size Box")
 	UMG_API void SetMaxDesiredWidth(float InMaxDesiredWidth);
@@ -163,6 +178,9 @@ public:
 	/** */
 	UMG_API float GetMaxDesiredHeight() const;
 
+	/** */
+	UMG_API bool IsMaxDesiredHeightOverride() const;
+
 	/** When specified, will report the MaxDesiredHeight if smaller than the content's desired height. */
 	UFUNCTION(BlueprintCallable, Category="Layout|Size Box")
 	UMG_API void SetMaxDesiredHeight(float InMaxDesiredHeight);
@@ -173,6 +191,9 @@ public:
 	/** */
 	UMG_API float GetMinAspectRatio() const;
 
+	/** */
+	UMG_API bool IsMinAspectRatioOverride() const;
+
 	UFUNCTION(BlueprintCallable, Category="Layout|Size Box")
 	UMG_API void SetMinAspectRatio(float InMinAspectRatio);
 
@@ -182,6 +203,9 @@ public:
 	/** */
 	UMG_API float GetMaxAspectRatio() const;
 
+	/** */
+	UMG_API bool IsMaxAspectRatioOverride() const;
+
 	UFUNCTION(BlueprintCallable, Category="Layout|Size Box")
 	UMG_API void SetMaxAspectRatio(float InMaxAspectRatio);
 
