Undo CL 24872301. The new function only works when the items is already generated and not virtualized. The virtualized items are only refresh on the next tick.

diff --git a/Engine/Source/Runtime/UMG/Private/Components/ListView.cpp b/Engine/Source/Runtime/UMG/Private/Components/ListView.cpp
--- a/Engine/Source/Runtime/UMG/Private/Components/ListView.cpp
+++ b/Engine/Source/Runtime/UMG/Private/Components/ListView.cpp
@@ -95,8 +95,7 @@ int32 UListView::GetNumItems() const
 
 int32 UListView::GetIndexForItem(const UObject* Item) const
 {
-	TSharedPtr<SObjectTableRow<UObject*>> TableRow = GetObjectRowFromItem<SObjectTableRow>(const_cast<UObject*>(Item));
-	return TableRow.IsValid() ? TableRow->GetIndexInList() : INDEX_NONE;
+	return ListItems.IndexOfByKey(Item);
 }
 
 void UListView::ClearListItems()
