UMG: Always the Initialize event to be called when there is no player context. The behavior was introduced with CL 4166480. The PreConstruct, Construct and Destruct are already called regardless if there is a Player context.

diff --git a/Engine/Source/Runtime/UMG/Private/UserWidget.cpp b/Engine/Source/Runtime/UMG/Private/UserWidget.cpp
--- a/Engine/Source/Runtime/UMG/Private/UserWidget.cpp
+++ b/Engine/Source/Runtime/UMG/Private/UserWidget.cpp
@@ -153,7 +153,7 @@ bool UUserWidget::Initialize()
 			InitializeNamedSlots();
 		}
 
-		if (!IsDesignTime() && PlayerContext.IsValid())
+		if (!IsDesignTime())
 		{
 			NativeOnInitialized();
 		}
