Undo changelist 27149747
#rnx

diff --git a/Engine/Source/Runtime/UMG/Private/UserWidget.cpp b/Engine/Source/Runtime/UMG/Private/UserWidget.cpp
--- a/Engine/Source/Runtime/UMG/Private/UserWidget.cpp
+++ b/Engine/Source/Runtime/UMG/Private/UserWidget.cpp
@@ -153,7 +153,7 @@ bool UUserWidget::Initialize()
 			InitializeNamedSlots();
 		}
 
-		if (!IsDesignTime())
+		if (!IsDesignTime() && PlayerContext.IsValid())
 		{
 			NativeOnInitialized();
 		}
