UMG: Uncached the world when the player context changes. It's not a cached all solution, the player controller can change world without notifying the widget.
Use the object key instead of the TWeakObjectPtr as key for the TMap.
Add flag to remove the widget from the viewport when the world is removed.

diff --git a/Engine/Source/Runtime/UMG/Private/Blueprint/GameViewportSubsystem.cpp b/Engine/Source/Runtime/UMG/Private/Blueprint/GameViewportSubsystem.cpp
--- a/Engine/Source/Runtime/UMG/Private/Blueprint/GameViewportSubsystem.cpp
+++ b/Engine/Source/Runtime/UMG/Private/Blueprint/GameViewportSubsystem.cpp
@@ -34,6 +34,7 @@ void UGameViewportSubsystem::Initialize(FSubsystemCollectionBase& Collection)
 {
 	Super::Initialize(Collection);
 	
+	//FWorldDelegates::LevelRemovedFromWorld.AddUObject(this, &ThisClass::HandleLevelRemovedFromWorld);
 	FWorldDelegates::OnWorldCleanup.AddUObject(this, &ThisClass::HandleWorldCleanup);
 	FWorldDelegates::OnWorldBeginTearDown.AddUObject(this, &ThisClass::HandleRemoveWorld);
 	FWorldDelegates::OnPreWorldFinishDestroy.AddUObject(this, &ThisClass::HandleRemoveWorld);
@@ -41,6 +42,7 @@ void UGameViewportSubsystem::Initialize(FSubsystemCollectionBase& Collection)
 
 void UGameViewportSubsystem::Deinitialize()
 {
+	//FWorldDelegates::LevelRemovedFromWorld.RemoveAll(this);
 	FWorldDelegates::OnPreWorldFinishDestroy.RemoveAll(this);
 	FWorldDelegates::OnWorldBeginTearDown.RemoveAll(this);
 	FWorldDelegates::OnWorldCleanup.RemoveAll(this);
@@ -53,7 +55,7 @@ void UGameViewportSubsystem::Deinitialize()
 		TmpViewportWidgets.Reserve(ViewportWidgets.Num());
 		for (auto& Itt : ViewportWidgets)
 		{
-			if (UWidget* Widget = Itt.Key.Get())
+			if (UWidget* Widget = Itt.Key.ResolveObjectPtr())
 			{
 				TmpViewportWidgets.Emplace(Widget, Itt.Value.FullScreenWidget, Itt.Value.LocalPlayer);
 			}
@@ -92,8 +94,8 @@ bool UGameViewportSubsystem::IsWidgetAdded(const UWidget* Widget) const
 
 	if (Widget->bIsManagedByGameViewportSubsystem)
 	{
-		TWeakObjectPtr<UWidget> TempObject = const_cast<UWidget*>(Widget);
-		if (const FSlotInfo* FoundSlot = ViewportWidgets.Find(TempObject))
+		TObjectKey<UWidget> WidgetKey = Widget;
+		if (const FSlotInfo* FoundSlot = ViewportWidgets.Find(WidgetKey))
 		{
 			return FoundSlot->FullScreenWidget.IsValid();
 		}
@@ -149,7 +151,8 @@ bool UGameViewportSubsystem::AddToScreen(UWidget* Widget, ULocalPlayer* Player,
 	SConstraintCanvas::FSlot* RawSlot = nullptr;
 	TSharedPtr<SConstraintCanvas> FullScreenCanvas;
 	{
-		FSlotInfo& SlotInfo = ViewportWidgets.FindOrAdd(Widget);
+		TObjectKey<UWidget> WidgetKey = Widget;
+		FSlotInfo& SlotInfo = ViewportWidgets.FindOrAdd(WidgetKey);
 		Widget->bIsManagedByGameViewportSubsystem = true;
 
 		if (SlotInfo.FullScreenWidget.IsValid())
@@ -196,7 +199,8 @@ void UGameViewportSubsystem::RemoveWidget(UWidget* Widget)
 	if (Widget && Widget->bIsManagedByGameViewportSubsystem)
 	{
 		FSlotInfo SlotInfo;
-		ViewportWidgets.RemoveAndCopyValue(Widget, SlotInfo);
+		TObjectKey<UWidget> WidgetKey = Widget;
+		ViewportWidgets.RemoveAndCopyValue(WidgetKey, SlotInfo);
 		RemoveWidgetInternal(Widget, SlotInfo.FullScreenWidget, SlotInfo.LocalPlayer);
 
 		OnWidgetRemoved.Broadcast(Widget);
@@ -229,8 +233,8 @@ FGameViewportWidgetSlot UGameViewportSubsystem::GetWidgetSlot(const UWidget* Wid
 {
 	if (Widget && Widget->bIsManagedByGameViewportSubsystem)
 	{
-		TWeakObjectPtr<UWidget> TempObject = const_cast<UWidget*>(Widget);
-		const FSlotInfo* SlotInfo = ViewportWidgets.Find(TempObject);
+		TObjectKey<UWidget> WidgetKey = Widget;
+		const FSlotInfo* SlotInfo = ViewportWidgets.Find(WidgetKey);
 		ensureMsgf(SlotInfo != nullptr, TEXT("The bIsManagedByGameViewportSubsystem should matches the ViewportWidgets state."));
 		if (SlotInfo)
 		{
@@ -244,8 +248,8 @@ void UGameViewportSubsystem::SetWidgetSlot(UWidget* Widget, FGameViewportWidgetS
 {
 	if (Widget && !Widget->HasAnyFlags(RF_BeginDestroyed))
 	{
-		TWeakObjectPtr<UWidget> TempObject = const_cast<UWidget*>(Widget);
-		FSlotInfo& SlotInfo = ViewportWidgets.FindOrAdd(TempObject);
+		TObjectKey<UWidget> WidgetKey = Widget;
+		FSlotInfo& SlotInfo = ViewportWidgets.FindOrAdd(WidgetKey);
 		Widget->bIsManagedByGameViewportSubsystem = true;
 		SlotInfo.Slot = Slot;
 		if (TSharedPtr<SConstraintCanvas> WidgetHost = SlotInfo.FullScreenWidget.Pin())
@@ -256,13 +260,12 @@ void UGameViewportSubsystem::SetWidgetSlot(UWidget* Widget, FGameViewportWidgetS
 			SlotInfo.FullScreenWidgetSlot->SetAutoSize(OffsetArgument.Get<1>());
 			SlotInfo.FullScreenWidgetSlot->SetAnchors(Slot.Anchors);
 			SlotInfo.FullScreenWidgetSlot->SetAlignment(Slot.Alignment);
+			SlotInfo.FullScreenWidgetSlot->SetZOrder(Slot.ZOrder);
 			WidgetHost->Invalidate(EInvalidateWidgetReason::Layout);
-			//todo set the zorder
 		}
 	}
 }
 
-
 FGameViewportWidgetSlot UGameViewportSubsystem::SetWidgetSlotPosition(FGameViewportWidgetSlot Slot, const UWidget* Widget, FVector2D Position, bool bRemoveDPIScale)
 {
 	if (bRemoveDPIScale)
@@ -296,9 +299,9 @@ void UGameViewportSubsystem::HandleRemoveWorld(UWorld* InWorld)
 	TArray<UWidget*, TInlineAllocator<16>> WidgetsToRemove;
 	for (FViewportWidgetList::TIterator Itt = ViewportWidgets.CreateIterator(); Itt; ++Itt)
 	{
-		if (UWidget* Widget = Itt.Key().Get())
+		if (UWidget* Widget = Itt.Key().ResolveObjectPtr())
 		{
-			if (InWorld == Widget->GetWorld())
+			if (Itt.Value().Slot.bAutoRemoveOnWorldRemoved && InWorld == Widget->GetWorld())
 			{
 				WidgetsToRemove.Add(Widget);
 			}
@@ -315,4 +318,13 @@ void UGameViewportSubsystem::HandleRemoveWorld(UWorld* InWorld)
 	}
 }
 
-
+void UGameViewportSubsystem::HandleLevelRemovedFromWorld(ULevel* InLevel, UWorld* InWorld)
+{
+	// If the InLevel is null, it's a signal that the entire world is about to disappear, so
+	// go ahead and remove this widget from the viewport, it could be holding onto too many
+	// dangerous actor references that won't carry over into the next world.
+	if (InLevel == nullptr)
+	{
+		HandleRemoveWorld(InWorld);
+	}
+}
diff --git a/Engine/Source/Runtime/UMG/Private/UserWidget.cpp b/Engine/Source/Runtime/UMG/Private/UserWidget.cpp
--- a/Engine/Source/Runtime/UMG/Private/UserWidget.cpp
+++ b/Engine/Source/Runtime/UMG/Private/UserWidget.cpp
@@ -1063,6 +1063,7 @@ void UUserWidget::SetVisibility(ESlateVisibility InVisibility)
 void UUserWidget::SetPlayerContext(const FLocalPlayerContext& InPlayerContext)
 {
 	PlayerContext = InPlayerContext;
+	CachedWorld.Reset();
 
 	if (WidgetTree)
 	{
@@ -1097,6 +1098,7 @@ void UUserWidget::SetOwningLocalPlayer(ULocalPlayer* LocalPlayer)
 	if ( LocalPlayer )
 	{
 		PlayerContext = FLocalPlayerContext(LocalPlayer, GetWorld());
+		CachedWorld.Reset();
 	}
 }
 
@@ -1110,6 +1112,7 @@ void UUserWidget::SetOwningPlayer(APlayerController* LocalPlayerController)
 	if (LocalPlayerController && LocalPlayerController->IsLocalController())
 	{
 		PlayerContext = FLocalPlayerContext(LocalPlayerController);
+		CachedWorld.Reset();
 	}
 }
 
diff --git a/Engine/Source/Runtime/UMG/Public/Blueprint/GameViewportSubsystem.h b/Engine/Source/Runtime/UMG/Public/Blueprint/GameViewportSubsystem.h
--- a/Engine/Source/Runtime/UMG/Public/Blueprint/GameViewportSubsystem.h
+++ b/Engine/Source/Runtime/UMG/Public/Blueprint/GameViewportSubsystem.h
@@ -4,6 +4,7 @@
 
 #include "Subsystems/EngineSubsystem.h"
 #include "Layout/Margin.h"
+#include "UObject/ObjectKey.h"
 #include "Widgets/Layout/Anchors.h"
 #include "Widgets/Layout/SConstraintCanvas.h"
 #include "GameViewportSubsystem.generated.h"
@@ -34,6 +35,13 @@ struct FGameViewportWidgetSlot
 	/** The higher the number, the more on top this widget will be. */
 	UPROPERTY(EditAnywhere, BlueprintReadWrite, Category = "User Interface")
 	int32 ZOrder = 0;
+
+	/**
+	 * Remove the widget when the Widget's World is removed.
+	 * @note The Widget is added to the GameViewportClient of the Widget's World. The GameViewportClient can migrate from World to World. The widget can stay visible if the owner of the widget also migrate.
+	 */
+	UPROPERTY(EditAnywhere, BlueprintReadWrite, Category = "User Interface")
+	bool bAutoRemoveOnWorldRemoved = true;
 };
 
 
@@ -121,11 +129,12 @@ private:
 	void RemoveWidgetInternal(UWidget* Widget, const TWeakPtr<SConstraintCanvas>& FullScreenWidget, const TWeakObjectPtr<ULocalPlayer>& LocalPlayer);
 	void HandleWorldCleanup(UWorld* InWorld, bool bSessionEnded, bool bCleanupResoures);
 	void HandleRemoveWorld(UWorld* InWorld);
+	void HandleLevelRemovedFromWorld(ULevel* InLevel, UWorld* InWorld);
 
 	FMargin GetFullScreenOffsetForWidget(UWidget* Widget) const;
 	TPair<FMargin, bool> GetOffsetAttribute(UWidget* Widget, const FGameViewportWidgetSlot& Slot) const;
 
 private:
-	using FViewportWidgetList = TMap<TWeakObjectPtr<UWidget>, FSlotInfo>;
+	using FViewportWidgetList = TMap<TObjectKey<UWidget>, FSlotInfo>;
 	FViewportWidgetList ViewportWidgets;
 };
