Sequencer- Fix regression where new Widget Material Parameter tracks were failing to work. Changes to the way component material parameters were handled by a new type of section broke this case.

diff --git a/Engine/Source/Runtime/MovieSceneTracks/Private/Tracks/MovieSceneMaterialTrack.cpp b/Engine/Source/Runtime/MovieSceneTracks/Private/Tracks/MovieSceneMaterialTrack.cpp
--- a/Engine/Source/Runtime/MovieSceneTracks/Private/Tracks/MovieSceneMaterialTrack.cpp
+++ b/Engine/Source/Runtime/MovieSceneTracks/Private/Tracks/MovieSceneMaterialTrack.cpp
@@ -143,7 +143,7 @@ void UMovieSceneMaterialTrack::AddScalarParameterKey(const FMaterialParameterInf
 	}
 	if (NearestSection == nullptr)
 	{
-		NearestSection = Cast<UMovieSceneComponentMaterialParameterSection>(CreateNewSection());
+		NearestSection = CreateNewSection();
 
 		UMovieScene* MovieScene = GetTypedOuter<UMovieScene>();
 		check(MovieScene);
@@ -191,7 +191,7 @@ void UMovieSceneMaterialTrack::AddColorParameterKey(const FMaterialParameterInfo
 	}
 	if (NearestSection == nullptr)
 	{
-		NearestSection = Cast<UMovieSceneComponentMaterialParameterSection>(CreateNewSection());
+		NearestSection = CreateNewSection();
 
 		UMovieScene* MovieScene = GetTypedOuter<UMovieScene>();
 		check(MovieScene);
diff --git a/Engine/Source/Runtime/UMG/Private/Animation/MovieSceneWidgetMaterialTrack.cpp b/Engine/Source/Runtime/UMG/Private/Animation/MovieSceneWidgetMaterialTrack.cpp
--- a/Engine/Source/Runtime/UMG/Private/Animation/MovieSceneWidgetMaterialTrack.cpp
+++ b/Engine/Source/Runtime/UMG/Private/Animation/MovieSceneWidgetMaterialTrack.cpp
@@ -10,6 +10,20 @@
 UMovieSceneWidgetMaterialTrack::UMovieSceneWidgetMaterialTrack( const FObjectInitializer& ObjectInitializer )
 	: Super(ObjectInitializer)
 {
+
+}
+
+UMovieSceneSection* UMovieSceneWidgetMaterialTrack::CreateNewSection()
+{
+	UMovieSceneSection* NewSection = NewObject<UMovieSceneParameterSection>(this, NAME_None, RF_Transactional);
+	NewSection->SetBlendType(EMovieSceneBlendType::Absolute);
+	NewSection->SetRange(TRange<FFrameNumber>::All());
+	return NewSection;
+}
+
+bool UMovieSceneWidgetMaterialTrack::SupportsType(TSubclassOf<UMovieSceneSection> SectionClass) const
+{
+	return SectionClass == UMovieSceneParameterSection::StaticClass();
 }
 
 void UMovieSceneWidgetMaterialTrack::ImportEntityImpl(UMovieSceneEntitySystemLinker* EntityLinker, const FEntityImportParams& Params, FImportedEntity* OutImportedEntity)
diff --git a/Engine/Source/Runtime/UMG/Public/Animation/MovieSceneWidgetMaterialTrack.h b/Engine/Source/Runtime/UMG/Public/Animation/MovieSceneWidgetMaterialTrack.h
--- a/Engine/Source/Runtime/UMG/Public/Animation/MovieSceneWidgetMaterialTrack.h
+++ b/Engine/Source/Runtime/UMG/Public/Animation/MovieSceneWidgetMaterialTrack.h
@@ -37,7 +37,10 @@ class UMovieSceneWidgetMaterialTrack
 public:
 
 	// UMovieSceneTrack interface
-	virtual FName GetTrackName() const override;
+	virtual FName GetTrackName() const override;	
+	
+	virtual bool SupportsType(TSubclassOf<UMovieSceneSection> SectionClass) const override;
+	virtual UMovieSceneSection* CreateNewSection() override;
 
 	/*~ IMovieSceneEntityProvider */
 	virtual void ImportEntityImpl(UMovieSceneEntitySystemLinker* EntityLinker, const FEntityImportParams& Params, FImportedEntity* OutImportedEntity) override;
