UMG: Fix invalid widget pointer in ListView when the row doesn't implement the UserListEntry interface.

diff --git a/Engine/Source/Runtime/UMG/Public/Slate/SObjectTableRow.h b/Engine/Source/Runtime/UMG/Public/Slate/SObjectTableRow.h
--- a/Engine/Source/Runtime/UMG/Public/Slate/SObjectTableRow.h
+++ b/Engine/Source/Runtime/UMG/Public/Slate/SObjectTableRow.h
@@ -122,15 +122,13 @@ public:
 		// Since event-based selection changes are only really needed in a game scenario, we (crudely) monitor it here to generate events.
 		// If desired, per-item selection events could be added as a longer-term todo
 		TSharedPtr<ITypedTableView<ItemType>> OwnerTable = OwnerTablePtr.Pin();
-		if (OwnerTable.IsValid())
+		const TObjectPtrWrapTypeOf<ItemType>* MyItemPtr = OwnerTable ? GetItemForThis(OwnerTable.ToSharedRef()) : nullptr;
+		if (MyItemPtr)
 		{
-			if (const TObjectPtrWrapTypeOf<ItemType>* MyItemPtr = GetItemForThis(OwnerTable.ToSharedRef()))
+			if (bIsAppearingSelected != OwnerTable->Private_IsItemSelected(*MyItemPtr))
 			{
-				if (bIsAppearingSelected != OwnerTable->Private_IsItemSelected(*MyItemPtr))
-				{
-					bIsAppearingSelected = !bIsAppearingSelected;
-					OnItemSelectionChanged(bIsAppearingSelected);
-				}
+				bIsAppearingSelected = !bIsAppearingSelected;
+				OnItemSelectionChanged(bIsAppearingSelected);
 			}
 		}
 
@@ -145,7 +143,7 @@ public:
 		}
 	}
 
-	// ITableRow interface
+	//~ ITableRow interface
 	virtual void InitializeRow() override final
 	{
 		// ObjectRows can be generated in the widget designer with dummy data, which we want to ignore
@@ -163,20 +161,51 @@ public:
 		}
 	}
 
-	virtual TSharedRef<SWidget> AsWidget() override { return SharedThis(this); }
-	virtual void SetIndexInList(int32 InIndexInList) override { IndexInList = InIndexInList; }
-	virtual int32 GetIndexInList() override { return IndexInList; }
-	virtual TSharedPtr<SWidget> GetContent() override { return ChildSlot.GetChildAt(0); }
-	virtual int32 GetIndentLevel() const override { return OwnerTablePtr.Pin()->Private_GetNestingDepth(IndexInList); }
-	virtual int32 DoesItemHaveChildren() const override { return OwnerTablePtr.Pin()->Private_DoesItemHaveChildren(IndexInList); }
-	virtual void Private_OnExpanderArrowShiftClicked() override { /* Intentionally blank - far too specific to be a valid game UI interaction */ }
-	virtual ESelectionMode::Type GetSelectionMode() const override { return OwnerTablePtr.Pin()->Private_GetSelectionMode(); }
-	virtual FVector2D GetRowSizeForColumn(const FName& InColumnName) const override { return FVector2D::ZeroVector; }
+	virtual TSharedRef<SWidget> AsWidget() override
+	{
+		return SharedThis(this);
+	}
+	virtual void SetIndexInList(int32 InIndexInList) override
+	{
+		IndexInList = InIndexInList;
+	}
+	virtual int32 GetIndexInList() override
+	{
+		return IndexInList;
+	}
+	virtual TSharedPtr<SWidget> GetContent() override
+	{
+		return ChildSlot.GetChildAt(0);
+	}
+	virtual int32 GetIndentLevel() const override
+	{
+		TSharedPtr<ITypedTableView<ItemType>> OwnerTable = OwnerTablePtr.Pin();
+		return OwnerTable ? OwnerTable->Private_GetNestingDepth(IndexInList) : 0;
+	}
+	virtual int32 DoesItemHaveChildren() const override
+	{
+		TSharedPtr<ITypedTableView<ItemType>> OwnerTable = OwnerTablePtr.Pin();
+		return OwnerTable ? OwnerTable->Private_DoesItemHaveChildren(IndexInList) : 0;
+	}
+	virtual void Private_OnExpanderArrowShiftClicked() override
+	{
+		/* Intentionally blank - far too specific to be a valid game UI interaction */
+	}
+	virtual ESelectionMode::Type GetSelectionMode() const override
+	{
+		TSharedPtr<ITypedTableView<ItemType>> OwnerTable = OwnerTablePtr.Pin();
+		return OwnerTable ? OwnerTable->Private_GetSelectionMode() : ESelectionMode::None;
+	}
+	virtual FVector2D GetRowSizeForColumn(const FName& InColumnName) const override
+	{
+		return FVector2D::ZeroVector;
+	}
 
 	virtual bool IsItemExpanded() const override
 	{
-		TSharedRef<ITypedTableView<ItemType>> OwnerTable = OwnerTablePtr.Pin().ToSharedRef();
-		if (const TObjectPtrWrapTypeOf<ItemType>* MyItemPtr = GetItemForThis(OwnerTable))
+		TSharedPtr<ITypedTableView<ItemType>> OwnerTable = OwnerTablePtr.Pin();
+		const TObjectPtrWrapTypeOf<ItemType>* MyItemPtr = OwnerTable ? GetItemForThis(OwnerTable.ToSharedRef()) : nullptr;
+		if (MyItemPtr)
 		{
 			return OwnerTable->Private_IsItemExpanded(*MyItemPtr);
 		}
@@ -186,10 +215,10 @@ public:
 
 	virtual void ToggleExpansion() override
 	{
-		TSharedRef<ITypedTableView<ItemType>> OwnerTable = OwnerTablePtr.Pin().ToSharedRef();
-		if (OwnerTable->Private_DoesItemHaveChildren(IndexInList))
+		TSharedPtr<ITypedTableView<ItemType>> OwnerTable = OwnerTablePtr.Pin();
+		if (OwnerTable && OwnerTable->Private_DoesItemHaveChildren(IndexInList))
 		{
-			if (const TObjectPtrWrapTypeOf<ItemType>* MyItemPtr = GetItemForThis(OwnerTable))
+			if (const TObjectPtrWrapTypeOf<ItemType>* MyItemPtr = GetItemForThis(OwnerTable.ToSharedRef()))
 			{
 				OwnerTable->Private_SetItemExpansion(*MyItemPtr, !OwnerTable->Private_IsItemExpanded(*MyItemPtr));
 			}
@@ -198,8 +227,9 @@ public:
 
 	virtual bool IsItemSelected() const override
 	{
-		TSharedRef<ITypedTableView<ItemType>> OwnerTable = OwnerTablePtr.Pin().ToSharedRef();
-		if (const TObjectPtrWrapTypeOf<ItemType>* MyItemPtr = GetItemForThis(OwnerTable))
+		TSharedPtr<ITypedTableView<ItemType>> OwnerTable = OwnerTablePtr.Pin();
+		const TObjectPtrWrapTypeOf<ItemType>* MyItemPtr = OwnerTable ? GetItemForThis(OwnerTable.ToSharedRef()) : nullptr;
+		if (MyItemPtr)
 		{
 			return OwnerTable->Private_IsItemSelected(*MyItemPtr);
 		}
@@ -209,18 +239,23 @@ public:
 
 	virtual TBitArray<> GetWiresNeededByDepth() const override
 	{
-		return OwnerTablePtr.Pin()->Private_GetWiresNeededByDepth(IndexInList);
+		TSharedPtr<ITypedTableView<ItemType>> OwnerTable = OwnerTablePtr.Pin();
+		return OwnerTable ? OwnerTable->Private_GetWiresNeededByDepth(IndexInList) : TBitArray<>();
 	}
 
 	virtual bool IsLastChild() const override
 	{
-		return OwnerTablePtr.Pin()->Private_IsLastChild(IndexInList);
+		TSharedPtr<ITypedTableView<ItemType>> OwnerTable = OwnerTablePtr.Pin();
+		return OwnerTable ? OwnerTable->Private_IsLastChild(IndexInList) : true;
 	}
 
-	// ~ITableRow interface
+	//~ ITableRow interface
 
-	// SWidget interface
-	virtual bool SupportsKeyboardFocus() const override { return true; }
+	//~ SWidget interface
+	virtual bool SupportsKeyboardFocus() const override
+	{
+		return true;
+	}
 
 	virtual void OnMouseEnter(const FGeometry& MyGeometry, const FPointerEvent& MouseEvent) override
 	{
@@ -244,9 +279,9 @@ public:
 	{
 		if (InMouseEvent.GetEffectingButton() == EKeys::LeftMouseButton)
 		{
-			TSharedRef<ITypedTableView<ItemType>> OwnerTable = OwnerTablePtr.Pin().ToSharedRef();
-
-			if (const TObjectPtrWrapTypeOf<ItemType>* MyItemPtr = GetItemForThis(OwnerTable))
+			TSharedPtr<ITypedTableView<ItemType>> OwnerTable = OwnerTablePtr.Pin();
+			const TObjectPtrWrapTypeOf<ItemType>* MyItemPtr = OwnerTable ? GetItemForThis(OwnerTable.ToSharedRef()) : nullptr;
+			if (MyItemPtr)
 			{
 				OwnerTable->Private_OnItemDoubleClicked(*MyItemPtr);
 				return FReply::Handled();
@@ -271,8 +306,9 @@ public:
 		if (bProcessingSelectionTouch)
 		{
 			bProcessingSelectionTouch = false;
-			TSharedRef<ITypedTableView<ItemType>> OwnerTable = OwnerTablePtr.Pin().ToSharedRef();
-			if (const TObjectPtrWrapTypeOf<ItemType>* MyItemPtr = GetItemForThis(OwnerTable))
+			TSharedPtr<ITypedTableView<ItemType>> OwnerTable = OwnerTablePtr.Pin();
+			const TObjectPtrWrapTypeOf<ItemType>* MyItemPtr = OwnerTable ? GetItemForThis(OwnerTable.ToSharedRef()) : nullptr;
+			if (MyItemPtr)
 			{
 				if (IsItemSelectable())
 				{
@@ -315,16 +351,20 @@ public:
 	{
 		if (bAllowDragging)
 		{
+			TSharedPtr<ITypedTableView<ItemType>> OwnerTable = OwnerTablePtr.Pin();
+
 			if (bProcessingSelectionTouch)
 			{
 				bProcessingSelectionTouch = false;
-				return FReply::Handled().CaptureMouse(OwnerTablePtr.Pin()->AsWidget());
+				return OwnerTable ? FReply::Handled().CaptureMouse(OwnerTable->AsWidget()) : FReply::Unhandled();
 			}
 			//@todo DanH TableRow: does this potentially trigger twice? If not, why does an unhandled drag detection result in not calling mouse up?
 			else if (HasMouseCapture() && bChangedSelectionOnMouseDown)
 			{
-				TSharedRef<ITypedTableView<ItemType>> OwnerTable = OwnerTablePtr.Pin().ToSharedRef();
-				OwnerTable->Private_SignalSelectionChanged(ESelectInfo::OnMouseClick);
+				if (OwnerTable)
+				{
+					OwnerTable->Private_SignalSelectionChanged(ESelectInfo::OnMouseClick);
+				}
 			}
 
 			return SObjectWidget::OnDragDetected(MyGeometry, MouseEvent);
@@ -341,37 +381,39 @@ public:
 		FReply Reply = SObjectWidget::OnMouseButtonDown(MyGeometry, MouseEvent);
 		if (!Reply.IsEventHandled())
 		{
-			TSharedRef<ITypedTableView<ItemType>> OwnerTable = OwnerTablePtr.Pin().ToSharedRef();
-
-			const ESelectionMode::Type SelectionMode = GetSelectionMode();
-			if (MouseEvent.GetEffectingButton() == EKeys::LeftMouseButton && SelectionMode != ESelectionMode::None)
+			TSharedPtr<ITypedTableView<ItemType>> OwnerTable = OwnerTablePtr.Pin();
+			if (OwnerTable)
 			{
-				if (IsItemSelectable())
+				const ESelectionMode::Type SelectionMode = GetSelectionMode();
+				if (MouseEvent.GetEffectingButton() == EKeys::LeftMouseButton && SelectionMode != ESelectionMode::None)
 				{
-					const TObjectPtrWrapTypeOf<ItemType>* MyItemPtr = GetItemForThis(OwnerTable);
-					// New selections are handled on mouse down, deselection is handled on mouse up
-					if (MyItemPtr)
+					if (IsItemSelectable())
 					{
-						const ItemType& MyItem = *MyItemPtr;
-						if (!OwnerTable->Private_IsItemSelected(MyItem))
+						const TObjectPtrWrapTypeOf<ItemType>* MyItemPtr = GetItemForThis(OwnerTable.ToSharedRef());
+						// New selections are handled on mouse down, deselection is handled on mouse up
+						if (MyItemPtr)
 						{
-							if (SelectionMode != ESelectionMode::Multi)
+							const ItemType& MyItem = *MyItemPtr;
+							if (!OwnerTable->Private_IsItemSelected(MyItem))
 							{
-								OwnerTable->Private_ClearSelection();
+								if (SelectionMode != ESelectionMode::Multi)
+								{
+									OwnerTable->Private_ClearSelection();
+								}
+								OwnerTable->Private_SetItemSelection(MyItem, true, true);
+								bChangedSelectionOnMouseDown = true;
 							}
-							OwnerTable->Private_SetItemSelection(MyItem, true, true);
-							bChangedSelectionOnMouseDown = true;
 						}
-					}
 
-					Reply = FReply::Handled()
-						.DetectDrag(SharedThis(this), EKeys::LeftMouseButton)
-						.CaptureMouse(SharedThis(this));
+						Reply = FReply::Handled()
+							.DetectDrag(SharedThis(this), EKeys::LeftMouseButton)
+							.CaptureMouse(SharedThis(this));
 
-					// Set focus back to the owning widget if the item is invalid somehow or its not selectable or can be navigated to
-					if (!MyItemPtr || !OwnerTable->Private_IsItemSelectableOrNavigable(*MyItemPtr))
-					{
-						Reply.SetUserFocus(OwnerTable->AsWidget(), EFocusCause::Mouse);
+						// Set focus back to the owning widget if the item is invalid somehow or its not selectable or can be navigated to
+						if (!MyItemPtr || !OwnerTable->Private_IsItemSelectableOrNavigable(*MyItemPtr))
+						{
+							Reply.SetUserFocus(OwnerTable->AsWidget(), EFocusCause::Mouse);
+						}
 					}
 				}
 			}
@@ -384,11 +426,11 @@ public:
 	{
 		FReply Reply = SObjectWidget::OnMouseButtonUp(MyGeometry, MouseEvent);
 
-		TSharedRef<ITypedTableView<ItemType>> OwnerTable = OwnerTablePtr.Pin().ToSharedRef();
-
 		if (!Reply.IsEventHandled())
 		{
-			if (const TObjectPtrWrapTypeOf<ItemType>* MyItemPtr = GetItemForThis(OwnerTable))
+			TSharedPtr<ITypedTableView<ItemType>> OwnerTable = OwnerTablePtr.Pin();
+			const TObjectPtrWrapTypeOf<ItemType>* MyItemPtr = OwnerTable ? GetItemForThis(OwnerTable.ToSharedRef()) : nullptr;
+			if (MyItemPtr)
 			{
 				const ESelectionMode::Type SelectionMode = GetSelectionMode();
 				if (MouseEvent.GetEffectingButton() == EKeys::LeftMouseButton && HasMouseCapture())
@@ -429,7 +471,7 @@ public:
 				else if (SelectionMode != ESelectionMode::None && MouseEvent.GetEffectingButton() == EKeys::RightMouseButton)
 				{
 					// Ignore the right click release if it was being used for scrolling
-					TSharedRef<STableViewBase> OwnerTableViewBase = StaticCastSharedRef<SListView<ItemType>>(OwnerTable);
+					TSharedPtr<STableViewBase> OwnerTableViewBase = StaticCastSharedPtr<SListView<ItemType>>(OwnerTable);
 					if (!OwnerTableViewBase->IsRightClickScrolling())
 					{
 						if (IsItemSelectable() && !OwnerTable->Private_IsItemSelected(*MyItemPtr))
@@ -455,8 +497,9 @@ public:
 protected:
 	virtual void InitializeObjectRow()
 	{
-		TSharedRef<ITypedTableView<ItemType>> OwnerTable = OwnerTablePtr.Pin().ToSharedRef();
-		if (const TObjectPtrWrapTypeOf<ItemType>* MyItemPtr = GetItemForThis(OwnerTable))
+		TSharedPtr<ITypedTableView<ItemType>> OwnerTable = OwnerTablePtr.Pin();
+		const TObjectPtrWrapTypeOf<ItemType>* MyItemPtr = OwnerTable ? GetItemForThis(OwnerTable.ToSharedRef()) : nullptr;
+		if (MyItemPtr)
 		{
 			InitObjectRowInternal(*WidgetObject, *MyItemPtr);
 
