Replace Cast by compile time check
#rnx

diff --git a/Engine/Source/Runtime/UMG/Public/Blueprint/WidgetTree.h b/Engine/Source/Runtime/UMG/Public/Blueprint/WidgetTree.h
--- a/Engine/Source/Runtime/UMG/Public/Blueprint/WidgetTree.h
+++ b/Engine/Source/Runtime/UMG/Public/Blueprint/WidgetTree.h
@@ -90,14 +90,15 @@ public:
 	template <typename WidgetT>
 	FORCEINLINE_DEBUGGABLE WidgetT* ConstructWidget(TSubclassOf<UWidget> WidgetClass = WidgetT::StaticClass(), FName WidgetName = NAME_None)
 	{
-		static_assert(TIsDerivedFrom<WidgetT, UWidget>::IsDerived, "WidgetTree::ConstructWidget can only create UWidget objects.");
-
-		if (WidgetClass->IsChildOf<UUserWidget>())
+		if constexpr(std::is_base_of_v<UUserWidget, WidgetT>)
 		{
 			return Cast<WidgetT>(CreateWidget(this, *WidgetClass, WidgetName));
 		}
-
-		return NewObject<WidgetT>(this, WidgetClass, WidgetName, RF_Transactional);
+		else
+		{
+			static_assert(std::is_base_of_v<UWidget, WidgetT>, "WidgetTree::ConstructWidget can only create UWidget objects.");
+			return NewObject<WidgetT>(this, WidgetClass, WidgetName, RF_Transactional);
+		}
 	}
 
 	// INamedSlotInterface
