Small change to ensure that we have the same behavior than before, even with bad data
#rnx

diff --git a/Engine/Source/Runtime/UMG/Public/Blueprint/WidgetTree.h b/Engine/Source/Runtime/UMG/Public/Blueprint/WidgetTree.h
--- a/Engine/Source/Runtime/UMG/Public/Blueprint/WidgetTree.h
+++ b/Engine/Source/Runtime/UMG/Public/Blueprint/WidgetTree.h
@@ -88,17 +88,22 @@ public:
 
 	/** Constructs the widget, and adds it to the tree. */
 	template <typename WidgetT>
-	FORCEINLINE_DEBUGGABLE WidgetT* ConstructWidget(TSubclassOf<UWidget> WidgetClass = WidgetT::StaticClass(), FName WidgetName = NAME_None)
+	FORCEINLINE_DEBUGGABLE [[nodiscard]] WidgetT* ConstructWidget(TSubclassOf<WidgetT> WidgetClass = WidgetT::StaticClass(), FName WidgetName = NAME_None)
 	{
-		if constexpr(std::is_base_of_v<UUserWidget, WidgetT>)
+		if(WidgetClass != nullptr)
 		{
-			return Cast<WidgetT>(CreateWidget(this, *WidgetClass, WidgetName));
-		}
-		else
-		{
-			static_assert(std::is_base_of_v<UWidget, WidgetT>, "WidgetTree::ConstructWidget can only create UWidget objects.");
-			return NewObject<WidgetT>(this, WidgetClass, WidgetName, RF_Transactional);
+			if constexpr(std::is_base_of_v<UUserWidget, WidgetT>)
+			{
+				return CreateWidget<WidgetT>(this, *WidgetClass, WidgetName);
+			}
+			else
+			{
+				static_assert(std::is_base_of_v<UWidget, WidgetT>, "WidgetTree::ConstructWidget can only create UWidget objects.");
+				return NewObject<WidgetT>(this, WidgetClass, WidgetName, RF_Transactional);
+			}
 		}
+
+		return nullptr;
 	}
 
 	// INamedSlotInterface
