Fix build on Clang
#rnx

diff --git a/Engine/Source/Runtime/UMG/Public/Blueprint/WidgetTree.h b/Engine/Source/Runtime/UMG/Public/Blueprint/WidgetTree.h
--- a/Engine/Source/Runtime/UMG/Public/Blueprint/WidgetTree.h
+++ b/Engine/Source/Runtime/UMG/Public/Blueprint/WidgetTree.h
@@ -88,7 +88,7 @@ public:
 
 	/** Constructs the widget, and adds it to the tree. */
 	template <typename WidgetT>
-	FORCEINLINE_DEBUGGABLE [[nodiscard]] WidgetT* ConstructWidget(TSubclassOf<WidgetT> WidgetClass = WidgetT::StaticClass(), FName WidgetName = NAME_None)
+	[[nodiscard]] FORCEINLINE_DEBUGGABLE WidgetT* ConstructWidget(TSubclassOf<WidgetT> WidgetClass = WidgetT::StaticClass(), FName WidgetName = NAME_None)
 	{
 		if(WidgetClass != nullptr)
 		{
