Fixed expensive sub-class iteration in UWidget::FindBinderClassForDestination. Putting an much faster iteration method in (4.6ms vs. 7.2us) as a hold-over before all of this gets deprecated in the future.

diff --git a/Engine/Source/Runtime/UMG/Private/Components/Widget.cpp b/Engine/Source/Runtime/UMG/Private/Components/Widget.cpp
--- a/Engine/Source/Runtime/UMG/Private/Components/Widget.cpp
+++ b/Engine/Source/Runtime/UMG/Private/Components/Widget.cpp
@@ -1715,22 +1715,22 @@ FString UWidget::GetDefaultFontName()
 
 TSubclassOf<UPropertyBinding> UWidget::FindBinderClassForDestination(FProperty* Property)
 {
-	if ( BinderClasses.Num() == 0 )
+	if (BinderClasses.IsEmpty())
 	{
-		for ( TObjectIterator<UClass> ClassIt; ClassIt; ++ClassIt )
+		TArray<UClass*> PropertyBindingClasses;
+		GetDerivedClasses(UPropertyBinding::StaticClass(), PropertyBindingClasses);
+		BinderClasses.Reserve(PropertyBindingClasses.Num());
+		for (UClass* PropertyBindingClass : PropertyBindingClasses)
 		{
-			if ( ClassIt->IsChildOf(UPropertyBinding::StaticClass()) )
-			{
-				BinderClasses.Add(*ClassIt);
-			}
+			BinderClasses.Emplace(PropertyBindingClass);
 		}
 	}
 
-	for ( int32 ClassIndex = 0; ClassIndex < BinderClasses.Num(); ClassIndex++ )
+	for (TSubclassOf<UPropertyBinding>& BinderClass : BinderClasses)
 	{
-		if ( GetDefault<UPropertyBinding>(BinderClasses[ClassIndex])->IsSupportedDestination(Property))
+		if (BinderClass.GetDefaultObject()->IsSupportedDestination(Property))
 		{
-			return BinderClasses[ClassIndex];
+			return BinderClass;
 		}
 	}
 
