Curve Editor: Add editor only slate arg to disable warning on size

#rnx

diff --git a/Engine/Source/Editor/CurveEditor/Private/SCurveEditorViewContainer.cpp b/Engine/Source/Editor/CurveEditor/Private/SCurveEditorViewContainer.cpp
--- a/Engine/Source/Editor/CurveEditor/Private/SCurveEditorViewContainer.cpp
+++ b/Engine/Source/Editor/CurveEditor/Private/SCurveEditorViewContainer.cpp
@@ -424,6 +424,7 @@ void SCurveEditorViewContainer::AddView(TSharedRef<SCurveEditorView> ViewToAdd)
 		.RenderOnPhase(false)
 		.RenderOnInvalidation(false)
 		.RenderWithLocalTransform(false)
+		.bWarnOnInvalidSize(false)
 		[
 			SNew(SBox)
 			.Padding(MakeAttributeSP(this, &SCurveEditorViewContainer::GetSlotPadding, InsertIndex))
diff --git a/Engine/Source/Runtime/UMG/Private/Slate/SRetainerWidget.cpp b/Engine/Source/Runtime/UMG/Private/Slate/SRetainerWidget.cpp
--- a/Engine/Source/Runtime/UMG/Private/Slate/SRetainerWidget.cpp
+++ b/Engine/Source/Runtime/UMG/Private/Slate/SRetainerWidget.cpp
@@ -217,6 +217,7 @@ void SRetainerWidget::Construct(const FArguments& InArgs)
 #if WITH_EDITOR
 	bIsDesignTime = false;
 	bShowEffectsInDesigner = true;
+	bWarnOnInvalidSize = InArgs._bWarnOnInvalidSize;
 #endif // WITH_EDITOR
 
 	RefreshRenderingMode();
@@ -502,13 +503,26 @@ SRetainerWidget::EPaintRetainedContentResult SRetainerWidget::PaintRetainedConte
 			if (!bInvalidSizeLogged)
 			{
 				bInvalidSizeLogged = true;
+
+				const bool bEnableWarnOnInvalidSize =
+#if WITH_EDITOR
+					bWarnOnInvalidSize;
+#else
+					true;
+#endif
 				if (bTextureIsTooLarge)
 				{
-					UE_LOG(LogUMG, Warning, TEXT("The requested size for SRetainerWidget is too large. W:%i H:%i"), RenderTargetWidth, RenderTargetHeight);
+					if (bEnableWarnOnInvalidSize)
+					{
+						UE_LOG(LogUMG, Warning, TEXT("The requested size for SRetainerWidget is too large. W:%i H:%i"), RenderTargetWidth, RenderTargetHeight);
+					}
 				}
 				else
 				{
-					UE_LOG(LogUMG, Warning, TEXT("The requested size for SRetainerWidget is 0. W:%i H:%i"), RenderTargetWidth, RenderTargetHeight);
+					if (bEnableWarnOnInvalidSize)
+					{
+						UE_LOG(LogUMG, Warning, TEXT("The requested size for SRetainerWidget is 0. W:%i H:%i"), RenderTargetWidth, RenderTargetHeight);
+					}
 				}
 			}
 			return bTextureIsTooLarge ? EPaintRetainedContentResult::TextureSizeTooBig : EPaintRetainedContentResult::TextureSizeZero;
diff --git a/Engine/Source/Runtime/UMG/Public/Slate/SRetainerWidget.h b/Engine/Source/Runtime/UMG/Public/Slate/SRetainerWidget.h
--- a/Engine/Source/Runtime/UMG/Public/Slate/SRetainerWidget.h
+++ b/Engine/Source/Runtime/UMG/Public/Slate/SRetainerWidget.h
@@ -58,7 +58,10 @@ public:
 		SLATE_ARGUMENT(int32, Phase)
 		SLATE_ARGUMENT(int32, PhaseCount)
 		SLATE_ARGUMENT(FName, StatId)
-	SLATE_END_ARGS()
+#if WITH_EDITOR
+		SLATE_ARGUMENT(bool, bWarnOnInvalidSize)
+#endif
+		SLATE_END_ARGS()
 
 	UMG_API SRetainerWidget();
 	UMG_API ~SRetainerWidget();
@@ -159,6 +162,9 @@ private:
 
 	/** True if we should retain rendering in designer */
 	bool bShowEffectsInDesigner;
+
+	/** True if we should warn when the requested size for the retainer is 0 or too large */
+	bool bWarnOnInvalidSize;
 #endif
 
 	bool RenderOnPhase;
