UMG: Add setter for LineHeightPercentage property for UTextLayoutWidget

diff --git a/Engine/Plugins/Runtime/CommonUI/Source/CommonUI/Private/CommonTextBlock.cpp b/Engine/Plugins/Runtime/CommonUI/Source/CommonUI/Private/CommonTextBlock.cpp
--- a/Engine/Plugins/Runtime/CommonUI/Source/CommonUI/Private/CommonTextBlock.cpp
+++ b/Engine/Plugins/Runtime/CommonUI/Source/CommonUI/Private/CommonTextBlock.cpp
@@ -348,8 +348,7 @@ void UCommonTextBlock::SetTextCase(bool bUseAllCaps)
 
 void UCommonTextBlock::SetLineHeightPercentage(float InLineHeightPercentage)
 {
-	LineHeightPercentage = InLineHeightPercentage;
-	SynchronizeProperties();
+	UTextLayoutWidget::SetLineHeightPercentage(InLineHeightPercentage);
 }
 
 void UCommonTextBlock::SetStyle(TSubclassOf<UCommonTextStyle> InStyle)
diff --git a/Engine/Source/Runtime/UMG/Private/Components/MultiLineEditableText.cpp b/Engine/Source/Runtime/UMG/Private/Components/MultiLineEditableText.cpp
--- a/Engine/Source/Runtime/UMG/Private/Components/MultiLineEditableText.cpp
+++ b/Engine/Source/Runtime/UMG/Private/Components/MultiLineEditableText.cpp
@@ -124,6 +124,15 @@ void UMultiLineEditableText::SetJustification(ETextJustify::Type InJustification
 	}
 }
 
+void UMultiLineEditableText::OnLineHeightPercentageChanged(float InLineHeightPercentage)
+{
+	Super::OnLineHeightPercentageChanged(InLineHeightPercentage);
+	if (MyMultiLineEditableText.IsValid())
+	{
+		MyMultiLineEditableText->SetLineHeightPercentage(InLineHeightPercentage);
+	}
+}
+
 FText UMultiLineEditableText::GetText() const
 {
 	if ( MyMultiLineEditableText.IsValid() )
diff --git a/Engine/Source/Runtime/UMG/Private/Components/MultiLineEditableTextBox.cpp b/Engine/Source/Runtime/UMG/Private/Components/MultiLineEditableTextBox.cpp
--- a/Engine/Source/Runtime/UMG/Private/Components/MultiLineEditableTextBox.cpp
+++ b/Engine/Source/Runtime/UMG/Private/Components/MultiLineEditableTextBox.cpp
@@ -135,6 +135,15 @@ void UMultiLineEditableTextBox::SetJustification(ETextJustify::Type InJustificat
 	}
 }
 
+void UMultiLineEditableTextBox::OnLineHeightPercentageChanged(float InLineHeightPercentage)
+{
+	Super::OnLineHeightPercentageChanged(InLineHeightPercentage);
+	if (MyEditableTextBlock.IsValid())
+	{
+		MyEditableTextBlock->SetLineHeightPercentage(InLineHeightPercentage);
+	}
+}
+
 FText UMultiLineEditableTextBox::GetText() const
 {
 	if ( MyEditableTextBlock.IsValid() )
diff --git a/Engine/Source/Runtime/UMG/Private/Components/RichTextBlock.cpp b/Engine/Source/Runtime/UMG/Private/Components/RichTextBlock.cpp
--- a/Engine/Source/Runtime/UMG/Private/Components/RichTextBlock.cpp
+++ b/Engine/Source/Runtime/UMG/Private/Components/RichTextBlock.cpp
@@ -403,6 +403,16 @@ void URichTextBlock::SetJustification(ETextJustify::Type InJustification)
 }
 
 PRAGMA_DISABLE_DEPRECATION_WARNINGS
+
+void URichTextBlock::OnLineHeightPercentageChanged(float InLineHeightPercentage)
+{
+	Super::OnLineHeightPercentageChanged(InLineHeightPercentage);
+	if (MyRichTextBlock.IsValid())
+	{
+		MyRichTextBlock->SetLineHeightPercentage(InLineHeightPercentage);
+	}
+}
+
 void URichTextBlock::SetMinDesiredWidth(float InMinDesiredWidth)
 {
 	MinDesiredWidth = InMinDesiredWidth;
diff --git a/Engine/Source/Runtime/UMG/Private/Components/TextBlock.cpp b/Engine/Source/Runtime/UMG/Private/Components/TextBlock.cpp
--- a/Engine/Source/Runtime/UMG/Private/Components/TextBlock.cpp
+++ b/Engine/Source/Runtime/UMG/Private/Components/TextBlock.cpp
@@ -152,6 +152,15 @@ void UTextBlock::SetJustification(ETextJustify::Type InJustification)
 	}
 }
 
+void UTextBlock::OnLineHeightPercentageChanged(float InLineHeightPercentage)
+{
+	Super::OnLineHeightPercentageChanged(InLineHeightPercentage);
+	if (MyTextBlock.IsValid())
+	{
+		MyTextBlock->SetLineHeightPercentage(InLineHeightPercentage);
+	}
+}
+
 float UTextBlock::GetMinDesiredWidth() const
 {
 	return MinDesiredWidth;
diff --git a/Engine/Source/Runtime/UMG/Private/Components/TextWidgetTypes.cpp b/Engine/Source/Runtime/UMG/Private/Components/TextWidgetTypes.cpp
--- a/Engine/Source/Runtime/UMG/Private/Components/TextWidgetTypes.cpp
+++ b/Engine/Source/Runtime/UMG/Private/Components/TextWidgetTypes.cpp
@@ -14,7 +14,6 @@ FShapedTextOptions::FShapedTextOptions()
 	TextFlowDirection = ETextFlowDirection::Auto;
 }
 
-
 UTextLayoutWidget::UTextLayoutWidget(const FObjectInitializer& ObjectInitializer)
 	: Super(ObjectInitializer)
 {
@@ -26,3 +25,8 @@ UTextLayoutWidget::UTextLayoutWidget(const FObjectInitializer& ObjectInitializer
 	LineHeightPercentage = 1.0f;
 }
 
+void UTextLayoutWidget::SetLineHeightPercentage(float InLineHeightPercentage)
+{
+	LineHeightPercentage = InLineHeightPercentage;
+	OnLineHeightPercentageChanged(InLineHeightPercentage);
+}
diff --git a/Engine/Source/Runtime/UMG/Public/Components/MultiLineEditableText.h b/Engine/Source/Runtime/UMG/Public/Components/MultiLineEditableText.h
--- a/Engine/Source/Runtime/UMG/Public/Components/MultiLineEditableText.h
+++ b/Engine/Source/Runtime/UMG/Public/Components/MultiLineEditableText.h
@@ -182,6 +182,10 @@ public:
 #endif
 
 protected:
+	//~ Begin UTextLayoutWidget Interface
+	UMG_API virtual void OnLineHeightPercentageChanged(float InLineHeightPercentage) override;
+	//~ End UTextLayoutWidget Interface
+
 	//~ Begin UWidget Interface
 	UMG_API virtual TSharedRef<SWidget> RebuildWidget() override;
 	// End of UWidget
diff --git a/Engine/Source/Runtime/UMG/Public/Components/MultiLineEditableTextBox.h b/Engine/Source/Runtime/UMG/Public/Components/MultiLineEditableTextBox.h
--- a/Engine/Source/Runtime/UMG/Public/Components/MultiLineEditableTextBox.h
+++ b/Engine/Source/Runtime/UMG/Public/Components/MultiLineEditableTextBox.h
@@ -150,6 +150,10 @@ public:
 	UMG_API virtual void Serialize(FArchive& Ar) override;
 
 protected:
+	//~ Begin UTextLayoutWidget Interface
+	UMG_API virtual void OnLineHeightPercentageChanged(float InLineHeightPercentage) override;
+	//~ End UTextLayoutWidget Interface
+
 	//~ Begin UWidget Interface
 	UMG_API virtual TSharedRef<SWidget> RebuildWidget() override;
 	// End of UWidget
diff --git a/Engine/Source/Runtime/UMG/Public/Components/RichTextBlock.h b/Engine/Source/Runtime/UMG/Public/Components/RichTextBlock.h
--- a/Engine/Source/Runtime/UMG/Public/Components/RichTextBlock.h
+++ b/Engine/Source/Runtime/UMG/Public/Components/RichTextBlock.h
@@ -227,6 +227,10 @@ public:
 	UMG_API void RefreshTextLayout();
 
 protected:
+	//~ Begin UTextLayoutWidget Interface
+	UMG_API virtual void OnLineHeightPercentageChanged(float InLineHeightPercentage) override;
+	//~ End UTextLayoutWidget Interface
+
 	UMG_API virtual TSharedRef<SWidget> RebuildWidget() override;
 	
 	UMG_API virtual void UpdateStyleData();
diff --git a/Engine/Source/Runtime/UMG/Public/Components/TextBlock.h b/Engine/Source/Runtime/UMG/Public/Components/TextBlock.h
--- a/Engine/Source/Runtime/UMG/Public/Components/TextBlock.h
+++ b/Engine/Source/Runtime/UMG/Public/Components/TextBlock.h
@@ -256,6 +256,10 @@ public:
 #endif
 
 protected:
+	//~ Begin UTextLayoutWidget Interface
+	UMG_API virtual void OnLineHeightPercentageChanged(float InLineHeightPercentage) override;
+	//~ End UTextLayoutWidget Interface
+
 	/**
 	 * If this is enabled, text shaping, wrapping, justification are disabled in favor of much faster text layout and measurement.
 	 * This feature is only suitable for "simple" text (ie, text containing only numbers or basic ASCII) as it disables the complex text rendering support required for certain languages (such as Arabic and Thai).
diff --git a/Engine/Source/Runtime/UMG/Public/Components/TextWidgetTypes.h b/Engine/Source/Runtime/UMG/Public/Components/TextWidgetTypes.h
--- a/Engine/Source/Runtime/UMG/Public/Components/TextWidgetTypes.h
+++ b/Engine/Source/Runtime/UMG/Public/Components/TextWidgetTypes.h
@@ -65,7 +65,11 @@ public:
 	UFUNCTION(BlueprintSetter)
 	virtual void SetJustification(ETextJustify::Type InJustification) { Justification = InJustification; }
 
+	UMG_API void SetLineHeightPercentage(float InLineHeightPercentage);
+
 protected:
+	virtual void OnLineHeightPercentageChanged(float InLineHeightPercentage) {};
+
 	/** Synchronize the properties with the given widget. A template as the Slate widgets conform to the same API, but don't derive from a common base. */
 	template <typename TWidgetType>
 	void SynchronizeTextLayoutProperties(TWidgetType& InWidget)
@@ -105,6 +109,6 @@ protected:
 	FMargin Margin;
 
 	/** The amount to scale each lines height by. */
-	UPROPERTY(EditAnywhere, BlueprintReadOnly, Category=Appearance, AdvancedDisplay)
+	UPROPERTY(EditAnywhere, BlueprintReadWrite, Setter, Category=Appearance, AdvancedDisplay)
 	float LineHeightPercentage;
 };
