UMG: Add UUserWidget::OnNativeDestruct event to know when it will be destroyed

diff --git a/Engine/Source/Runtime/UMG/Private/UserWidget.cpp b/Engine/Source/Runtime/UMG/Private/UserWidget.cpp
--- a/Engine/Source/Runtime/UMG/Private/UserWidget.cpp
+++ b/Engine/Source/Runtime/UMG/Private/UserWidget.cpp
@@ -1523,6 +1523,8 @@ void UUserWidget::NativeConstruct()
 void UUserWidget::NativeDestruct()
 {
 	StopListeningForAllInputActions();
+	OnNativeDestruct.Broadcast(this);
+
 	Destruct();
 
 	// Extension can remove other extensions.
diff --git a/Engine/Source/Runtime/UMG/Public/Blueprint/UserWidget.h b/Engine/Source/Runtime/UMG/Public/Blueprint/UserWidget.h
--- a/Engine/Source/Runtime/UMG/Public/Blueprint/UserWidget.h
+++ b/Engine/Source/Runtime/UMG/Public/Blueprint/UserWidget.h
@@ -916,6 +916,9 @@ public:
 	DECLARE_EVENT_OneParam(UUserWidget, FNativeOnVisibilityChangedEvent, ESlateVisibility);
 	FNativeOnVisibilityChangedEvent OnNativeVisibilityChanged;
 
+	DECLARE_EVENT_OneParam(UUserWidget, FNativeOnDestruct, UUserWidget*);
+	FNativeOnDestruct OnNativeDestruct;
+
 	UE_DEPRECATED(5.2, "Direct access to Padding is deprecated. Please use the getter or setter.")
 	/** The padding area around the content. */
 	UPROPERTY(EditAnywhere, BlueprintReadWrite, Getter, Setter, BlueprintSetter = "SetPadding", Category = "Appearance")
