UMG: Add setter for Margin  property for UTextLayoutWidget

diff --git a/Engine/Plugins/Runtime/CommonUI/Source/CommonUI/Private/CommonTextBlock.cpp b/Engine/Plugins/Runtime/CommonUI/Source/CommonUI/Private/CommonTextBlock.cpp
--- a/Engine/Plugins/Runtime/CommonUI/Source/CommonUI/Private/CommonTextBlock.cpp
+++ b/Engine/Plugins/Runtime/CommonUI/Source/CommonUI/Private/CommonTextBlock.cpp
@@ -364,8 +364,7 @@ const FMargin& UCommonTextBlock::GetMargin()
 
 void UCommonTextBlock::SetMargin(const FMargin& InMargin)
 {
-	Margin = InMargin;
-	SynchronizeProperties();
+	UTextLayoutWidget::SetMargin(InMargin);
 }
 
 void UCommonTextBlock::ResetScrollState()
diff --git a/Engine/Source/Runtime/UMG/Private/Components/MultiLineEditableText.cpp b/Engine/Source/Runtime/UMG/Private/Components/MultiLineEditableText.cpp
--- a/Engine/Source/Runtime/UMG/Private/Components/MultiLineEditableText.cpp
+++ b/Engine/Source/Runtime/UMG/Private/Components/MultiLineEditableText.cpp
@@ -133,6 +133,15 @@ void UMultiLineEditableText::OnLineHeightPercentageChanged(float InLineHeightPer
 	}
 }
 
+void UMultiLineEditableText::OnMarginChanged(const FMargin& InMargin)
+{
+	Super::OnMarginChanged(InMargin);
+	if (MyMultiLineEditableText.IsValid())
+	{
+		MyMultiLineEditableText->SetMargin(InMargin);
+	}
+}
+
 FText UMultiLineEditableText::GetText() const
 {
 	if ( MyMultiLineEditableText.IsValid() )
diff --git a/Engine/Source/Runtime/UMG/Private/Components/MultiLineEditableTextBox.cpp b/Engine/Source/Runtime/UMG/Private/Components/MultiLineEditableTextBox.cpp
--- a/Engine/Source/Runtime/UMG/Private/Components/MultiLineEditableTextBox.cpp
+++ b/Engine/Source/Runtime/UMG/Private/Components/MultiLineEditableTextBox.cpp
@@ -144,6 +144,15 @@ void UMultiLineEditableTextBox::OnLineHeightPercentageChanged(float InLineHeight
 	}
 }
 
+void UMultiLineEditableTextBox::OnMarginChanged(const FMargin& InMargin)
+{
+	Super::OnMarginChanged(InMargin);
+	if (MyEditableTextBlock.IsValid())
+	{
+		MyEditableTextBlock->SetMargin(InMargin);
+	}
+}
+
 FText UMultiLineEditableTextBox::GetText() const
 {
 	if ( MyEditableTextBlock.IsValid() )
diff --git a/Engine/Source/Runtime/UMG/Private/Components/RichTextBlock.cpp b/Engine/Source/Runtime/UMG/Private/Components/RichTextBlock.cpp
--- a/Engine/Source/Runtime/UMG/Private/Components/RichTextBlock.cpp
+++ b/Engine/Source/Runtime/UMG/Private/Components/RichTextBlock.cpp
@@ -413,6 +413,15 @@ void URichTextBlock::OnLineHeightPercentageChanged(float InLineHeightPercentage)
 	}
 }
 
+void URichTextBlock::OnMarginChanged(const FMargin& InMargin)
+{
+	Super::OnMarginChanged(InMargin);
+	if (MyRichTextBlock.IsValid())
+	{
+		MyRichTextBlock->SetMargin(InMargin);
+	}
+}
+
 void URichTextBlock::SetMinDesiredWidth(float InMinDesiredWidth)
 {
 	MinDesiredWidth = InMinDesiredWidth;
diff --git a/Engine/Source/Runtime/UMG/Private/Components/TextBlock.cpp b/Engine/Source/Runtime/UMG/Private/Components/TextBlock.cpp
--- a/Engine/Source/Runtime/UMG/Private/Components/TextBlock.cpp
+++ b/Engine/Source/Runtime/UMG/Private/Components/TextBlock.cpp
@@ -161,6 +161,15 @@ void UTextBlock::OnLineHeightPercentageChanged(float InLineHeightPercentage)
 	}
 }
 
+void UTextBlock::OnMarginChanged(const FMargin& InMargin)
+{
+	Super::OnMarginChanged(InMargin);
+	if (MyTextBlock.IsValid())
+	{
+		MyTextBlock->SetMargin(InMargin);
+	}
+}
+
 float UTextBlock::GetMinDesiredWidth() const
 {
 	return MinDesiredWidth;
diff --git a/Engine/Source/Runtime/UMG/Private/Components/TextWidgetTypes.cpp b/Engine/Source/Runtime/UMG/Private/Components/TextWidgetTypes.cpp
--- a/Engine/Source/Runtime/UMG/Private/Components/TextWidgetTypes.cpp
+++ b/Engine/Source/Runtime/UMG/Private/Components/TextWidgetTypes.cpp
@@ -27,6 +27,18 @@ UTextLayoutWidget::UTextLayoutWidget(const FObjectInitializer& ObjectInitializer
 
 void UTextLayoutWidget::SetLineHeightPercentage(float InLineHeightPercentage)
 {
-	LineHeightPercentage = InLineHeightPercentage;
-	OnLineHeightPercentageChanged(InLineHeightPercentage);
+	if (LineHeightPercentage != InLineHeightPercentage)
+	{
+		LineHeightPercentage = InLineHeightPercentage;
+		OnLineHeightPercentageChanged(InLineHeightPercentage);
+	}
+}
+
+void UTextLayoutWidget::SetMargin(const FMargin& InMargin)
+{
+	if (Margin != InMargin)
+	{
+		Margin = InMargin;
+		OnMarginChanged(Margin);
+	}
 }
diff --git a/Engine/Source/Runtime/UMG/Public/Components/MultiLineEditableText.h b/Engine/Source/Runtime/UMG/Public/Components/MultiLineEditableText.h
--- a/Engine/Source/Runtime/UMG/Public/Components/MultiLineEditableText.h
+++ b/Engine/Source/Runtime/UMG/Public/Components/MultiLineEditableText.h
@@ -184,6 +184,7 @@ public:
 protected:
 	//~ Begin UTextLayoutWidget Interface
 	UMG_API virtual void OnLineHeightPercentageChanged(float InLineHeightPercentage) override;
+	UMG_API virtual void OnMarginChanged(const FMargin& InMargin) override;
 	//~ End UTextLayoutWidget Interface
 
 	//~ Begin UWidget Interface
diff --git a/Engine/Source/Runtime/UMG/Public/Components/MultiLineEditableTextBox.h b/Engine/Source/Runtime/UMG/Public/Components/MultiLineEditableTextBox.h
--- a/Engine/Source/Runtime/UMG/Public/Components/MultiLineEditableTextBox.h
+++ b/Engine/Source/Runtime/UMG/Public/Components/MultiLineEditableTextBox.h
@@ -152,6 +152,7 @@ public:
 protected:
 	//~ Begin UTextLayoutWidget Interface
 	UMG_API virtual void OnLineHeightPercentageChanged(float InLineHeightPercentage) override;
+	UMG_API virtual void OnMarginChanged(const FMargin& InMargin) override;
 	//~ End UTextLayoutWidget Interface
 
 	//~ Begin UWidget Interface
diff --git a/Engine/Source/Runtime/UMG/Public/Components/RichTextBlock.h b/Engine/Source/Runtime/UMG/Public/Components/RichTextBlock.h
--- a/Engine/Source/Runtime/UMG/Public/Components/RichTextBlock.h
+++ b/Engine/Source/Runtime/UMG/Public/Components/RichTextBlock.h
@@ -229,6 +229,7 @@ public:
 protected:
 	//~ Begin UTextLayoutWidget Interface
 	UMG_API virtual void OnLineHeightPercentageChanged(float InLineHeightPercentage) override;
+	UMG_API virtual void OnMarginChanged(const FMargin& InMargin) override;
 	//~ End UTextLayoutWidget Interface
 
 	UMG_API virtual TSharedRef<SWidget> RebuildWidget() override;
diff --git a/Engine/Source/Runtime/UMG/Public/Components/TextBlock.h b/Engine/Source/Runtime/UMG/Public/Components/TextBlock.h
--- a/Engine/Source/Runtime/UMG/Public/Components/TextBlock.h
+++ b/Engine/Source/Runtime/UMG/Public/Components/TextBlock.h
@@ -258,6 +258,7 @@ public:
 protected:
 	//~ Begin UTextLayoutWidget Interface
 	UMG_API virtual void OnLineHeightPercentageChanged(float InLineHeightPercentage) override;
+	UMG_API virtual void OnMarginChanged(const FMargin& InMargin) override;
 	//~ End UTextLayoutWidget Interface
 
 	/**
diff --git a/Engine/Source/Runtime/UMG/Public/Components/TextWidgetTypes.h b/Engine/Source/Runtime/UMG/Public/Components/TextWidgetTypes.h
--- a/Engine/Source/Runtime/UMG/Public/Components/TextWidgetTypes.h
+++ b/Engine/Source/Runtime/UMG/Public/Components/TextWidgetTypes.h
@@ -66,9 +66,11 @@ public:
 	virtual void SetJustification(ETextJustify::Type InJustification) { Justification = InJustification; }
 
 	UMG_API void SetLineHeightPercentage(float InLineHeightPercentage);
+	UMG_API void SetMargin(const FMargin& InMargin);
 
 protected:
 	virtual void OnLineHeightPercentageChanged(float InLineHeightPercentage) {};
+	virtual void OnMarginChanged(const FMargin& InMargin) {};
 
 	/** Synchronize the properties with the given widget. A template as the Slate widgets conform to the same API, but don't derive from a common base. */
 	template <typename TWidgetType>
@@ -105,10 +107,10 @@ protected:
 	float WrapTextAt;
 
 	/** The amount of blank space left around the edges of text area. */
-	UPROPERTY(EditAnywhere, BlueprintReadOnly, Category=Appearance, AdvancedDisplay)
+	UPROPERTY(EditAnywhere, BlueprintReadOnly, Setter, Category=Appearance, AdvancedDisplay)
 	FMargin Margin;
 
 	/** The amount to scale each lines height by. */
-	UPROPERTY(EditAnywhere, BlueprintReadWrite, Setter, Category=Appearance, AdvancedDisplay)
+	UPROPERTY(EditAnywhere, BlueprintReadOnly, Setter, Category=Appearance, AdvancedDisplay)
 	float LineHeightPercentage;
 };
