UMG: Make LineHeightPercentage and Margin on UTextLayoutWidget as BlueprintReadWrite

diff --git a/Engine/Source/Runtime/UMG/Public/Components/TextWidgetTypes.h b/Engine/Source/Runtime/UMG/Public/Components/TextWidgetTypes.h
--- a/Engine/Source/Runtime/UMG/Public/Components/TextWidgetTypes.h
+++ b/Engine/Source/Runtime/UMG/Public/Components/TextWidgetTypes.h
@@ -107,10 +107,10 @@ protected:
 	float WrapTextAt;
 
 	/** The amount of blank space left around the edges of text area. */
-	UPROPERTY(EditAnywhere, BlueprintReadOnly, Setter, Category=Appearance, AdvancedDisplay)
+	UPROPERTY(EditAnywhere, BlueprintReadWrite, Setter, Category=Appearance, AdvancedDisplay)
 	FMargin Margin;
 
 	/** The amount to scale each lines height by. */
-	UPROPERTY(EditAnywhere, BlueprintReadOnly, Setter, Category=Appearance, AdvancedDisplay)
+	UPROPERTY(EditAnywhere, BlueprintReadWrite, Setter, Category=Appearance, AdvancedDisplay)
 	float LineHeightPercentage;
 };
