UMG AsyncTaskDownloadImage
- Check for http response codes before proceeding.

diff --git a/Engine/Source/Runtime/UMG/Private/Blueprint/AsyncTaskDownloadImage.cpp b/Engine/Source/Runtime/UMG/Private/Blueprint/AsyncTaskDownloadImage.cpp
--- a/Engine/Source/Runtime/UMG/Private/Blueprint/AsyncTaskDownloadImage.cpp
+++ b/Engine/Source/Runtime/UMG/Private/Blueprint/AsyncTaskDownloadImage.cpp
@@ -57,7 +57,8 @@ void UAsyncTaskDownloadImage::HandleImageRequest(FHttpRequestPtr HttpRequest, FH
 
 	RemoveFromRoot();
 
-	if ( bSucceeded && HttpResponse.IsValid() && HttpResponse->GetContentLength() > 0 && HttpResponse->GetContent().Num() > 0 )
+	if ( bSucceeded && HttpResponse.IsValid() && EHttpResponseCodes::IsOk(HttpResponse->GetResponseCode()) &&
+		 HttpResponse->GetContentLength() > 0 && HttpResponse->GetContent().Num() > 0 )
 	{
 		IImageWrapperModule& ImageWrapperModule = FModuleManager::LoadModuleChecked<IImageWrapperModule>(FName("ImageWrapper"));
 		TSharedPtr<IImageWrapper> ImageWrappers[3] =
