WidgetComponent: Add support for invalidation.

diff --git a/Engine/Source/Runtime/Slate/Private/Widgets/SVirtualWindow.cpp b/Engine/Source/Runtime/Slate/Private/Widgets/SVirtualWindow.cpp
--- a/Engine/Source/Runtime/Slate/Private/Widgets/SVirtualWindow.cpp
+++ b/Engine/Source/Runtime/Slate/Private/Widgets/SVirtualWindow.cpp
@@ -60,6 +60,11 @@ bool SVirtualWindow::SupportsKeyboardFocus() const
 	return bIsFocusable;
 }
 
+bool SVirtualWindow::Advanced_IsInvalidationRoot() const
+{
+	return bAllowFastUpdate;
+}
+
 FVector2D SVirtualWindow::ComputeDesiredSize(float LayoutScaleMultiplier) const
 {
 	return SCompoundWidget::ComputeDesiredSize(LayoutScaleMultiplier);
diff --git a/Engine/Source/Runtime/Slate/Public/Widgets/SVirtualWindow.h b/Engine/Source/Runtime/Slate/Public/Widgets/SVirtualWindow.h
--- a/Engine/Source/Runtime/Slate/Public/Widgets/SVirtualWindow.h
+++ b/Engine/Source/Runtime/Slate/Public/Widgets/SVirtualWindow.h
@@ -32,11 +32,11 @@ public:
 	SLATE_API void Construct(const FArguments& InArgs);
 
 	/**
-	* We allow users to control the resolve policy for deferred content.  When a virtual window is used in a retainer, we don't want it trying to resolve
-	* the deferred content, we instead need it to pass that content up to the non-retained window drawing the retainer.  If a virtual window is used in
-	* a case like the WidgetComponent, we always want to resolve the deferred content because there won't be another opportunity since the 3D widget has
-	* no owner window that can draw it.
-	*/
+	 * We allow users to control the resolve policy for deferred content. When a virtual window is used in a retainer, we don't want it trying to resolve
+	 * the deferred content, we instead need it to pass that content up to the non-retained window drawing the retainer. If a virtual window is used in
+	 * a case like the WidgetComponent, we always want to resolve the deferred content because there won't be another opportunity since the 3D widget has
+	 * no owner window that can draw it.
+	 */
 	SLATE_API void SetShouldResolveDeferred(bool bResolve);
 
 	/**  */
@@ -48,6 +48,7 @@ public:
 	SLATE_API virtual void OnArrangeChildren(const FGeometry& AllottedGeometry, FArrangedChildren& ArrangedChildren) const override;
 	SLATE_API virtual bool SupportsKeyboardFocus() const override;
 	SLATE_API virtual FVector2D ComputeDesiredSize(float LayoutScaleMultiplier) const override;
+	SLATE_API virtual bool Advanced_IsInvalidationRoot() const override;
 
 private:
 	TSharedPtr<class STooltipPresenter> TooltipPresenter;
diff --git a/Engine/Source/Runtime/SlateCore/Private/Widgets/SWindow.cpp b/Engine/Source/Runtime/SlateCore/Private/Widgets/SWindow.cpp
--- a/Engine/Source/Runtime/SlateCore/Private/Widgets/SWindow.cpp
+++ b/Engine/Source/Runtime/SlateCore/Private/Widgets/SWindow.cpp
@@ -1109,12 +1109,12 @@ bool SWindow::Advanced_IsInvalidationRoot() const
 
 const FSlateInvalidationRoot* SWindow::Advanced_AsInvalidationRoot() const
 {
-	return (bAllowFastUpdate && GSlateEnableGlobalInvalidation) ? this : nullptr;
+	return Advanced_IsInvalidationRoot() ? this : nullptr;
 }
 
 void SWindow::ProcessWindowInvalidation()
 {
-	if (bAllowFastUpdate && GSlateEnableGlobalInvalidation)
+	if (Advanced_IsInvalidationRoot())
 	{
 		ProcessInvalidation();
 	}
@@ -1122,7 +1122,7 @@ void SWindow::ProcessWindowInvalidation()
 
 bool SWindow::CustomPrepass(float LayoutScaleMultiplier)
 {
-	if (bAllowFastUpdate && GSlateEnableGlobalInvalidation)
+	if (Advanced_IsInvalidationRoot())
 	{
 		return NeedsPrepass();
 	}
@@ -1797,7 +1797,7 @@ FVector2D SWindow::ComputeDesiredSize(float LayoutScaleMultiplier) const
 bool SWindow::ComputeVolatility() const
 {
 	// If the entire window is volatile in fast path that defeats the whole purpose.
-	return bAllowFastUpdate ? false : SWidget::ComputeVolatility();
+	return Advanced_IsInvalidationRoot() ? false : SWidget::ComputeVolatility();
 }
 
 void SWindow::OnGlobalInvalidationToggled(bool bGlobalInvalidationEnabled)
@@ -2088,7 +2088,7 @@ int32 SWindow::PaintWindow( double CurrentTime, float DeltaTime, FSlateWindowEle
 	FSlateInvalidationContext Context(OutDrawElements, InWidgetStyle);
 	Context.bParentEnabled = bParentEnabled;
 	// Fast path at the window level should only be enabled if global invalidation is allowed
-	Context.bAllowFastPathUpdate = bAllowFastUpdate && GSlateEnableGlobalInvalidation;
+	Context.bAllowFastPathUpdate = Advanced_IsInvalidationRoot();
 	Context.LayoutScaleMultiplier = FSlateApplicationBase::Get().GetApplicationScale() * GetDPIScaleFactor();
 	Context.PaintArgs = &PaintArgs;
 	Context.IncomingLayerId = 0;
diff --git a/Engine/Source/Runtime/SlateCore/Public/FastUpdate/SlateInvalidationRoot.h b/Engine/Source/Runtime/SlateCore/Public/FastUpdate/SlateInvalidationRoot.h
--- a/Engine/Source/Runtime/SlateCore/Public/FastUpdate/SlateInvalidationRoot.h
+++ b/Engine/Source/Runtime/SlateCore/Public/FastUpdate/SlateInvalidationRoot.h
@@ -89,7 +89,7 @@ public:
 	//~ End FGCObject interface
 
 	/** Rebuild the list and request a SlowPath. */
-	UE_DEPRECATED(4.27, "InvalidateRoot is deprecated, use InvalidateRootChildOrder or InvalidateRootChildOrder")
+	UE_DEPRECATED(4.27, "InvalidateRoot is deprecated, use InvalidateRootChildOrder or InvalidateRootLayout")
 	SLATECORE_API void InvalidateRoot(const SWidget* Investigator = nullptr);
 	/** Rebuild the list and request a SlowPath. */
 	SLATECORE_API void InvalidateRootChildOrder(const SWidget* Investigator = nullptr);
diff --git a/Engine/Source/Runtime/UMG/Private/Components/WidgetComponent.cpp b/Engine/Source/Runtime/UMG/Private/Components/WidgetComponent.cpp
--- a/Engine/Source/Runtime/UMG/Private/Components/WidgetComponent.cpp
+++ b/Engine/Source/Runtime/UMG/Private/Components/WidgetComponent.cpp
@@ -21,11 +21,13 @@
 #include "DynamicMeshBuilder.h"
 #include "PhysicsEngine/BoxElem.h"
 #include "PhysicsEngine/BodySetup.h"
+#include "Rendering/SlateDrawBuffer.h"
 #include "Slate/SGameLayerManager.h"
 #include "Slate/WidgetRenderer.h"
 #include "Slate/SWorldWidgetScreenLayer.h"
 #include "UObject/EditorObjectVersion.h"
 #include "Widgets/SViewport.h"
+#include "Widgets/SVirtualWindow.h"
 #include "SceneInterface.h"
 
 #include UE_INLINE_GENERATED_CPP_BY_NAME(WidgetComponent)
@@ -609,9 +611,10 @@ UWidgetComponent::UWidgetComponent( const FObjectInitializer& PCIP )
 	, DrawSize( FIntPoint( 500, 500 ) )
 	, bManuallyRedraw(false)
 	, bRedrawRequested(true)
-	, RedrawTime(0)
-	, LastWidgetRenderTime(0)
+	, RedrawTime(0.0f)
+	, LastWidgetRenderTime(0.0)
 	, CurrentDrawSize(FIntPoint(0, 0))
+	, bUseInvalidationInWorldSpace(false)
 	, bReceiveHardwareInput(false)
 	, bWindowFocusable(true)
 	, WindowVisibility(EWindowVisibility::SelfHitTestInvisible)
@@ -974,10 +977,14 @@ void UWidgetComponent::SetWindowVisibility(EWindowVisibility InVisibility)
 	ensure(Widget);
 
 	WindowVisibility = InVisibility;
- 	if (SlateWindow.IsValid())
- 	{		
- 		SlateWindow->SetVisibility(ConvertWindowVisibilityToVisibility(WindowVisibility));
- 	}
+	if (SlateWindow.IsValid())
+	{
+		SlateWindow->SetVisibility(ConvertWindowVisibilityToVisibility(WindowVisibility));
+		if (bUseInvalidationInWorldSpace)
+		{
+			SlateWindow->InvalidateRootLayout();
+		}
+	}
 
 	if (IsWidgetVisible())
 	{
@@ -1344,27 +1351,35 @@ void UWidgetComponent::DrawWidgetToRenderTarget(float DeltaTime)
 
 	const float DrawScale = 1.0f;
 
-	bool bHasValidSize = true;
-	if ( bDrawAtDesiredSize )
+	if (bUseInvalidationInWorldSpace)
+	{
+		SlateWindow->ProcessWindowInvalidation();
+		SlateWindow->SlatePrepass(DrawScale);
+		WidgetRenderer->SetIsPrepassNeeded(false);
+	}
+	else if (bDrawAtDesiredSize)
 	{
 		SlateWindow->SlatePrepass(DrawScale);
+		WidgetRenderer->SetIsPrepassNeeded(false);
+	}
+	else
+	{
+		WidgetRenderer->SetIsPrepassNeeded(true);
+	}
 
+	bool bHasValidSize = true;
+	if ( bDrawAtDesiredSize )
+	{
 		FVector2D DesiredSize = SlateWindow->GetDesiredSize();
 		DesiredSize.X = FMath::RoundToInt(DesiredSize.X);
 		DesiredSize.Y = FMath::RoundToInt(DesiredSize.Y);
 		CurrentDrawSize = DesiredSize.IntPoint();
 
-		WidgetRenderer->SetIsPrepassNeeded(false);
-
 		if (DesiredSize.X <= 0 || DesiredSize.Y <= 0)
 		{
 			bHasValidSize = false;
 		}
 	}
-	else
-	{
-		WidgetRenderer->SetIsPrepassNeeded(true);
-	}
 
 	if ( CurrentDrawSize != PreviousDrawSize )
 	{
@@ -1386,13 +1401,27 @@ void UWidgetComponent::DrawWidgetToRenderTarget(float DeltaTime)
 	{
 		bRedrawRequested = false;
 
-		WidgetRenderer->DrawWindow(
-			RenderTarget,
-			SlateWindow->GetHittestGrid(),
-			SlateWindow.ToSharedRef(),
-			DrawScale,
-			CurrentDrawSize,
-			DeltaTime);
+		if (bUseInvalidationInWorldSpace)
+		{
+			static IConsoleVariable* CVar = IConsoleManager::Get().FindConsoleVariable(TEXT("Slate.DeferRetainedRenderingRenderThread"));
+			ensure(CVar);
+			bool bDeferRetainedRenderingRenderThread = CVar ? CVar->GetInt() != 0 : false;
+
+			FPaintArgs PaintArgs(nullptr, SlateWindow->GetHittestGrid(), FVector2D::ZeroVector, FApp::GetCurrentTime(), DeltaTime);
+			TSharedRef<SVirtualWindow> WindowToDraw = SlateWindow.ToSharedRef();
+
+			WidgetRenderer->DrawInvalidationRoot(WindowToDraw, RenderTarget, PaintArgs, DrawScale, CurrentDrawSize, bDeferRetainedRenderingRenderThread);
+		}
+		else
+		{
+			WidgetRenderer->DrawWindow(
+				RenderTarget,
+				SlateWindow->GetHittestGrid(),
+				SlateWindow.ToSharedRef(),
+				DrawScale,
+				CurrentDrawSize,
+				DeltaTime);
+		}
 
 		LastWidgetRenderTime = GetCurrentTime();
 
@@ -1534,20 +1563,19 @@ void UWidgetComponent::PostEditChangeProperty(FPropertyChangedEvent& PropertyCha
 
 	if( Property && PropertyChangedEvent.ChangeType != EPropertyChangeType::Interactive )
 	{
-		static FName DrawSizeName("DrawSize");
-		static FName PivotName("Pivot");
-		static FName WidgetClassName("WidgetClass");
-		static FName IsOpaqueName("bIsOpaque");
-		static FName IsTwoSidedName("bIsTwoSided");
-		static FName BackgroundColorName("BackgroundColor");
-		static FName TintColorAndOpacityName("TintColorAndOpacity");
-		static FName OpacityFromTextureName("OpacityFromTexture");
-		static FName ParabolaDistortionName(TEXT("ParabolaDistortion"));
-		static FName BlendModeName( TEXT( "BlendMode" ) );
-		static FName GeometryModeName( TEXT("GeometryMode") );
-		static FName CylinderArcAngleName( TEXT("CylinderArcAngle") );
-		static FName bWindowFocusableName(TEXT("bWindowFocusable"));
-		static FName WindowVisibilityName(TEXT("WindowVisibility"));
+		static FName DrawSizeName = GET_MEMBER_NAME_CHECKED(UWidgetComponent, DrawSize);
+		static FName PivotName = GET_MEMBER_NAME_CHECKED(UWidgetComponent, Pivot);
+		static FName WidgetClassName = GET_MEMBER_NAME_CHECKED(UWidgetComponent, WidgetClass);
+		static FName IsTwoSidedName = GET_MEMBER_NAME_CHECKED(UWidgetComponent, bIsTwoSided);
+		static FName BackgroundColorName = GET_MEMBER_NAME_CHECKED(UWidgetComponent, BackgroundColor);
+		static FName TintColorAndOpacityName = GET_MEMBER_NAME_CHECKED(UWidgetComponent, TintColorAndOpacity);
+		static FName OpacityFromTextureName = GET_MEMBER_NAME_CHECKED(UWidgetComponent, OpacityFromTexture);
+		static FName BlendModeName = GET_MEMBER_NAME_CHECKED(UWidgetComponent, BlendMode);
+		static FName GeometryModeName = GET_MEMBER_NAME_CHECKED(UWidgetComponent, GeometryMode);
+		static FName CylinderArcAngleName = GET_MEMBER_NAME_CHECKED(UWidgetComponent, CylinderArcAngle);
+		static FName bWindowFocusableName = GET_MEMBER_NAME_CHECKED(UWidgetComponent, bWindowFocusable);
+		static FName WindowVisibilityName = GET_MEMBER_NAME_CHECKED(UWidgetComponent, WindowVisibility);
+		static FName UseInvalidationName = GET_MEMBER_NAME_CHECKED(UWidgetComponent, bUseInvalidationInWorldSpace);
 
 		auto PropertyName = Property->GetFName();
 
@@ -1563,28 +1591,36 @@ void UWidgetComponent::PostEditChangeProperty(FPropertyChangedEvent& PropertyCha
 			}
 		}
 
-		if( PropertyName == WidgetClassName )
+		if (PropertyName == WidgetClassName)
 		{
 			Widget = nullptr;
 
 			UpdateWidget();
 			MarkRenderStateDirty();
 		}
-		else if ( PropertyName == DrawSizeName || PropertyName == PivotName || PropertyName == GeometryModeName || PropertyName == CylinderArcAngleName )
-		{
-			MarkRenderStateDirty();
-			UpdateBodySetup(true);
-			RecreatePhysicsState();
-		}
-		else if ( PropertyName == IsOpaqueName || PropertyName == IsTwoSidedName || PropertyName == BlendModeName )
+		else if (PropertyName == UseInvalidationName)
 		{
-			MarkRenderStateDirty();
+			if (SlateWindow)
+			{
+				SlateWindow->SetAllowFastUpdate(bUseInvalidationInWorldSpace);
+				SlateWindow->InvalidateRootLayout();
+			}
 		}
-		else if( PropertyName == BackgroundColorName || PropertyName == ParabolaDistortionName )
+		else if (PropertyName == DrawSizeName
+			|| PropertyName == PivotName
+			|| PropertyName == GeometryModeName
+			|| PropertyName == CylinderArcAngleName)
 		{
 			MarkRenderStateDirty();
+			UpdateBodySetup(true);
+			RecreatePhysicsState();
 		}
-		else if( PropertyName == TintColorAndOpacityName || PropertyName == OpacityFromTextureName )
+		else if (PropertyName == BlendModeName
+			|| PropertyName == IsTwoSidedName
+			|| PropertyName == BlendModeName
+			|| PropertyName == BackgroundColorName
+			|| PropertyName == TintColorAndOpacityName
+			|| PropertyName == OpacityFromTextureName)
 		{
 			MarkRenderStateDirty();
 		}
@@ -1731,6 +1767,7 @@ void UWidgetComponent::UpdateWidget()
 
 				SlateWindow = SNew(SVirtualWindow).Size(CurrentDrawSize);
 				SlateWindow->SetIsFocusable(bWindowFocusable);
+				SlateWindow->SetAllowFastUpdate(bUseInvalidationInWorldSpace);
 				SlateWindow->SetVisibility(ConvertWindowVisibilityToVisibility(WindowVisibility));
 				RegisterWindow();
 
diff --git a/Engine/Source/Runtime/UMG/Private/Slate/SRetainerWidget.cpp b/Engine/Source/Runtime/UMG/Private/Slate/SRetainerWidget.cpp
--- a/Engine/Source/Runtime/UMG/Private/Slate/SRetainerWidget.cpp
+++ b/Engine/Source/Runtime/UMG/Private/Slate/SRetainerWidget.cpp
@@ -564,6 +564,7 @@ SRetainerWidget::EPaintRetainedContentResult SRetainerWidget::PaintRetainedConte
 				SurfaceBrush.ImageSize = DrawSize;
 
 				WidgetRenderer->ViewOffset = -ViewOffset;
+				WidgetRenderer->SetIsPrepassNeeded(false);
 
 				bool bRepaintedWidgets = WidgetRenderer->DrawInvalidationRoot(VirtualWindow, RenderTarget, *this, Context, GDeferRetainedRenderingRenderThread != 0);
 				bRenderRequested = false;
diff --git a/Engine/Source/Runtime/UMG/Private/Slate/WidgetRenderer.cpp b/Engine/Source/Runtime/UMG/Private/Slate/WidgetRenderer.cpp
--- a/Engine/Source/Runtime/UMG/Private/Slate/WidgetRenderer.cpp
+++ b/Engine/Source/Runtime/UMG/Private/Slate/WidgetRenderer.cpp
@@ -287,6 +287,14 @@ void FWidgetRenderer::DrawWindow(
 
 bool FWidgetRenderer::DrawInvalidationRoot(TSharedRef<SVirtualWindow>& VirtualWindow, UTextureRenderTarget2D* RenderTarget, FSlateInvalidationRoot& Root, const FSlateInvalidationContext& Context, bool bDeferRenderTargetUpdate)
 {
+	{
+		const bool bIsInvalidationRoot = VirtualWindow->Advanced_IsInvalidationRoot();
+		ensure(bIsInvalidationRoot);
+		if (!bIsInvalidationRoot)
+		{
+			return false;
+		}
+	}
 	bool bRepaintedWidgets = false;
 #if !UE_SERVER
 	FSlateRenderer* MainSlateRenderer = FSlateApplication::Get().GetRenderer();
@@ -295,6 +303,13 @@ bool FWidgetRenderer::DrawInvalidationRoot(TSharedRef<SVirtualWindow>& VirtualWi
 
 	if (LIKELY(FApp::CanEverRender()))
 	{
+
+		if (bPrepassNeeded)
+		{
+			VirtualWindow->ProcessWindowInvalidation();
+			VirtualWindow->SlatePrepass(Context.LayoutScaleMultiplier);
+		}
+
 		// Need to set a new window element list so make a copy
 		FSlateInvalidationContext ContextCopy = Context;
 		ContextCopy.ViewOffset = ViewOffset;
@@ -304,15 +319,12 @@ bool FWidgetRenderer::DrawInvalidationRoot(TSharedRef<SVirtualWindow>& VirtualWi
 			ISlate3DRenderer::FScopedAcquireDrawBuffer ScopedDrawBuffer{ *Renderer, bDeferRenderTargetUpdate };
 			FSlateWindowElementList& WindowElementList = ScopedDrawBuffer.GetDrawBuffer().AddWindowElementList(VirtualWindow);
 
-			// Populates cached element data lists using the providied invalidation root
+			// Populates cached element data lists using the provide invalidation root
 			ContextCopy.WindowElementList = &WindowElementList;
 			FSlateInvalidationResult Result = Root.PaintInvalidationRoot(ContextCopy);
 
-			const int32 MaxLayerId = Result.MaxLayerIdPainted;
-
 			if (Result.bRepaintedWidgets)
 			{
-				//MaxLayerId = WindowElementList.PaintDeferred(MaxLayerId);
 				DeferredPaints = WindowElementList.GetDeferredPaintList();
 
 				Renderer->DrawWindow_GameThread(ScopedDrawBuffer.GetDrawBuffer());
@@ -353,4 +365,90 @@ bool FWidgetRenderer::DrawInvalidationRoot(TSharedRef<SVirtualWindow>& VirtualWi
 	return bRepaintedWidgets;
 }
 
+bool FWidgetRenderer::DrawInvalidationRoot(TSharedRef<SVirtualWindow>& VirtualWindow, UTextureRenderTarget2D* RenderTarget, FPaintArgs PaintArgs, float DrawScale, FVector2D DrawSize, bool bDeferRenderTargetUpdate)
+{
+	{
+		const bool bIsInvalidationRoot = VirtualWindow->Advanced_IsInvalidationRoot();
+		ensure(bIsInvalidationRoot);
+		if (!bIsInvalidationRoot)
+		{
+			return false;
+		}
+	}
+
+	bool bRepaintedWidgets = false;
+#if !UE_SERVER
+	FSlateRenderer* MainSlateRenderer = FSlateApplication::Get().GetRenderer();
+	FScopeLock ScopeLock(MainSlateRenderer->GetResourceCriticalSection());
+
+	if (LIKELY(FApp::CanEverRender()))
+	{
+		if (bPrepassNeeded)
+		{
+			VirtualWindow->ProcessWindowInvalidation();
+			VirtualWindow->SlatePrepass(DrawScale);
+		}
+
+		{
+			FGeometry WindowGeometry = FGeometry::MakeRoot(DrawSize * (1 / DrawScale), FSlateLayoutTransform(DrawScale));
+			PaintArgs.GetHittestGrid().SetHittestArea(WindowGeometry.GetLayoutBoundingRect().GetTopLeft(), WindowGeometry.GetLayoutBoundingRect().GetSize());
+		}
+
+		{
+			// Get the free buffer & add our virtual window
+			ISlate3DRenderer::FScopedAcquireDrawBuffer ScopedDrawBuffer{ *Renderer, bDeferRenderTargetUpdate };
+			FSlateWindowElementList& WindowElementList = ScopedDrawBuffer.GetDrawBuffer().AddWindowElementList(VirtualWindow);
+
+			// Populates cached element data lists using the provide invalidation root
+			FSlateInvalidationContext InvalidationContext(WindowElementList, FWidgetStyle());
+			InvalidationContext.bParentEnabled = true;
+			InvalidationContext.bAllowFastPathUpdate = true;
+			InvalidationContext.LayoutScaleMultiplier = DrawScale;
+			InvalidationContext.PaintArgs = &PaintArgs;
+			InvalidationContext.IncomingLayerId = 0;
+			InvalidationContext.CullingRect = VirtualWindow->GetClippingRectangleInWindow();
+
+			FSlateInvalidationResult Result = VirtualWindow->PaintInvalidationRoot(InvalidationContext);
+			if (Result.bRepaintedWidgets)
+			{
+				DeferredPaints = WindowElementList.GetDeferredPaintList();
+
+				Renderer->DrawWindow_GameThread(ScopedDrawBuffer.GetDrawBuffer());
+
+				ScopedDrawBuffer.GetDrawBuffer().ViewOffset = Result.ViewOffset;
+
+				FRenderThreadUpdateContext RenderThreadUpdateContext =
+				{
+					&(ScopedDrawBuffer.GetDrawBuffer()),
+					(FApp::GetCurrentTime() - GStartTime),
+					static_cast<float>(FApp::GetDeltaTime()),
+					(FPlatformTime::Seconds() - GStartTime),
+					static_cast<float>(FApp::GetDeltaTime()),
+					static_cast<FRenderTarget*>(RenderTarget->GameThread_GetRenderTargetResource()),
+					Renderer.Get(),
+					bClearTarget
+				};
+
+				bRepaintedWidgets = Result.bRepaintedWidgets;
+				FSlateApplication::Get().GetRenderer()->AddWidgetRendererUpdate(RenderThreadUpdateContext, bDeferRenderTargetUpdate);
+
+
+				// Any deferred painted elements of the retainer should be drawn directly by the main renderer, not rendered into the render target,
+				// as most of those sorts of things will break the rendering rect, things like tooltips, and popup menus.
+				for (auto& DeferredPaint : DeferredPaints)
+				{
+					InvalidationContext.WindowElementList->QueueDeferredPainting(DeferredPaint->Copy(*InvalidationContext.PaintArgs));
+				}
+			}
+			else
+			{
+				WindowElementList.ResetElementList();
+			}
+		}
+	}
+#endif // !UE_SERVER
+
+	return bRepaintedWidgets;
+}
+
 
diff --git a/Engine/Source/Runtime/UMG/Public/Components/WidgetComponent.h b/Engine/Source/Runtime/UMG/Public/Components/WidgetComponent.h
--- a/Engine/Source/Runtime/UMG/Public/Components/WidgetComponent.h
+++ b/Engine/Source/Runtime/UMG/Public/Components/WidgetComponent.h
@@ -478,10 +478,17 @@ protected:
 	UPROPERTY(Transient, DuplicateTransient)
 	FIntPoint CurrentDrawSize;
 
+	/**
+	 * Use the invalidation system to update this widget.
+	 * Only valid in World space. In Screen space, the widget is updated by the viewport owners.
+	 */
+	UPROPERTY(EditAnywhere, Category = UserInterface, meta=(EditCondition="Space==EWidgetSpace::World", DisplayName="Use Invalidation"))
+	bool bUseInvalidationInWorldSpace;
+
 	/**
 	 * Causes the render target to automatically match the desired size.
 	 * 
-	 * WARNING: If you change this every frame, it will be very expensive.  If you need 
+	 * WARNING: If you change this every frame, it will be very expensive. If you need 
 	 *    that effect, you should keep the outer widget's sized locked and dynamically
 	 *    scale or resize some inner widget.
 	 */
diff --git a/Engine/Source/Runtime/UMG/Public/Slate/WidgetRenderer.h b/Engine/Source/Runtime/UMG/Public/Slate/WidgetRenderer.h
--- a/Engine/Source/Runtime/UMG/Public/Slate/WidgetRenderer.h
+++ b/Engine/Source/Runtime/UMG/Public/Slate/WidgetRenderer.h
@@ -144,6 +144,14 @@ public:
 		const FSlateInvalidationContext& Context,
 		bool bDeferRenderTargetUpdate = false);
 
+	UMG_API bool DrawInvalidationRoot(
+		TSharedRef<SVirtualWindow>& VirtualWindow,
+		UTextureRenderTarget2D* RenderTarget,
+		FPaintArgs PaintArgs,
+		float DrawScale,
+		FVector2D DrawSize,
+		bool bDeferRenderTargetUpdate = false);
+
 	TArray< TSharedPtr<FSlateWindowElementList::FDeferredPaint> > DeferredPaints;
 private:
 	/** The slate 3D renderer used to render the user slate widget */
