Fix warning V501 introduced with CL 29709782

diff --git a/Engine/Source/Runtime/UMG/Private/Components/WidgetComponent.cpp b/Engine/Source/Runtime/UMG/Private/Components/WidgetComponent.cpp
--- a/Engine/Source/Runtime/UMG/Private/Components/WidgetComponent.cpp
+++ b/Engine/Source/Runtime/UMG/Private/Components/WidgetComponent.cpp
@@ -1615,8 +1615,7 @@ void UWidgetComponent::PostEditChangeProperty(FPropertyChangedEvent& PropertyCha
 			UpdateBodySetup(true);
 			RecreatePhysicsState();
 		}
-		else if (PropertyName == BlendModeName
-			|| PropertyName == IsTwoSidedName
+		else if (PropertyName == IsTwoSidedName
 			|| PropertyName == BlendModeName
 			|| PropertyName == BackgroundColorName
 			|| PropertyName == TintColorAndOpacityName
