UMG: maintain the hierarchy order when creating and adding namedslot contents to the widget tree: makes sure a namedslot content is created after its owning namedslot (the owning namedslot may be another namedslot's content)
- also fix a bug where namedslots with content were still shown as available in a child widget if they had content in the parent widget.

diff --git a/Engine/Source/Editor/UMGEditor/Private/WidgetBlueprintCompiler.cpp b/Engine/Source/Editor/UMGEditor/Private/WidgetBlueprintCompiler.cpp
--- a/Engine/Source/Editor/UMGEditor/Private/WidgetBlueprintCompiler.cpp
+++ b/Engine/Source/Editor/UMGEditor/Private/WidgetBlueprintCompiler.cpp
@@ -953,6 +953,7 @@ void FWidgetBlueprintCompilerContext::FinishCompilingClass(UClass* Class)
 
 		// Add all the names of the named slot widgets to the slot names structure.
 		{
+			TArray<FName> NamedSlotsWithContentInSameTree;
 		#if WITH_EDITOR
 			BPGClass->NamedSlotsWithID.Reset();
 		#endif
@@ -974,6 +975,13 @@ void FWidgetBlueprintCompilerContext::FinishCompilingClass(UClass* Class)
 						{
 							BPGClass->InstanceNamedSlots.Add(Widget->GetFName());
 						}
+
+						// A namedslot whose content is in the same blueprint class is treated as a regular panel widget.
+						// We need to keep track of these to later remove them from the available namedslots list.
+						if (NamedSlot->GetChildrenCount() > 0)
+						{
+							NamedSlotsWithContentInSameTree.Add(NamedSlot->GetFName());
+						}
 					}
 				});
 				
@@ -989,6 +997,12 @@ void FWidgetBlueprintCompilerContext::FinishCompilingClass(UClass* Class)
 				BPGClass->AvailableNamedSlots.Remove(SlotName);
 			});
 
+			// Remove any named slots with content in the same widget tree from the available slots.
+			for (const FName& NamedSlotWithContent : NamedSlotsWithContentInSameTree)
+			{
+				BPGClass->AvailableNamedSlots.Remove(NamedSlotWithContent);
+			}
+
 			// Remove any available subclass named slots that are marked as instance named slot.
 			for (const FName& InstanceNamedSlot : BPGClass->InstanceNamedSlots)
 			{
diff --git a/Engine/Source/Runtime/UMG/Private/UserWidget.cpp b/Engine/Source/Runtime/UMG/Private/UserWidget.cpp
--- a/Engine/Source/Runtime/UMG/Private/UserWidget.cpp
+++ b/Engine/Source/Runtime/UMG/Private/UserWidget.cpp
@@ -219,8 +219,28 @@ void UUserWidget::DuplicateAndInitializeFromWidgetTree(UWidgetTree* InWidgetTree
 			}
 		});
 
+		TArray<UWidget*> AllNamedSlotContentWidgets;
+		NamedSlotContentToMerge.GenerateValueArray(AllNamedSlotContentWidgets);
+
+		auto SetContentWidgetForNamedSlot = [this](FName NamedSlotName, UWidget* TemplateSlotContent)
+		{
+			FObjectInstancingGraph NamedSlotInstancingGraph;
+			// We need to add a mapping from the template's widget tree to the new widget tree, that way
+			// as we instance the widget hierarchy it's grafted onto the new widget tree.
+			NamedSlotInstancingGraph.AddNewObject(WidgetTree, TemplateSlotContent->GetTypedOuter<UWidgetTree>());
+
+			// Instance the new widget from the foreign tree, but do it in a way that grafts it onto the tree we're instancing.
+			UWidget* Content = NewObject<UWidget>(WidgetTree, TemplateSlotContent->GetClass(), TemplateSlotContent->GetFName(), RF_Transactional, TemplateSlotContent, false, &NamedSlotInstancingGraph);
+			Content->SetFlags(RF_Transient | RF_DuplicateTransient);
+
+			// Insert the newly constructed widget into the named slot that corresponds.  The above creates
+			// it as if it was always part of the widget tree, but this actually puts it into a widget's
+			// slot for the named slot.
+			SetContentForSlot(NamedSlotName, Content);
+		};
+
 		// This block controls merging named slot content specified in a child class for the widget we're templated after.
-		for (const auto& KVP_SlotContent : NamedSlotContentToMerge)
+		for (const TPair<FName, UWidget*>& KVP_SlotContent : NamedSlotContentToMerge)
 		{
 			// Don't insert the named slot content if the named slot is filled already.  This is a problematic
 			// scenario though, if someone inserted content, but we have class default instances, we sorta leave
@@ -230,19 +250,35 @@ void UUserWidget::DuplicateAndInitializeFromWidgetTree(UWidgetTree* InWidgetTree
 			{
 				if (UWidget* TemplateSlotContent = KVP_SlotContent.Value)
 				{
-					FObjectInstancingGraph NamedSlotInstancingGraph;
-					// We need to add a mapping from the template's widget tree to the new widget tree, that way
-					// as we instance the widget hierarchy it's grafted onto the new widget tree.
-					NamedSlotInstancingGraph.AddNewObject(WidgetTree, TemplateSlotContent->GetTypedOuter<UWidgetTree>());
-
-					// Instance the new widget from the foreign tree, but do it in a way that grafts it onto the tree we're instancing.
-					UWidget* Content = NewObject<UWidget>(WidgetTree, TemplateSlotContent->GetClass(), TemplateSlotContent->GetFName(), RF_Transactional, TemplateSlotContent, false, &NamedSlotInstancingGraph);
-					Content->SetFlags(RF_Transient | RF_DuplicateTransient);
-
-					// Insert the newly constructed widget into the named slot that corresponds.  The above creates
-					// it as if it was always part of the widget tree, but this actually puts it into a widget's
-					// slot for the named slot.
-					SetContentForSlot(KVP_SlotContent.Key, Content);
+					TArray<TPair<FName, UWidget*>> NamedSlotContentCreationStack;
+					FName OwningNamedSlot = KVP_SlotContent.Key;
+					NamedSlotContentCreationStack.Add(TTuple<FName, UWidget*>(OwningNamedSlot, TemplateSlotContent));
+
+					// Search for the owning Namedslot to see if it is the content of another Namedslot itself.
+					// If so, we need to ensure it is added to the widget tree prior to its content.
+					// Repeat until the owning Namedslot is no longer found as the content of another.
+					while (UWidget** FoundContentWidget = AllNamedSlotContentWidgets.FindByPredicate([OwningNamedSlot](const UWidget* Content) {return Content ? Content->GetFName() == OwningNamedSlot : false;}))
+					{
+						UWidget* NestedNamedSlotContent = *FoundContentWidget;
+						OwningNamedSlot = *NamedSlotContentToMerge.FindKey(NestedNamedSlotContent);
+
+						// Make sure we have not already iterated on this Namedslot.
+						if (!GetContentForSlot(OwningNamedSlot) && !NamedSlotContentCreationStack.ContainsByPredicate([OwningNamedSlot](const TTuple<FName, UWidget*>& Content) {return Content.Key == OwningNamedSlot;}))
+						{
+							NamedSlotContentCreationStack.Add(TPair<FName, UWidget*>(OwningNamedSlot, NestedNamedSlotContent));
+						}
+						else 
+						{
+							break;
+						}
+					}
+
+					// Go through the namedslot/content pair in hierarchy order and add them to the widget tree.
+					for (int32 Index = NamedSlotContentCreationStack.Num() - 1; Index >= 0; Index--)
+					{
+						TPair<FName, UWidget*>& KVP = NamedSlotContentCreationStack[Index];
+						SetContentWidgetForNamedSlot(KVP.Key, KVP.Value);
+					}
 				}
 			}
 		}
