UMG Viewport: Exposed ability to SetEnableAdvancedFeatures which controls things like AntiAliasing and Advanced PostProcessing in the FPreviewScene, and functions SetLightIntensity and SetSkyIntensity to control the Default PreviewScene lights.

diff --git a/Engine/Source/Runtime/UMG/Private/Components/Viewport.cpp b/Engine/Source/Runtime/UMG/Private/Components/Viewport.cpp
--- a/Engine/Source/Runtime/UMG/Private/Components/Viewport.cpp
+++ b/Engine/Source/Runtime/UMG/Private/Components/Viewport.cpp
@@ -501,6 +501,36 @@ AActor* UViewport::Spawn(TSubclassOf<AActor> ActorClass)
 	return NULL;
 }
 
+void UViewport::SetEnableAdvancedFeatures(bool InEnableAdvancedFeatures)
+{
+	ShowFlags.DisableAdvancedFeatures();
+	if (ViewportWidget.IsValid())
+	{
+		check(ViewportWidget->ViewportClient.IsValid());
+		if (InEnableAdvancedFeatures)
+		{
+			ShowFlags.EnableAdvancedFeatures();
+			ViewportWidget->ViewportClient->SetEngineShowFlags(ShowFlags);
+		}
+		else
+		{
+			ShowFlags.DisableAdvancedFeatures();
+			ViewportWidget->ViewportClient->SetEngineShowFlags(ShowFlags);
+		}
+
+	}
+}
+
+void UViewport::SetLightIntensity(float InLightIntensity)
+{
+	ViewportWidget->PreviewScene.SetLightBrightness(InLightIntensity);
+}
+
+void UViewport::SetSkyIntensity(float InLightIntensity)
+{
+	ViewportWidget->PreviewScene.SetSkyBrightness(InLightIntensity);
+}
+
 PRAGMA_DISABLE_DEPRECATION_WARNINGS
 const FLinearColor& UViewport::GetBackgroundColor() const
 {
diff --git a/Engine/Source/Runtime/UMG/Public/Components/Viewport.h b/Engine/Source/Runtime/UMG/Public/Components/Viewport.h
--- a/Engine/Source/Runtime/UMG/Public/Components/Viewport.h
+++ b/Engine/Source/Runtime/UMG/Public/Components/Viewport.h
@@ -257,6 +257,15 @@ class UViewport : public UContentWidget
 
 	UMG_API const FLinearColor& GetBackgroundColor() const;
 
+	UFUNCTION(BlueprintCallable, Category = "Viewport")
+	UMG_API void SetEnableAdvancedFeatures(bool InEnableAdvancedFeatures);
+
+	UFUNCTION(BlueprintCallable, Category = "Viewport")
+	UMG_API void SetLightIntensity(float LightIntensity);
+
+	UFUNCTION(BlueprintCallable, Category = "Viewport")
+	UMG_API void SetSkyIntensity(float LightIntensity);
+
 	// UWidget interface
 	UMG_API virtual void SynchronizeProperties() override;
 	// End of UWidget interface
