UMG Viewport: Added function to set ShowFlags on the FPreviewScene.

diff --git a/Engine/Source/Runtime/UMG/Private/Components/Viewport.cpp b/Engine/Source/Runtime/UMG/Private/Components/Viewport.cpp
--- a/Engine/Source/Runtime/UMG/Private/Components/Viewport.cpp
+++ b/Engine/Source/Runtime/UMG/Private/Components/Viewport.cpp
@@ -501,6 +501,16 @@ AActor* UViewport::Spawn(TSubclassOf<AActor> ActorClass)
 	return NULL;
 }
 
+void UViewport::SetShowFlag(FString InShowFlagName, bool InValue)
+{
+	if (ShowFlags.IsNameThere(*InShowFlagName, 0))
+	{
+		int32 FlagIndex = ShowFlags.FindIndexByName(*InShowFlagName);
+		ShowFlags.SetSingleFlag(FlagIndex, InValue);
+		ViewportWidget->ViewportClient->SetEngineShowFlags(ShowFlags);
+	}
+}
+
 void UViewport::SetEnableAdvancedFeatures(bool InEnableAdvancedFeatures)
 {
 	ShowFlags.DisableAdvancedFeatures();
diff --git a/Engine/Source/Runtime/UMG/Public/Components/Viewport.h b/Engine/Source/Runtime/UMG/Public/Components/Viewport.h
--- a/Engine/Source/Runtime/UMG/Public/Components/Viewport.h
+++ b/Engine/Source/Runtime/UMG/Public/Components/Viewport.h
@@ -266,6 +266,9 @@ class UViewport : public UContentWidget
 	UFUNCTION(BlueprintCallable, Category = "Viewport")
 	UMG_API void SetSkyIntensity(float LightIntensity);
 
+	UFUNCTION(BlueprintCallable, Category = "Viewport")
+	UMG_API void SetShowFlag(FString InShowFlagName, bool InValue);
+
 	// UWidget interface
 	UMG_API virtual void SynchronizeProperties() override;
 	// End of UWidget interface
