Fixed AllowPrivateAccess specifiers that aren't following by "= true", which is required by some parsing code

#rnx

diff --git a/Engine/Source/Runtime/UMG/Public/Components/TextBlock.h b/Engine/Source/Runtime/UMG/Public/Components/TextBlock.h
--- a/Engine/Source/Runtime/UMG/Public/Components/TextBlock.h
+++ b/Engine/Source/Runtime/UMG/Public/Components/TextBlock.h
@@ -266,7 +266,7 @@ protected:
 	 * This feature is only suitable for "simple" text (ie, text containing only numbers or basic ASCII) as it disables the complex text rendering support required for certain languages (such as Arabic and Thai).
 	 * It is significantly faster for text that can take advantage of it (particularly if that text changes frequently), but shouldn't be used for localized user-facing text.
 	 */
-	UPROPERTY(EditAnywhere, BlueprintReadOnly, Category=Performance, AdvancedDisplay, meta=(AllowPrivateAccess, DesignerRebuild))
+	UPROPERTY(EditAnywhere, BlueprintReadOnly, Category=Performance, AdvancedDisplay, meta=(AllowPrivateAccess = "true", DesignerRebuild))
 	bool bSimpleTextMode;
 
 	//~ Begin UWidget Interface
