UMG: Correction to PanelWidget to prevent a null widget to be added to the panel which could cause a crash in many situation.

diff --git a/Engine/Source/Runtime/UMG/Private/Components/PanelWidget.cpp b/Engine/Source/Runtime/UMG/Private/Components/PanelWidget.cpp
--- a/Engine/Source/Runtime/UMG/Private/Components/PanelWidget.cpp
+++ b/Engine/Source/Runtime/UMG/Private/Components/PanelWidget.cpp
@@ -55,6 +55,7 @@ UWidget* UPanelWidget::GetChildAt(int32 Index) const
 TArray<UWidget*> UPanelWidget::GetAllChildren() const
 {
 	TArray<UWidget*> Result;
+	Result.Reserve(Slots.Num());
 
 	for(UPanelSlot* ChildSlot : Slots)
 	{
@@ -66,12 +67,15 @@ TArray<UWidget*> UPanelWidget::GetAllChildren() const
 
 int32 UPanelWidget::GetChildIndex(const UWidget* Content) const
 {
-	const int32 ChildCount = GetChildrenCount();
-	for ( int32 ChildIndex = 0; ChildIndex < ChildCount; ChildIndex++ )
+	if (Content)
 	{
-		if ( GetChildAt(ChildIndex) == Content )
+		const int32 ChildCount = GetChildrenCount();
+		for (int32 ChildIndex = 0; ChildIndex < ChildCount; ChildIndex++)
 		{
-			return ChildIndex;
+			if (GetChildAt(ChildIndex) == Content)
+			{
+				return ChildIndex;
+			}
 		}
 	}
 	return INDEX_NONE;
@@ -94,19 +98,26 @@ bool UPanelWidget::RemoveChildAt(int32 Index)
 	}
 
 	UPanelSlot* PanelSlot = Slots[Index];
-	if ( PanelSlot->Content )
-	{
-		PanelSlot->Content->Slot = nullptr;
-	}
-
 	Slots.RemoveAt(Index);
 
-	OnSlotRemoved(PanelSlot);
+	if (PanelSlot)
+	{
+		if (PanelSlot->Content)
+		{
+			PanelSlot->Content->Slot = nullptr;
+		}
+
+		OnSlotRemoved(PanelSlot);
 
-	const bool bReleaseChildren = true;
-	PanelSlot->ReleaseSlateResources(bReleaseChildren);
-	PanelSlot->Parent = nullptr;
-	PanelSlot->Content = nullptr;
+		const bool bReleaseChildren = true;
+		PanelSlot->ReleaseSlateResources(bReleaseChildren);
+		PanelSlot->Parent = nullptr;
+		PanelSlot->Content = nullptr;
+	}
+	else
+	{
+		return false;
+	}
 
 	InvalidateLayoutAndVolatility();
 
@@ -152,18 +163,14 @@ UPanelSlot* UPanelWidget::AddChild(UWidget* Content)
 
 bool UPanelWidget::ReplaceChildAt(int32 Index, UWidget* Content)
 {
-	if ( Index < 0 || Index >= Slots.Num() )
+	if ( Index < 0 || Index >= Slots.Num() || Content == nullptr)
 	{
 		return false;
 	}
 
 	UPanelSlot* PanelSlot = Slots[Index];
 	PanelSlot->Content = Content;
-
-	if ( Content )
-	{
-		Content->Slot = PanelSlot;
-	}
+	Content->Slot = PanelSlot;
 
 	PanelSlot->SynchronizeProperties();
 
@@ -172,10 +179,13 @@ bool UPanelWidget::ReplaceChildAt(int32 Index, UWidget* Content)
 
 bool UPanelWidget::ReplaceChild(UWidget* CurrentChild, UWidget* NewChild)
 {
-	int32 Index = GetChildIndex(CurrentChild);
-	if ( Index != -1 )
+	if (NewChild)
 	{
-		return ReplaceChildAt(Index, NewChild);
+		int32 Index = GetChildIndex(CurrentChild);
+		if (Index != INDEX_NONE)
+		{
+			return ReplaceChildAt(Index, NewChild);
+		}
 	}
 
 	return false;
@@ -183,16 +193,23 @@ bool UPanelWidget::ReplaceChild(UWidget* CurrentChild, UWidget* NewChild)
 
 UPanelSlot* UPanelWidget::InsertChildAt(int32 Index, UWidget* Content)
 {
-	UPanelSlot* NewSlot = AddChild(Content);
-	ShiftChild(Index, Content);
-	return NewSlot;
+	if (Content)
+	{
+		UPanelSlot* NewSlot = AddChild(Content);
+		ShiftChild(Index, Content);
+		return NewSlot;
+	}
+	return nullptr;
 }
 
 void UPanelWidget::ShiftChild(int32 Index, UWidget* Child)
 {
 	int32 CurrentIndex = GetChildIndex(Child);
-	Slots.RemoveAt(CurrentIndex);
-	Slots.Insert(Child->Slot, FMath::Clamp(Index, 0, Slots.Num()));
+	if (CurrentIndex != INDEX_NONE)
+	{
+		Slots.RemoveAt(CurrentIndex);
+		Slots.Insert(Child->Slot, FMath::Clamp(Index, 0, Slots.Num()));
+	}
 }
 
 void UPanelWidget::SetDesignerFlags(EWidgetDesignFlags NewFlags)
@@ -215,7 +232,7 @@ void UPanelWidget::SetDesignerFlags(EWidgetDesignFlags NewFlags)
 bool UPanelWidget::RemoveChild(UWidget* Content)
 {
 	int32 ChildIndex = GetChildIndex(Content);
-	if ( ChildIndex != -1 )
+	if ( ChildIndex != INDEX_NONE)
 	{
 		return RemoveChildAt(ChildIndex);
 	}
