PR #10925: Fix UserWidget's PlayAnimationForward and PlayAnimationReverse not playing animations if the animation was recently stopped

In PlayAnimationForward/PlayAnimationReverse, just return the player if it's already playing. GetSequencePlayers should only be returning players that are NOT stopped.  However, there is the possibility that Stop() has been called on a player, but its status has not been changed to Stopped. In that case, this check will be bypassed and PlayAnimation will be called so that the animation will play.

diff --git a/Engine/Source/Runtime/UMG/Private/UserWidget.cpp b/Engine/Source/Runtime/UMG/Private/UserWidget.cpp
--- a/Engine/Source/Runtime/UMG/Private/UserWidget.cpp
+++ b/Engine/Source/Runtime/UMG/Private/UserWidget.cpp
@@ -612,7 +612,11 @@ UUMGSequencePlayer* UUserWidget::PlayAnimationForward(UWidgetAnimation* InAnimat
 {
 	// Don't create the player, only search for it.
 	UUMGSequencePlayer* Player = GetSequencePlayer(InAnimation);
-	if (Player)
+
+	// Just return the player if it's already playing. GetSequencePlayers should only be returning players that are NOT stopped. 
+	// However, there is the possibility that Stop() has been called on a player, but its status has not been changed to Stopped. 
+	// In that case, this check will be bypassed and PlayAnimation will be called so that the animation will play.
+	if (Player && Player->GetPlaybackStatus() == EMovieScenePlayerStatus::Playing)
 	{
 		if (!Player->IsPlayingForward())
 		{
@@ -630,7 +634,11 @@ UUMGSequencePlayer* UUserWidget::PlayAnimationReverse(UWidgetAnimation* InAnimat
 {
 	// Don't create the player, only search for it.
 	UUMGSequencePlayer* Player = GetSequencePlayer(InAnimation);
-	if (Player)
+
+	// Just return the player if it's already playing. GetSequencePlayers should only be returning players that are NOT stopped. 
+	// However, there is the possibility that Stop() has been called on a player, but its status has not been changed to Stopped. 
+	// In that case, this check will be bypassed and PlayAnimation will be called so that the animation will play.
+	if (Player && Player->GetPlaybackStatus() == EMovieScenePlayerStatus::Playing)
 	{
 		if (Player->IsPlayingForward())
 		{
