Fix to AllowRightClickDragScrolling to use the right setting

diff --git a/Engine/Source/Runtime/UMG/Private/Components/ScrollBox.cpp b/Engine/Source/Runtime/UMG/Private/Components/ScrollBox.cpp
--- a/Engine/Source/Runtime/UMG/Private/Components/ScrollBox.cpp
+++ b/Engine/Source/Runtime/UMG/Private/Components/ScrollBox.cpp
@@ -489,16 +489,16 @@ float UScrollBox::GetNavigationScrollPadding() const
 
 void UScrollBox::SetAllowRightClickDragScrolling(bool bShouldAllowRightClickDragScrolling)
 {
-	AlwaysShowScrollbarTrack = bShouldAllowRightClickDragScrolling;
+	bAllowRightClickDragScrolling = bShouldAllowRightClickDragScrolling;
 	if (MyScrollBox)
 	{
-		MyScrollBox->SetScrollBarRightClickDragAllowed(AlwaysShowScrollbarTrack);
+		MyScrollBox->SetScrollBarRightClickDragAllowed(bAllowRightClickDragScrolling);
 	}
 }
 
 bool UScrollBox::IsAllowRightClickDragScrolling() const
 {
-	return AlwaysShowScrollbarTrack;
+	return bAllowRightClickDragScrolling;
 }
 
 bool UScrollBox::IsFrontPadScrolling() const
