UMG: Broadcast the text FieldNotify after the underlying SWidget got updated. If the widget is created, the GetText function returns the SWidget value instead of the UWidget value.

diff --git a/Engine/Source/Runtime/UMG/Private/Components/EditableText.cpp b/Engine/Source/Runtime/UMG/Private/Components/EditableText.cpp
--- a/Engine/Source/Runtime/UMG/Private/Components/EditableText.cpp
+++ b/Engine/Source/Runtime/UMG/Private/Components/EditableText.cpp
@@ -144,9 +144,13 @@ FText UEditableText::GetText() const
 
 void UEditableText::SetText(FText InText)
 {
-	if (SetTextInternal(InText) && MyEditableText.IsValid() )
+	if (SetTextInternal(InText))
 	{
-		MyEditableText->SetText(Text);
+		if (MyEditableText.IsValid())
+		{
+			MyEditableText->SetText(Text);
+		}
+		BroadcastFieldValueChanged(FFieldNotificationClassDescriptor::Text);
 	}
 }
 
@@ -155,7 +159,6 @@ bool UEditableText::SetTextInternal(const FText& InText)
 	if (!Text.IdenticalTo(InText, ETextIdenticalModeFlags::DeepCompare | ETextIdenticalModeFlags::LexicalCompareInvariants))
 	{
 		Text = InText;
-		BroadcastFieldValueChanged(FFieldNotificationClassDescriptor::Text);
 		return true;
 	}
 
@@ -377,13 +380,17 @@ void UEditableText::HandleOnTextChanged(const FText& InText)
 {
 	if (SetTextInternal(InText))
 	{
+		BroadcastFieldValueChanged(FFieldNotificationClassDescriptor::Text);
 		OnTextChanged.Broadcast(InText);
 	}
 }
 
 void UEditableText::HandleOnTextCommitted(const FText& InText, ETextCommit::Type CommitMethod)
 {
-	SetTextInternal(InText);
+	if (SetTextInternal(InText))
+	{
+		BroadcastFieldValueChanged(FFieldNotificationClassDescriptor::Text);
+	}
 	OnTextCommitted.Broadcast(InText, CommitMethod);
 }
 
diff --git a/Engine/Source/Runtime/UMG/Private/Components/EditableTextBox.cpp b/Engine/Source/Runtime/UMG/Private/Components/EditableTextBox.cpp
--- a/Engine/Source/Runtime/UMG/Private/Components/EditableTextBox.cpp
+++ b/Engine/Source/Runtime/UMG/Private/Components/EditableTextBox.cpp
@@ -160,9 +160,13 @@ FText UEditableTextBox::GetText() const
 
 void UEditableTextBox::SetText(FText InText)
 {
-	if (SetTextInternal(InText) && MyEditableTextBlock.IsValid() )
+	if (SetTextInternal(InText))
 	{
-		MyEditableTextBlock->SetText(Text);
+		if (MyEditableTextBlock.IsValid())
+		{
+			MyEditableTextBlock->SetText(Text);
+		}
+		BroadcastFieldValueChanged(FFieldNotificationClassDescriptor::Text);
 	}
 }
 
@@ -171,7 +175,6 @@ bool UEditableTextBox::SetTextInternal(const FText& InText)
 	if (!Text.IdenticalTo(InText, ETextIdenticalModeFlags::DeepCompare | ETextIdenticalModeFlags::LexicalCompareInvariants))
 	{
 		Text = InText;
-		BroadcastFieldValueChanged(FFieldNotificationClassDescriptor::Text);
 		return true;
 	}
 
@@ -378,13 +381,17 @@ void UEditableTextBox::HandleOnTextChanged(const FText& InText)
 {
 	if (SetTextInternal(InText))
 	{
+		BroadcastFieldValueChanged(FFieldNotificationClassDescriptor::Text);
 		OnTextChanged.Broadcast(InText);
 	}
 }
 
 void UEditableTextBox::HandleOnTextCommitted(const FText& InText, ETextCommit::Type CommitMethod)
 {
-	SetTextInternal(InText);
+	if (SetTextInternal(InText))
+	{
+		BroadcastFieldValueChanged(FFieldNotificationClassDescriptor::Text);
+	}
 	OnTextCommitted.Broadcast(InText, CommitMethod);
 }
 PRAGMA_ENABLE_DEPRECATION_WARNINGS
diff --git a/Engine/Source/Runtime/UMG/Private/Components/MultiLineEditableText.cpp b/Engine/Source/Runtime/UMG/Private/Components/MultiLineEditableText.cpp
--- a/Engine/Source/Runtime/UMG/Private/Components/MultiLineEditableText.cpp
+++ b/Engine/Source/Runtime/UMG/Private/Components/MultiLineEditableText.cpp
@@ -163,9 +163,13 @@ FText UMultiLineEditableText::GetText() const
 
 void UMultiLineEditableText::SetText(FText InText)
 {
-	if (SetTextInternal(InText) && MyMultiLineEditableText.IsValid() )
+	if (SetTextInternal(InText))
 	{
-		MyMultiLineEditableText->SetText(Text);
+		if (MyMultiLineEditableText.IsValid())
+		{
+			MyMultiLineEditableText->SetText(Text);
+		}
+		BroadcastFieldValueChanged(FFieldNotificationClassDescriptor::Text);
 	}
 }
 
@@ -174,7 +178,6 @@ bool UMultiLineEditableText::SetTextInternal(const FText& InText)
 	if (!Text.IdenticalTo(InText, ETextIdenticalModeFlags::DeepCompare | ETextIdenticalModeFlags::LexicalCompareInvariants))
 	{
 		Text = InText;
-		BroadcastFieldValueChanged(FFieldNotificationClassDescriptor::Text);
 		return true;
 	}
 
@@ -327,13 +330,17 @@ void UMultiLineEditableText::HandleOnTextChanged(const FText& InText)
 {
 	if (SetTextInternal(InText))
 	{
+		BroadcastFieldValueChanged(FFieldNotificationClassDescriptor::Text);
 		OnTextChanged.Broadcast(InText);
 	}
 }
 
 void UMultiLineEditableText::HandleOnTextCommitted(const FText& InText, ETextCommit::Type CommitMethod)
 {
-	SetTextInternal(InText);
+	if (SetTextInternal(InText))
+	{
+		BroadcastFieldValueChanged(FFieldNotificationClassDescriptor::Text);
+	}
 	OnTextCommitted.Broadcast(InText, CommitMethod);
 }
 
diff --git a/Engine/Source/Runtime/UMG/Private/Components/MultiLineEditableTextBox.cpp b/Engine/Source/Runtime/UMG/Private/Components/MultiLineEditableTextBox.cpp
--- a/Engine/Source/Runtime/UMG/Private/Components/MultiLineEditableTextBox.cpp
+++ b/Engine/Source/Runtime/UMG/Private/Components/MultiLineEditableTextBox.cpp
@@ -174,9 +174,13 @@ FText UMultiLineEditableTextBox::GetText() const
 
 void UMultiLineEditableTextBox::SetText(FText InText)
 {
-	if (SetTextInternal(InText) && MyEditableTextBlock.IsValid() )
+	if (SetTextInternal(InText))
 	{
-		MyEditableTextBlock->SetText(Text);
+		if (MyEditableTextBlock.IsValid())
+		{
+			MyEditableTextBlock->SetText(Text);
+		}
+		BroadcastFieldValueChanged(FFieldNotificationClassDescriptor::Text);
 	}
 }
 
@@ -185,7 +189,6 @@ bool UMultiLineEditableTextBox::SetTextInternal(const FText& InText)
 	if (!Text.IdenticalTo(InText, ETextIdenticalModeFlags::DeepCompare | ETextIdenticalModeFlags::LexicalCompareInvariants))
 	{
 		Text = InText;
-		BroadcastFieldValueChanged(FFieldNotificationClassDescriptor::Text);
 		return true;
 	}
 
@@ -260,13 +263,17 @@ void UMultiLineEditableTextBox::HandleOnTextChanged(const FText& InText)
 {
 	if (SetTextInternal(InText))
 	{
+		BroadcastFieldValueChanged(FFieldNotificationClassDescriptor::Text);
 		OnTextChanged.Broadcast(InText);
 	}
 }
 
 void UMultiLineEditableTextBox::HandleOnTextCommitted(const FText& InText, ETextCommit::Type CommitMethod)
 {
-	SetTextInternal(InText);
+	if (SetTextInternal(InText))
+	{
+		BroadcastFieldValueChanged(FFieldNotificationClassDescriptor::Text);
+	}
 	OnTextCommitted.Broadcast(InText, CommitMethod);
 }
 
