PR #11091: Don't pixel snap screen-space widget components in `SWorldWidgetScreenLayer` if the root widget has snapping disabled

diff --git a/Engine/Source/Runtime/UMG/Private/Slate/SWorldWidgetScreenLayer.cpp b/Engine/Source/Runtime/UMG/Private/Slate/SWorldWidgetScreenLayer.cpp
--- a/Engine/Source/Runtime/UMG/Private/Slate/SWorldWidgetScreenLayer.cpp
+++ b/Engine/Source/Runtime/UMG/Private/Slate/SWorldWidgetScreenLayer.cpp
@@ -152,8 +152,15 @@ void SWorldWidgetScreenLayer::Tick(const FGeometry& AllottedGeometry, const doub
 					{
 						const float ViewportDist = FVector::Dist(ProjectionData.ViewOrigin, WorldLocation);
 						const FVector2D RoundedPosition2D(FMath::RoundToInt(ScreenPosition2D.X), FMath::RoundToInt(ScreenPosition2D.Y));
+
+						// If the root widget has pixel snapping disabled, then don't pixel snap the screen coordinates either otherwise
+						// it'll always jump between pixels. This saves needing an explicit flag on the widget component, and is probably 
+						// a better delegation of responsibility anyway, since changing the widget type can change the snapping as it wants
+						bool bDisablePixelSnapping = Entry.Widget->GetPixelSnapping() == EWidgetPixelSnapping::Disabled;
+						const FVector2D ScreenPositionToUse = bDisablePixelSnapping ? ScreenPosition2D : RoundedPosition2D;
+						
 						FVector2D ViewportPosition2D;
-						USlateBlueprintLibrary::ScreenToViewport(PlayerController, RoundedPosition2D, ViewportPosition2D);
+						USlateBlueprintLibrary::ScreenToViewport(PlayerController, ScreenPositionToUse, OUT ViewportPosition2D);
 
 						const FVector ViewportPosition(ViewportPosition2D.X, ViewportPosition2D.Y, ViewportDist);
 
