UMG truncation warning easy fixes

diff --git a/Engine/Source/Runtime/UMG/Private/Components/TileView.cpp b/Engine/Source/Runtime/UMG/Private/Components/TileView.cpp
--- a/Engine/Source/Runtime/UMG/Private/Components/TileView.cpp
+++ b/Engine/Source/Runtime/UMG/Private/Components/TileView.cpp
@@ -31,7 +31,7 @@ float UTileView::GetTotalEntryHeight() const
 	}
 	else
 	{
-		return EntryHeight + GetVerticalEntrySpacing() * 0.5;
+		return EntryHeight + GetVerticalEntrySpacing() * 0.5f;
 	}
 }
 
@@ -43,7 +43,7 @@ float UTileView::GetTotalEntryWidth() const
 	}
 	else
 	{
-		return EntryWidth + GetHorizontalEntrySpacing() * 0.5;
+		return EntryWidth + GetHorizontalEntrySpacing() * 0.5f;
 	}
 }
 
diff --git a/Engine/Source/Runtime/UMG/Private/Components/Viewport.cpp b/Engine/Source/Runtime/UMG/Private/Components/Viewport.cpp
--- a/Engine/Source/Runtime/UMG/Private/Components/Viewport.cpp
+++ b/Engine/Source/Runtime/UMG/Private/Components/Viewport.cpp
@@ -251,7 +251,7 @@ FLinearColor FUMGViewportClient::GetBackgroundColor() const
 
 float FUMGViewportClient::GetOrthoUnitsPerPixel(const FViewport* InViewport) const
 {
-	const float SizeX = InViewport->GetSizeXY().X;
+	const int32 SizeX = InViewport->GetSizeXY().X;
 
 	// 15.0f was coming from the CAMERA_ZOOM_DIV marco, seems it was chosen arbitrarily
 	return ( GetOrthoZoom() / ( SizeX * 15.f ) )/* * ComputeOrthoZoomFactor(SizeX)*/;
diff --git a/Engine/Source/Runtime/UMG/Private/Components/WidgetInteractionComponent.cpp b/Engine/Source/Runtime/UMG/Private/Components/WidgetInteractionComponent.cpp
--- a/Engine/Source/Runtime/UMG/Private/Components/WidgetInteractionComponent.cpp
+++ b/Engine/Source/Runtime/UMG/Private/Components/WidgetInteractionComponent.cpp
@@ -596,8 +596,11 @@ bool UWidgetInteractionComponent::PressKey(FKey Key, bool bRepeat)
 	bool bKeyCharResult = false;
 	if (bHasCharCode)
 	{
-		FCharacterEvent CharacterEvent(CharCode, ModifierKeys, VirtualUser->GetUserIndex(), bRepeat);
-		bKeyCharResult = FSlateApplication::Get().ProcessKeyCharEvent(CharacterEvent);
+		if (CharCode <= 0xD7FF || (CharCode >= 0xE000 && CharCode <= 0xFFFF)) // This is a valid UTF16 char from Basic Multilangual Plane
+		{
+			FCharacterEvent CharacterEvent(static_cast<const TCHAR>(CharCode), ModifierKeys, VirtualUser->GetUserIndex(), bRepeat);
+			bKeyCharResult = FSlateApplication::Get().ProcessKeyCharEvent(CharacterEvent);
+		}
 	}
 
 	return bDownResult || bKeyCharResult;
diff --git a/Engine/Source/Runtime/UMG/Private/Slate/SWorldWidgetScreenLayer.cpp b/Engine/Source/Runtime/UMG/Private/Slate/SWorldWidgetScreenLayer.cpp
--- a/Engine/Source/Runtime/UMG/Private/Slate/SWorldWidgetScreenLayer.cpp
+++ b/Engine/Source/Runtime/UMG/Private/Slate/SWorldWidgetScreenLayer.cpp
@@ -150,8 +150,8 @@ void SWorldWidgetScreenLayer::Tick(const FGeometry& AllottedGeometry, const doub
 
 					if (bProjected)
 					{
-						const float ViewportDist = FVector::Dist(ProjectionData.ViewOrigin, WorldLocation);
-						const FVector2D RoundedPosition2D(FMath::RoundToInt(ScreenPosition2D.X), FMath::RoundToInt(ScreenPosition2D.Y));
+						const double ViewportDist = FVector::Dist(ProjectionData.ViewOrigin, WorldLocation);
+						const FVector2D RoundedPosition2D(FMath::RoundToDouble(ScreenPosition2D.X), FMath::RoundToDouble(ScreenPosition2D.Y));
 
 						// If the root widget has pixel snapping disabled, then don't pixel snap the screen coordinates either otherwise
 						// it'll always jump between pixels. This saves needing an explicit flag on the widget component, and is probably 
@@ -185,7 +185,7 @@ void SWorldWidgetScreenLayer::Tick(const FGeometry& AllottedGeometry, const doub
 								
 								if (GSlateWorldWidgetZOrder != 0)
 								{
-									CanvasSlot->SetZOrder(-ViewportPosition.Z);
+									CanvasSlot->SetZOrder(static_cast<float>(- ViewportPosition.Z));
 								}
 							}
 							else
@@ -197,7 +197,7 @@ void SWorldWidgetScreenLayer::Tick(const FGeometry& AllottedGeometry, const doub
 
 								if (GSlateWorldWidgetZOrder != 0)
 								{
-									CanvasSlot->SetZOrder(-ViewportPosition.Z);
+									CanvasSlot->SetZOrder(static_cast<float>( - ViewportPosition.Z));
 								}
 							}
 						}
diff --git a/Engine/Source/Runtime/UMG/Private/Slate/SlateVectorArtInstanceData.cpp b/Engine/Source/Runtime/UMG/Private/Slate/SlateVectorArtInstanceData.cpp
--- a/Engine/Source/Runtime/UMG/Private/Slate/SlateVectorArtInstanceData.cpp
+++ b/Engine/Source/Runtime/UMG/Private/Slate/SlateVectorArtInstanceData.cpp
@@ -10,12 +10,12 @@ void FSlateVectorArtInstanceData::SetPositionFixedPoint16(FVector2D Position)
 	//     [0..8191].[0..7]
 	{
 		uint32 CurrentDataX = static_cast<uint32>(GetData().X) & 0xFFFF0000;
-		GetData().X = CurrentDataX | FMath::RoundToInt(Position.X*8.0f) << 0;
+		GetData().X = CurrentDataX | FMath::RoundToInt32(Position.X*8.0f) << 0;
 	}
 
 	{
 		uint32 CurrentDataY = static_cast<uint32>(GetData().Y) & 0xFFFF0000;
-		GetData().Y = CurrentDataY | FMath::RoundToInt(Position.Y*8.0f) << 0;
+		GetData().Y = CurrentDataY | FMath::RoundToInt32(Position.Y*8.0f) << 0;
 	}
 }
 
diff --git a/Engine/Source/Runtime/UMG/Private/Slate/UMGDragDropOp.cpp b/Engine/Source/Runtime/UMG/Private/Slate/UMGDragDropOp.cpp
--- a/Engine/Source/Runtime/UMG/Private/Slate/UMGDragDropOp.cpp
+++ b/Engine/Source/Runtime/UMG/Private/Slate/UMGDragDropOp.cpp
@@ -118,7 +118,7 @@ void FUMGDragDropOp::OnDragged( const class FDragDropEvent& DragDropEvent )
 	
 		if ( DeltaTime < AnimationTime )
 		{
-			float T = DeltaTime / AnimationTime;
+			double T = DeltaTime / AnimationTime;
 			FVector2D LerpPosition = ( Position - StartingScreenPos ) * T;
 			
 			DecoratorPosition = StartingScreenPos + LerpPosition;
diff --git a/Engine/Source/Runtime/UMG/Private/Slate/WidgetRenderer.cpp b/Engine/Source/Runtime/UMG/Private/Slate/WidgetRenderer.cpp
--- a/Engine/Source/Runtime/UMG/Private/Slate/WidgetRenderer.cpp
+++ b/Engine/Source/Runtime/UMG/Private/Slate/WidgetRenderer.cpp
@@ -90,7 +90,7 @@ UTextureRenderTarget2D* FWidgetRenderer::CreateTargetFor(FVector2D DrawSize, Tex
 		RenderTarget->ClearColor = FLinearColor::Transparent;
 		RenderTarget->SRGB = bIsLinearSpace;
 		RenderTarget->TargetGamma = 1;
-		RenderTarget->InitCustomFormat(DrawSize.X, DrawSize.Y, requestedFormat, bIsLinearSpace);
+		RenderTarget->InitCustomFormat(FMath::FloorToInt32(DrawSize.X), FMath::FloorToInt32(DrawSize.Y), requestedFormat, bIsLinearSpace);
 		RenderTarget->UpdateResourceImmediate(true);
 
 		return RenderTarget;
