Fix to truncation Nudge in UMG Designer

diff --git a/Engine/Source/Runtime/UMG/Private/Components/GridSlot.cpp b/Engine/Source/Runtime/UMG/Private/Components/GridSlot.cpp
--- a/Engine/Source/Runtime/UMG/Private/Components/GridSlot.cpp
+++ b/Engine/Source/Runtime/UMG/Private/Components/GridSlot.cpp
@@ -196,8 +196,8 @@ PRAGMA_ENABLE_DEPRECATION_WARNINGS
 bool UGridSlot::NudgeByDesigner(const FVector2D& NudgeDirection, const TOptional<int32>& GridSnapSize)
 {
 	const FVector2D ClampedDirection = NudgeDirection.ClampAxes(-1.0f, 1.0f);
-	const int32 NewColumn = GetColumn() + ClampedDirection.X;
-	const int32 NewRow = GetRow() + ClampedDirection.Y;
+	const int32 NewColumn = GetColumn() + FMath::TruncToInt32(ClampedDirection.X);
+	const int32 NewRow = GetRow() + FMath::TruncToInt32(ClampedDirection.Y);
 
 	if (NewColumn < 0 || NewRow < 0 || (NewColumn == GetColumn() && NewRow == GetRow()))
 	{
diff --git a/Engine/Source/Runtime/UMG/Private/Components/HorizontalBoxSlot.cpp b/Engine/Source/Runtime/UMG/Private/Components/HorizontalBoxSlot.cpp
--- a/Engine/Source/Runtime/UMG/Private/Components/HorizontalBoxSlot.cpp
+++ b/Engine/Source/Runtime/UMG/Private/Components/HorizontalBoxSlot.cpp
@@ -131,7 +131,7 @@ bool UHorizontalBoxSlot::NudgeByDesigner(const FVector2D& NudgeDirection, const
 	}
 
 	ParentHorizontalBox->Modify();
-	ParentHorizontalBox->ShiftChild(CurrentIndex + ClampedDirection.X, Content);
+	ParentHorizontalBox->ShiftChild(CurrentIndex + FMath::TruncToInt32(ClampedDirection.X), Content);
 
 	return true;
 }
diff --git a/Engine/Source/Runtime/UMG/Private/Components/UniformGridSlot.cpp b/Engine/Source/Runtime/UMG/Private/Components/UniformGridSlot.cpp
--- a/Engine/Source/Runtime/UMG/Private/Components/UniformGridSlot.cpp
+++ b/Engine/Source/Runtime/UMG/Private/Components/UniformGridSlot.cpp
@@ -111,8 +111,8 @@ PRAGMA_ENABLE_DEPRECATION_WARNINGS
 bool UUniformGridSlot::NudgeByDesigner(const FVector2D& NudgeDirection, const TOptional<int32>& GridSnapSize)
 {
 	const FVector2D ClampedDirection = NudgeDirection.ClampAxes(-1.0f, 1.0f);
-	const int32 NewColumn = GetColumn() + ClampedDirection.X;
-	const int32 NewRow = GetRow() + ClampedDirection.Y;
+	const int32 NewColumn = GetColumn() + FMath::TruncToInt32(ClampedDirection.X);
+	const int32 NewRow = GetRow() + FMath::TruncToInt32(ClampedDirection.Y);
 
 	if (NewColumn < 0 || NewRow < 0 || (NewColumn == GetColumn() && NewRow == GetRow()))
 	{
diff --git a/Engine/Source/Runtime/UMG/Private/Components/VerticalBoxSlot.cpp b/Engine/Source/Runtime/UMG/Private/Components/VerticalBoxSlot.cpp
--- a/Engine/Source/Runtime/UMG/Private/Components/VerticalBoxSlot.cpp
+++ b/Engine/Source/Runtime/UMG/Private/Components/VerticalBoxSlot.cpp
@@ -131,7 +131,7 @@ bool UVerticalBoxSlot::NudgeByDesigner(const FVector2D& NudgeDirection, const TO
 	}
 
 	ParentVerticalBox->Modify();
-	ParentVerticalBox->ShiftChild(CurrentIndex + ClampedDirection.Y, Content);
+	ParentVerticalBox->ShiftChild(CurrentIndex + FMath::TruncToInt32(ClampedDirection.Y), Content);
 
 	return true;
 }
