PR#1337 : Fix: The mouse cursor is always the default one while dragging a UMG.
#11337

diff --git a/Engine/Source/Runtime/UMG/Private/Slate/UMGDragDropOp.cpp b/Engine/Source/Runtime/UMG/Private/Slate/UMGDragDropOp.cpp
--- a/Engine/Source/Runtime/UMG/Private/Slate/UMGDragDropOp.cpp
+++ b/Engine/Source/Runtime/UMG/Private/Slate/UMGDragDropOp.cpp
@@ -153,6 +153,17 @@ FCursorReply FUMGDragDropOp::OnCursorQuery()
 		}
 	}
 
+	if (TSharedPtr<SObjectWidget> SourceUserWidgetObj = SourceUserWidget.Pin())
+	{
+		if (UUserWidget* SourceUserWidgetPtr = SourceUserWidgetObj->GetWidgetObject())
+		{
+			if (SourceUserWidgetPtr->bOverride_Cursor)
+			{
+				CursorReply = CursorReply.Cursor(SourceUserWidgetPtr->GetCursor());
+			}
+		}
+	}
+
 	return CursorReply;
 }
 
