UMG: Always cancel the async operation even when the value is the same. The last set should take ownership of the brush. The async operation may change it to something else and the last set wouldn't work.

diff --git a/Engine/Source/Runtime/UMG/Private/Components/Image.cpp b/Engine/Source/Runtime/UMG/Private/Components/Image.cpp
--- a/Engine/Source/Runtime/UMG/Private/Components/Image.cpp
+++ b/Engine/Source/Runtime/UMG/Private/Components/Image.cpp
@@ -96,6 +96,8 @@ const FSlateBrush* UImage::ConvertImage(TAttribute<FSlateBrush> InImageAsset) co
 
 void UImage::SetBrush(const FSlateBrush& InBrush)
 {
+	CancelImageStreaming();
+
 	if(Brush != InBrush)
 	{
 		Brush = InBrush;
@@ -144,6 +146,8 @@ void UImage::SetBrushTintColor(FSlateColor TintColor)
 
 void UImage::SetBrushResourceObject(UObject* ResourceObject)
 {
+	CancelImageStreaming();
+
 	if (Brush.GetResourceObject() != ResourceObject)
 	{
 		Brush.SetResourceObject(ResourceObject);
@@ -158,9 +162,10 @@ void UImage::SetBrushResourceObject(UObject* ResourceObject)
 
 void UImage::SetBrushFromAsset(USlateBrushAsset* Asset)
 {
+	CancelImageStreaming();
+
 	if(!Asset || Brush != Asset->Brush)
 	{
-		CancelImageStreaming();
 		if (Asset)
 		{
 			UImage::SetBrush(Asset->Brush);
@@ -212,9 +217,10 @@ void UImage::SetBrushFromTexture(UTexture2D* Texture, bool bMatchSize)
 
 void UImage::SetBrushFromAtlasInterface(TScriptInterface<ISlateTextureAtlasInterface> AtlasRegion, bool bMatchSize)
 {
+	CancelImageStreaming();
+
 	if(Brush.GetResourceObject() != AtlasRegion.GetObject())
 	{
-		CancelImageStreaming();
 		Brush.SetResourceObject(AtlasRegion.GetObject());
 		BroadcastFieldValueChanged(FFieldNotificationClassDescriptor::Brush);
 
@@ -240,9 +246,10 @@ void UImage::SetBrushFromAtlasInterface(TScriptInterface<ISlateTextureAtlasInter
 
 void UImage::SetBrushFromTextureDynamic(UTexture2DDynamic* Texture, bool bMatchSize)
 {
+	CancelImageStreaming();
+
 	if(Brush.GetResourceObject() != Texture)
 	{
-		CancelImageStreaming();
 		Brush.SetResourceObject(Texture);
 		BroadcastFieldValueChanged(FFieldNotificationClassDescriptor::Brush);
 
@@ -261,9 +268,10 @@ void UImage::SetBrushFromTextureDynamic(UTexture2DDynamic* Texture, bool bMatchS
 
 void UImage::SetBrushFromMaterial(UMaterialInterface* Material)
 {
+	CancelImageStreaming();
+
 	if(Brush.GetResourceObject() != Material)
 	{
-		CancelImageStreaming();
 		Brush.SetResourceObject(Material);
 		BroadcastFieldValueChanged(FFieldNotificationClassDescriptor::Brush);
 
