Fix retainer widget rendering and input issues.

We previously had 2 render paths, with each one having some issues.
With the attribute RenderWithLocalTransform set to to true, scale and absolute position were wrong in the geometry, resulting in texts scaled incorrectly and GetAccumulatedLayoutTransform returning incorrect (shifted) values.
With the attribute RenderWithLocalTransform set to to false, render transform like rotation/shear were applyed twice

There is now a single render path with the best of both world. The render path is similar to the one with RenderWithLocalTransform set to true, but it keep the scale and position info in the FGeometry that were lost in the process of getting rid of additional render transform previously.

#tests Added auto test validating different use case of the retainer, based on the different jira related to this issue. Tested retainer usage from the curve editor where keyframe drag selection was working only with RenderWithLocalTransform deactivated, and other input issue test case from UDN.

#virtualized

diff --git a/Engine/Source/Editor/CurveEditor/Private/SCurveEditorViewContainer.cpp b/Engine/Source/Editor/CurveEditor/Private/SCurveEditorViewContainer.cpp
--- a/Engine/Source/Editor/CurveEditor/Private/SCurveEditorViewContainer.cpp
+++ b/Engine/Source/Editor/CurveEditor/Private/SCurveEditorViewContainer.cpp
@@ -423,7 +423,6 @@ void SCurveEditorViewContainer::AddView(TSharedRef<SCurveEditorView> ViewToAdd)
 		SAssignNew(RetainerWidget, SRetainerWidget)
 		.RenderOnPhase(false)
 		.RenderOnInvalidation(false)
-		.RenderWithLocalTransform(false)
 		.bWarnOnInvalidSize(false)
 		[
 			SNew(SBox)
diff --git a/Engine/Source/Runtime/UMG/Private/Slate/SRetainerWidget.cpp b/Engine/Source/Runtime/UMG/Private/Slate/SRetainerWidget.cpp
--- a/Engine/Source/Runtime/UMG/Private/Slate/SRetainerWidget.cpp
+++ b/Engine/Source/Runtime/UMG/Private/Slate/SRetainerWidget.cpp
@@ -46,7 +46,7 @@ bool GSlateEnableRenderWithLocalTransform = true;
 FAutoConsoleVariableRef CVarGlateEnableRenderWithLocalTransform(
 	TEXT("Slate.EnableRetainedRenderingWithLocalTransform"),
 	GSlateEnableRenderWithLocalTransform,
-	TEXT("Whether to render with the local transform or the one passed down from the parent widget.")
+	TEXT("This console variable is no longer useful, as the SRetainerWidget render bugs workedaround by it have been fixed. It will be removed.")
 );
 
 static bool IsRetainedRenderingEnabled()
@@ -211,7 +211,6 @@ void SRetainerWidget::Construct(const FArguments& InArgs)
 
 	bEnableRetainedRenderingDesire = true;
 	bEnableRetainedRendering = false;
-	bEnableRenderWithLocalTransform = InArgs._RenderWithLocalTransform;
 	SetVolatilePrepass(bEnableRetainedRendering);
 
 #if WITH_EDITOR
@@ -736,15 +735,10 @@ TSharedRef<SWidget> SRetainerWidget::GetRootWidget()
 
 int32 SRetainerWidget::PaintSlowPath(const FSlateInvalidationContext& Context)
 {
-	if (bEnableRenderWithLocalTransform && GSlateEnableRenderWithLocalTransform)
-	{
-		FGeometry OriginalPaintSpaceGeometry = GetPaintSpaceGeometry();
-		FSlateRenderTransform SimplifiedRenderTransform(OriginalPaintSpaceGeometry.GetAccumulatedRenderTransform().GetMatrix().GetScale(), OriginalPaintSpaceGeometry.GetAccumulatedRenderTransform().GetTranslation());
-		const FGeometry NewPaintSpaceGeometry = FGeometry::MakeRoot(OriginalPaintSpaceGeometry.GetLocalSize(), FSlateLayoutTransform()).MakeChild(SimplifiedRenderTransform, FVector2D::ZeroVector);
-		return SCompoundWidget::OnPaint(*Context.PaintArgs, NewPaintSpaceGeometry, Context.CullingRect, *Context.WindowElementList, Context.IncomingLayerId, Context.WidgetStyle, Context.bParentEnabled);
-	}
-	else
-	{
-		return SCompoundWidget::OnPaint(*Context.PaintArgs, GetPaintSpaceGeometry(), Context.CullingRect, *Context.WindowElementList, Context.IncomingLayerId, Context.WidgetStyle, Context.bParentEnabled);
-	}
+	const FGeometry& OriginalPaintSpaceGeometry = GetPaintSpaceGeometry();
+	FSlateRenderTransform SimplifiedRenderTransform(OriginalPaintSpaceGeometry.GetAccumulatedRenderTransform().GetMatrix().GetScale(), OriginalPaintSpaceGeometry.GetAccumulatedRenderTransform().GetTranslation());
+	FTransform2D ScalePosTransform(OriginalPaintSpaceGeometry.Scale, OriginalPaintSpaceGeometry.AbsolutePosition);
+	SimplifiedRenderTransform = Concatenate(SimplifiedRenderTransform, Inverse(ScalePosTransform));
+	const FGeometry NewPaintSpaceGeometry = FGeometry::MakeRoot(OriginalPaintSpaceGeometry.GetLocalSize(), FSlateLayoutTransform(OriginalPaintSpaceGeometry.Scale, OriginalPaintSpaceGeometry.AbsolutePosition)).MakeChild(SimplifiedRenderTransform, FVector2D::ZeroVector);
+	return SCompoundWidget::OnPaint(*Context.PaintArgs, NewPaintSpaceGeometry, Context.CullingRect, *Context.WindowElementList, Context.IncomingLayerId, Context.WidgetStyle, Context.bParentEnabled);
 }
\ No newline at end of file
diff --git a/Engine/Source/Runtime/UMG/Public/Slate/SRetainerWidget.h b/Engine/Source/Runtime/UMG/Public/Slate/SRetainerWidget.h
--- a/Engine/Source/Runtime/UMG/Public/Slate/SRetainerWidget.h
+++ b/Engine/Source/Runtime/UMG/Public/Slate/SRetainerWidget.h
@@ -49,12 +49,11 @@ public:
 		_PhaseCount = 1;
 		_RenderOnPhase = true;
 		_RenderOnInvalidation = false;
-		_RenderWithLocalTransform = true;
 	}
 	SLATE_DEFAULT_SLOT(FArguments, Content)
 		SLATE_ARGUMENT(bool, RenderOnPhase)
 		SLATE_ARGUMENT(bool, RenderOnInvalidation)
-		SLATE_ARGUMENT(bool, RenderWithLocalTransform)
+		SLATE_ARGUMENT_DEPRECATED(bool, RenderWithLocalTransform, 5.4, "This argument has no effect anymore, because bugs from each render path (with RenderWithLocalTransform set to true or false) have been fixed, making it no longer necessary.")
 		SLATE_ARGUMENT(int32, Phase)
 		SLATE_ARGUMENT(int32, PhaseCount)
 		SLATE_ARGUMENT(FName, StatId)
@@ -154,7 +153,6 @@ private:
 
 	bool bEnableRetainedRenderingDesire;
 	bool bEnableRetainedRendering;
-	bool bEnableRenderWithLocalTransform;
 
 #if WITH_EDITOR
 	/** True if widget is used in design time */
