UMG truncation fixes

diff --git a/Engine/Source/Runtime/UMG/Private/Animation/WidgetAnimation.cpp b/Engine/Source/Runtime/UMG/Private/Animation/WidgetAnimation.cpp
--- a/Engine/Source/Runtime/UMG/Private/Animation/WidgetAnimation.cpp
+++ b/Engine/Source/Runtime/UMG/Private/Animation/WidgetAnimation.cpp
@@ -91,12 +91,12 @@ ETrackSupport UWidgetAnimation::IsTrackSupported(TSubclassOf<class UMovieSceneTr
 
 float UWidgetAnimation::GetStartTime() const
 {
-	return MovieScene->GetPlaybackRange().GetLowerBoundValue() / MovieScene->GetTickResolution();
+	return static_cast<float>(MovieScene->GetPlaybackRange().GetLowerBoundValue() / MovieScene->GetTickResolution());
 }
 
 float UWidgetAnimation::GetEndTime() const
 {
-	return MovieScene->GetPlaybackRange().GetUpperBoundValue() / MovieScene->GetTickResolution();
+	return static_cast<float>(MovieScene->GetPlaybackRange().GetUpperBoundValue() / MovieScene->GetTickResolution());
 }
 
 void UWidgetAnimation::BindToAnimationStarted(UUserWidget* Widget, FWidgetAnimationDynamicEvent Delegate)
diff --git a/Engine/Source/Runtime/UMG/Private/Blueprint/AsyncTaskDownloadImage.cpp b/Engine/Source/Runtime/UMG/Private/Blueprint/AsyncTaskDownloadImage.cpp
--- a/Engine/Source/Runtime/UMG/Private/Blueprint/AsyncTaskDownloadImage.cpp
+++ b/Engine/Source/Runtime/UMG/Private/Blueprint/AsyncTaskDownloadImage.cpp
@@ -70,13 +70,17 @@ void UAsyncTaskDownloadImage::HandleImageRequest(FHttpRequestPtr HttpRequest, FH
 
 		for ( auto ImageWrapper : ImageWrappers )
 		{
-			if ( ImageWrapper.IsValid() && ImageWrapper->SetCompressed(HttpResponse->GetContent().GetData(), HttpResponse->GetContent().Num()) )
+			if ( ImageWrapper.IsValid() && 
+				 ImageWrapper->SetCompressed(HttpResponse->GetContent().GetData(), HttpResponse->GetContent().Num()) &&
+				 ImageWrapper->GetWidth() <= TNumericLimits<int32>::Max() && 
+				 ImageWrapper->GetHeight() <= TNumericLimits<int32>::Max())
 			{
+				
 				TArray64<uint8> RawData;
 				const ERGBFormat InFormat = ERGBFormat::BGRA;
 				if ( ImageWrapper->GetRaw(InFormat, 8, RawData) )
-				{
-					if ( UTexture2DDynamic* Texture = UTexture2DDynamic::Create(ImageWrapper->GetWidth(), ImageWrapper->GetHeight()) )
+				{					
+					if ( UTexture2DDynamic* Texture = UTexture2DDynamic::Create(static_cast<int32>(ImageWrapper->GetWidth()), static_cast<int32>(ImageWrapper->GetHeight())) )
 					{
 						Texture->SRGB = true;
 						Texture->UpdateResource();
diff --git a/Engine/Source/Runtime/UMG/Private/Blueprint/GameViewportSubsystem.cpp b/Engine/Source/Runtime/UMG/Private/Blueprint/GameViewportSubsystem.cpp
--- a/Engine/Source/Runtime/UMG/Private/Blueprint/GameViewportSubsystem.cpp
+++ b/Engine/Source/Runtime/UMG/Private/Blueprint/GameViewportSubsystem.cpp
@@ -260,7 +260,7 @@ void UGameViewportSubsystem::SetWidgetSlot(UWidget* Widget, FGameViewportWidgetS
 			SlotInfo.FullScreenWidgetSlot->SetAutoSize(OffsetArgument.Get<1>());
 			SlotInfo.FullScreenWidgetSlot->SetAnchors(Slot.Anchors);
 			SlotInfo.FullScreenWidgetSlot->SetAlignment(Slot.Alignment);
-			SlotInfo.FullScreenWidgetSlot->SetZOrder(Slot.ZOrder);
+			SlotInfo.FullScreenWidgetSlot->SetZOrder(static_cast<float>(Slot.ZOrder));
 			WidgetHost->Invalidate(EInvalidateWidgetReason::Layout);
 		}
 	}
@@ -274,8 +274,9 @@ FGameViewportWidgetSlot UGameViewportSubsystem::SetWidgetSlotPosition(FGameViewp
 		Position /= Scale;
 	}
 
-	Slot.Offsets.Left = Position.X;
-	Slot.Offsets.Top = Position.Y;
+	const FVector2f NewPosition = UE::Slate::CastToVector2f(Position);
+	Slot.Offsets.Left = NewPosition.X;
+	Slot.Offsets.Top = NewPosition.Y;
 	Slot.Anchors = FAnchors(0.f, 0.f);
 
 	return Slot;
@@ -283,8 +284,9 @@ FGameViewportWidgetSlot UGameViewportSubsystem::SetWidgetSlotPosition(FGameViewp
 
 FGameViewportWidgetSlot UGameViewportSubsystem::SetWidgetSlotDesiredSize(FGameViewportWidgetSlot Slot, FVector2D Size)
 {
-	Slot.Offsets.Right = Size.X;
-	Slot.Offsets.Bottom = Size.Y;
+	const FVector2f NewSize = UE::Slate::CastToVector2f(Size);
+	Slot.Offsets.Right = NewSize.X;
+	Slot.Offsets.Bottom = NewSize.Y;
 	Slot.Anchors = FAnchors(0.f, 0.f);
 	return Slot;
 }
diff --git a/Engine/Source/Runtime/UMG/Private/Components/CanvasPanelSlot.cpp b/Engine/Source/Runtime/UMG/Private/Components/CanvasPanelSlot.cpp
--- a/Engine/Source/Runtime/UMG/Private/Components/CanvasPanelSlot.cpp
+++ b/Engine/Source/Runtime/UMG/Private/Components/CanvasPanelSlot.cpp
@@ -158,8 +158,9 @@ FAnchorData UCanvasPanelSlot::GetLayout() const
 
 void UCanvasPanelSlot::SetPosition(FVector2D InPosition)
 {
-	LayoutData.Offsets.Left = InPosition.X;
-	LayoutData.Offsets.Top = InPosition.Y;
+	FVector2f Position = UE::Slate::CastToVector2f(InPosition);
+	LayoutData.Offsets.Left = Position.X;
+	LayoutData.Offsets.Top = Position.Y;
 
 	if ( Slot )
 	{
@@ -180,8 +181,9 @@ FVector2D UCanvasPanelSlot::GetPosition() const
 
 void UCanvasPanelSlot::SetSize(FVector2D InSize)
 {
-	LayoutData.Offsets.Right = InSize.X;
-	LayoutData.Offsets.Bottom = InSize.Y;
+	FVector2f Size = UE::Slate::CastToVector2f(InSize);
+	LayoutData.Offsets.Right = Size.X;
+	LayoutData.Offsets.Bottom = Size.Y;
 
 	if ( Slot )
 	{
@@ -281,7 +283,7 @@ void UCanvasPanelSlot::SetZOrder(int32 InZOrder)
 	ZOrder = InZOrder;
 	if ( Slot )
 	{
-		Slot->SetZOrder(InZOrder);
+		Slot->SetZOrder(static_cast<float>(InZOrder));
 	}
 }
 
@@ -289,7 +291,7 @@ int32 UCanvasPanelSlot::GetZOrder() const
 {
 	if ( Slot )
 	{
-		return Slot->GetZOrder();
+		return FMath::TruncToInt32(Slot->GetZOrder());
 	}
 
 	return ZOrder;
@@ -390,15 +392,15 @@ void UCanvasPanelSlot::RebaseLayout(bool PreserveSize)
 				LocalLayoutData.Anchors.Minimum.Y * CanvasSize.Y,
 				LocalLayoutData.Anchors.Maximum.X * CanvasSize.X,
 				LocalLayoutData.Anchors.Maximum.Y * CanvasSize.Y);
-			FVector2D DefaultAnchorPosition = FVector2D(AnchorPositions.Left, AnchorPositions.Top);
+			FVector2f DefaultAnchorPosition = FVector2f(AnchorPositions.Left, AnchorPositions.Top);
 
 			// Determine the amount that would be offset from the anchor position if alignment was applied.
-			FVector2D AlignmentOffset = LocalLayoutData.Alignment * PreEditGeometry.Size;
+			FVector2f AlignmentOffset = UE::Slate::CastToVector2f(LocalLayoutData.Alignment) * PreEditGeometry.Size;
 
-			FVector2D MoveDelta = FVector2D(Geometry.Position) - FVector2D(PreEditGeometry.Position);
+			FVector2f MoveDelta = Geometry.Position - PreEditGeometry.Position;
 
 			// Determine where the widget's new position needs to be to maintain a stable location when the anchors change.
-			FVector2D LeftTopDelta = FVector2D(PreEditGeometry.Position) - FVector2D(DefaultAnchorPosition);
+			FVector2f LeftTopDelta = PreEditGeometry.Position - DefaultAnchorPosition;
 
 			const bool bAnchorsMoved = PreEditLayoutData.Anchors.Minimum != LocalLayoutData.Anchors.Minimum || PreEditLayoutData.Anchors.Maximum != LocalLayoutData.Anchors.Maximum;
 			const bool bMoved = PreEditLayoutData.Offsets.Left != LocalLayoutData.Offsets.Left || PreEditLayoutData.Offsets.Top != LocalLayoutData.Offsets.Top;
@@ -456,7 +458,7 @@ void UCanvasPanelSlot::RebaseLayout(bool PreserveSize)
 			}
 			else if ( DesiredPosition.IsSet() )
 			{
-				FVector2D NewLocalPosition = DesiredPosition.GetValue();
+				FVector2f NewLocalPosition = UE::Slate::CastToVector2f(DesiredPosition.GetValue());
 
 				LocalLayoutData.Offsets.Left = NewLocalPosition.X - AnchorPositions.Left;
 				LocalLayoutData.Offsets.Top = NewLocalPosition.Y - AnchorPositions.Top;
diff --git a/Engine/Source/Runtime/UMG/Private/Components/DynamicEntryBoxBase.cpp b/Engine/Source/Runtime/UMG/Private/Components/DynamicEntryBoxBase.cpp
--- a/Engine/Source/Runtime/UMG/Private/Components/DynamicEntryBoxBase.cpp
+++ b/Engine/Source/Runtime/UMG/Private/Components/DynamicEntryBoxBase.cpp
@@ -127,6 +127,7 @@ void UDynamicEntryBoxBase::RemoveEntryInternal(UUserWidget* EntryWidget)
 void UDynamicEntryBoxBase::SetEntrySpacing(const FVector2D& InEntrySpacing)
 {
 	EntrySpacing = InEntrySpacing;
+	FVector2f Spacing = UE::Slate::CastToVector2f(EntrySpacing);
 
 	if (MyPanelWidget.IsValid())
 	{
@@ -143,13 +144,14 @@ void UDynamicEntryBoxBase::SetEntrySpacing(const FVector2D& InEntrySpacing)
 				FMargin Padding;
 				if (SpacingPattern.Num() > 0)
 				{
-					FVector2D Spacing(0.f, 0.f);
+					// Initialize to no spacing
+					Spacing = FVector2f(0.f, 0.f);
 
 					// First establish the starting location
 					for (int32 CountIdx = 0; CountIdx < ChildIdx; ++CountIdx)
 					{
 						int32 PatternIdx = CountIdx % SpacingPattern.Num();
-						Spacing += SpacingPattern[PatternIdx];
+						Spacing += UE::Slate::CastToVector2f(SpacingPattern[PatternIdx]);
 					}
 					// Negative padding is no good, so negative spacing is expressed as positive spacing on the opposite side
 					if (Spacing.X >= 0.f)
@@ -160,6 +162,8 @@ void UDynamicEntryBoxBase::SetEntrySpacing(const FVector2D& InEntrySpacing)
 					{
 						Padding.Right = -Spacing.X;
 					}
+
+					// Negative padding is no good, so negative spacing is expressed as positive spacing on the opposite side
 					if (Spacing.Y >= 0.f)
 					{
 						Padding.Top = Spacing.Y;
@@ -171,22 +175,24 @@ void UDynamicEntryBoxBase::SetEntrySpacing(const FVector2D& InEntrySpacing)
 				}
 				else
 				{
-					if (EntrySpacing.X >= 0.f)
+					// Negative padding is no good, so negative spacing is expressed as positive spacing on the opposite side
+					if (Spacing.X >= 0.f)
 					{
-						Padding.Left = ChildIdx * EntrySpacing.X;
+						Padding.Left = ChildIdx * Spacing.X;
 					}
 					else
 					{
-						Padding.Right = ChildIdx * -EntrySpacing.X;
+						Padding.Right = ChildIdx * -Spacing.X;
 					}
 
-					if (EntrySpacing.Y >= 0.f)
+					// Negative padding is no good, so negative spacing is expressed as positive spacing on the opposite side
+					if (Spacing.Y >= 0.f)
 					{
-						Padding.Top = ChildIdx * EntrySpacing.Y;
+						Padding.Top = ChildIdx * Spacing.Y;
 					}
 					else
 					{
-						Padding.Bottom = ChildIdx * -EntrySpacing.Y;
+						Padding.Bottom = ChildIdx * -Spacing.Y;
 					}
 				}
 				SOverlay::FOverlaySlot& OverlaySlot = (*OverlayChildren)[ChildIdx];
@@ -203,8 +209,8 @@ void UDynamicEntryBoxBase::SetEntrySpacing(const FVector2D& InEntrySpacing)
 				const bool bIsFirstChild = ChildIdx == 0;
 
 				FMargin Padding;
-				Padding.Top = bIsHBox || bIsFirstChild ? 0.f : EntrySpacing.Y;
-				Padding.Left = bIsHBox && !bIsFirstChild ? EntrySpacing.X : 0.f;
+				Padding.Top = bIsHBox || bIsFirstChild ? 0.f : Spacing.Y;
+				Padding.Left = bIsHBox && !bIsFirstChild ? Spacing.X : 0.f;
 				SBoxPanel::FSlot& BoxSlot = (*BoxChildren)[ChildIdx];
 				BoxSlot.SetPadding(Padding);
 			}
@@ -387,9 +393,10 @@ UUserWidget* UDynamicEntryBoxBase::CreateEntryInternal(TSubclassOf<UUserWidget>
 	return NewEntryWidget;
 }
 
-FMargin UDynamicEntryBoxBase::BuildEntryPadding(const FVector2D& DesiredSpacing)
+FMargin UDynamicEntryBoxBase::BuildEntryPadding(const FVector2D& InDesiredSpacing)
 {
 	FMargin EntryPadding;
+	const FVector2f DesiredSpacing = UE::Slate::CastToVector2f(InDesiredSpacing);
 	if (DesiredSpacing.X >= 0.f)
 	{
 		EntryPadding.Left = DesiredSpacing.X;
@@ -472,9 +479,10 @@ void UDynamicEntryBoxBase::AddEntryChild(UUserWidget& ChildWidget)
 			const bool bIsHBox = EntryBoxType == EDynamicBoxType::Horizontal;
 			const bool bIsFirstChild = MyPanelWidget->GetChildren()->Num() == 0;
 
+			const FVector2f Spacing = UE::Slate::CastToVector2f(EntrySpacing);
 			FMargin Padding;
-			Padding.Top = bIsHBox || bIsFirstChild ? 0.f : EntrySpacing.Y;
-			Padding.Left = bIsHBox && !bIsFirstChild ? EntrySpacing.X : 0.f;
+			Padding.Top = bIsHBox || bIsFirstChild ? 0.f : Spacing.Y;
+			Padding.Left = bIsHBox && !bIsFirstChild ? Spacing.X : 0.f;
 
 			if (bIsHBox)
 			{
diff --git a/Engine/Source/Runtime/UMG/Private/Components/GridPanel.cpp b/Engine/Source/Runtime/UMG/Private/Components/GridPanel.cpp
--- a/Engine/Source/Runtime/UMG/Private/Components/GridPanel.cpp
+++ b/Engine/Source/Runtime/UMG/Private/Components/GridPanel.cpp
@@ -109,7 +109,7 @@ void UGridPanel::SetColumnFill(int32 ColumnIndex, float Coefficient)
 {
 	while (ColumnFill.Num() <= ColumnIndex)
 	{
-		ColumnFill.Emplace(0);
+		ColumnFill.Emplace(0.0f);
 	}
 	ColumnFill[ColumnIndex] = Coefficient;
 
@@ -123,7 +123,7 @@ void UGridPanel::SetRowFill(int32 RowIndex, float Coefficient)
 {
 	while (RowFill.Num() <= RowIndex)
 	{
-		RowFill.Emplace(0);
+		RowFill.Emplace(0.0f);
 	}
 	RowFill[RowIndex] = Coefficient;
 
diff --git a/Engine/Source/Runtime/UMG/Private/Components/Image.cpp b/Engine/Source/Runtime/UMG/Private/Components/Image.cpp
--- a/Engine/Source/Runtime/UMG/Private/Components/Image.cpp
+++ b/Engine/Source/Runtime/UMG/Private/Components/Image.cpp
@@ -194,8 +194,8 @@ void UImage::SetBrushFromTexture(UTexture2D* Texture, bool bMatchSize)
 #if WITH_EDITOR
 				FTextureCompilingManager::Get().FinishCompilation({ Texture });
 #endif
-				Brush.ImageSize.X = Texture->GetSizeX();
-				Brush.ImageSize.Y = Texture->GetSizeY();
+				Brush.ImageSize.X = static_cast<float>(Texture->GetSizeX());
+				Brush.ImageSize.Y = static_cast<float>(Texture->GetSizeY());
 			}
 			else
 			{
@@ -248,8 +248,8 @@ void UImage::SetBrushFromTextureDynamic(UTexture2DDynamic* Texture, bool bMatchS
 
 		if (bMatchSize && Texture)
 		{
-			Brush.ImageSize.X = Texture->SizeX;
-			Brush.ImageSize.Y = Texture->SizeY;
+			Brush.ImageSize.X = static_cast<float>(Texture->SizeX);
+			Brush.ImageSize.Y = static_cast<float>(Texture->SizeY);
 		}
 
 		if (MyImage.IsValid())
diff --git a/Engine/Source/Runtime/UMG/Private/Components/RichTextBlockDecorator.cpp b/Engine/Source/Runtime/UMG/Private/Components/RichTextBlockDecorator.cpp
--- a/Engine/Source/Runtime/UMG/Private/Components/RichTextBlockDecorator.cpp
+++ b/Engine/Source/Runtime/UMG/Private/Components/RichTextBlockDecorator.cpp
@@ -46,7 +46,7 @@ TSharedRef<ISlateRun> FRichTextDecorator::Create(const TSharedRef<class FTextLay
 		TAttribute<int16> GetBaseline = TAttribute<int16>::CreateLambda([Font, ShadowOffsetY]()
 		{
 			const TSharedRef<FSlateFontMeasure> FontMeasure = FSlateApplication::Get().GetRenderer()->GetFontMeasureService();
-			return FontMeasure->GetBaseline(Font) - ShadowOffsetY;
+			return static_cast<int16>(FontMeasure->GetBaseline(Font) - ShadowOffsetY);
 		});
 
 		FSlateWidgetRun::FWidgetRunInfo WidgetRunInfo(DecoratorWidget.ToSharedRef(), GetBaseline);
diff --git a/Engine/Source/Runtime/UMG/Private/Components/RichTextBlockImageDecorator.cpp b/Engine/Source/Runtime/UMG/Private/Components/RichTextBlockImageDecorator.cpp
--- a/Engine/Source/Runtime/UMG/Private/Components/RichTextBlockImageDecorator.cpp
+++ b/Engine/Source/Runtime/UMG/Private/Components/RichTextBlockImageDecorator.cpp
@@ -40,13 +40,13 @@ public:
 
 		if (Height.IsSet())
 		{
-			IconHeight = Height.GetValue();
+			IconHeight = static_cast<float>(Height.GetValue());
 		}
 
 		float IconWidth = IconHeight;
 		if (Width.IsSet())
 		{
-			IconWidth = Width.GetValue();
+			IconWidth = static_cast<float>(Width.GetValue());
 		}
 
 		ChildSlot
@@ -125,7 +125,7 @@ protected:
 				{
 					if (FCString::Stricmp(GetData(*WidthString), TEXT("desired")) == 0)
 					{
-						Width = Brush->ImageSize.X;
+						Width = FMath::TruncToInt32(Brush->ImageSize.X);
 					}
 				}
 			}
@@ -142,7 +142,7 @@ protected:
 				{
 					if (FCString::Stricmp(GetData(*HeightString), TEXT("desired")) == 0)
 					{
-						Height = Brush->ImageSize.Y;
+						Height = FMath::TruncToInt32(Brush->ImageSize.Y);
 					}
 				}
 			}
diff --git a/Engine/Source/Runtime/UMG/Private/SlateBlueprintLibrary.cpp b/Engine/Source/Runtime/UMG/Private/SlateBlueprintLibrary.cpp
--- a/Engine/Source/Runtime/UMG/Private/SlateBlueprintLibrary.cpp
+++ b/Engine/Source/Runtime/UMG/Private/SlateBlueprintLibrary.cpp
@@ -145,12 +145,12 @@ void USlateBlueprintLibrary::ScreenToViewport(UObject* WorldContextObject, FVect
 
 float USlateBlueprintLibrary::TransformScalarAbsoluteToLocal(const FGeometry& Geometry, float AbsoluteScalar)
 {
-	return Geometry.GetAccumulatedRenderTransform().TransformVector(FVector2D(AbsoluteScalar, 0)).Size();
+	return Geometry.GetAccumulatedRenderTransform().TransformVector(FVector2f(AbsoluteScalar, 0.0f)).Size();
 }
 
 float USlateBlueprintLibrary::TransformScalarLocalToAbsolute(const FGeometry& Geometry, float LocalScalar)
 {
-	return Inverse(Geometry.GetAccumulatedRenderTransform()).TransformVector(FVector2D(LocalScalar, 0)).Size();
+	return Inverse(Geometry.GetAccumulatedRenderTransform()).TransformVector(FVector2f(LocalScalar, 0.0f)).Size();
 }
 
 FVector2D USlateBlueprintLibrary::TransformVectorAbsoluteToLocal(const FGeometry& Geometry, FVector2D AbsoluteVector)
diff --git a/Engine/Source/Runtime/UMG/Private/WidgetBlueprintLibrary.cpp b/Engine/Source/Runtime/UMG/Private/WidgetBlueprintLibrary.cpp
--- a/Engine/Source/Runtime/UMG/Private/WidgetBlueprintLibrary.cpp
+++ b/Engine/Source/Runtime/UMG/Private/WidgetBlueprintLibrary.cpp
@@ -670,7 +670,7 @@ void UWidgetBlueprintLibrary::GetSafeZonePadding(UObject* WorldContextObject, FV
 
 void UWidgetBlueprintLibrary::SetColorVisionDeficiencyType(EColorVisionDeficiency Type, float Severity, bool CorrectDeficiency, bool ShowCorrectionWithDeficiency)
 {
-	int32 AdjustedSeverity = FMath::Clamp(Severity, 0.f, 1.f) * 10;
+	int32 AdjustedSeverity = FMath::TruncToInt32(FMath::Clamp(Severity, 0.f, 1.f) * 10);
 	FSlateApplicationBase::Get().GetRenderer()->SetColorVisionDeficiencyType(Type, AdjustedSeverity, CorrectDeficiency, ShowCorrectionWithDeficiency);
 }
 
diff --git a/Engine/Source/Runtime/UMG/Private/WidgetLayoutLibrary.cpp b/Engine/Source/Runtime/UMG/Private/WidgetLayoutLibrary.cpp
--- a/Engine/Source/Runtime/UMG/Private/WidgetLayoutLibrary.cpp
+++ b/Engine/Source/Runtime/UMG/Private/WidgetLayoutLibrary.cpp
@@ -64,7 +64,7 @@ bool UWidgetLayoutLibrary::ProjectWorldLocationToWidgetPositionWithDistance(APla
 			// I do this before I remove scaling, because scaling is going to be applied later
 			// in the opposite direction, so as long as we round, before inverse scale, scale should
 			// result in more or less the same value, especially after slate does layout rounding.
-			FVector2D ScreenPosition(FMath::RoundToInt(PixelLocation.X), FMath::RoundToInt(PixelLocation.Y));
+			FVector2D ScreenPosition(FMath::RoundToDouble(PixelLocation.X), FMath::RoundToDouble(PixelLocation.Y));
 
 			FVector2D ViewportPosition2D;
 			USlateBlueprintLibrary::ScreenToViewport(PlayerController, ScreenPosition, ViewportPosition2D);
@@ -96,7 +96,7 @@ float UWidgetLayoutLibrary::GetViewportScale(const UObject* WorldContextObject)
 			{
 				FVector2D ViewportSize;
 				ViewportClient->GetViewportSize(ViewportSize);
-				ViewportScale = GetDefault<UUserInterfaceSettings>()->GetDPIScaleBasedOnSize(FIntPoint(ViewportSize.X, ViewportSize.Y));
+				ViewportScale = GetDefault<UUserInterfaceSettings>()->GetDPIScaleBasedOnSize(FIntPoint(FMath::TruncToInt32(ViewportSize.X), FMath::TruncToInt32(ViewportSize.Y)));
 			}
 		}
 
@@ -111,7 +111,7 @@ float UWidgetLayoutLibrary::GetViewportScale(const UGameViewportClient* Viewport
 	FVector2D ViewportSize;
 	ViewportClient->GetViewportSize(ViewportSize);
 	
-	float UserResolutionScale = GetDefault<UUserInterfaceSettings>()->GetDPIScaleBasedOnSize(FIntPoint(ViewportSize.X, ViewportSize.Y));
+	float UserResolutionScale = GetDefault<UUserInterfaceSettings>()->GetDPIScaleBasedOnSize(FIntPoint(FMath::TruncToInt32(ViewportSize.X), FMath::TruncToInt32(ViewportSize.Y)));
 
 	// Normally we'd factor in native DPI Scale here too, but because the SGameLayerManager already
 	// accounts for the native DPI scale, and extracts it from the calculations, the UMG/Slate portion of the
diff --git a/Engine/Source/Runtime/UMG/Public/Components/ListViewBase.h b/Engine/Source/Runtime/UMG/Public/Components/ListViewBase.h
--- a/Engine/Source/Runtime/UMG/Public/Components/ListViewBase.h
+++ b/Engine/Source/Runtime/UMG/Public/Components/ListViewBase.h
@@ -429,9 +429,9 @@ private:
 	{
 		if (SListView<ItemType>* MyListView = GetMyListView())
 		{
-			const FVector2D DistanceRemaining = MyListView->GetScrollDistanceRemaining();
-			OnListViewScrolledInternal(OffsetInItems, DistanceRemaining.Y);
-			OnListViewScrolled().Broadcast(OffsetInItems, DistanceRemaining.Y);
+			const FVector2f DistanceRemaining = UE::Slate::CastToVector2f(MyListView->GetScrollDistanceRemaining());
+			OnListViewScrolledInternal(static_cast<float>(OffsetInItems), DistanceRemaining.Y);
+			OnListViewScrolled().Broadcast(static_cast<float>(OffsetInItems), DistanceRemaining.Y);
 		}
 	}
 
