WidgetComponent Truncation Warning Fixes

diff --git a/Engine/Source/Runtime/UMG/Private/Components/WidgetComponent.cpp b/Engine/Source/Runtime/UMG/Private/Components/WidgetComponent.cpp
--- a/Engine/Source/Runtime/UMG/Private/Components/WidgetComponent.cpp
+++ b/Engine/Source/Runtime/UMG/Private/Components/WidgetComponent.cpp
@@ -372,10 +372,10 @@ public:
 			{
 				if (GeometryMode == EWidgetGeometryMode::Plane)
 				{
-					float U = -RenderTarget->SizeX * Pivot.X;
-					float V = -RenderTarget->SizeY * Pivot.Y;
-					float UL = RenderTarget->SizeX * (1.0f - Pivot.X);
-					float VL = RenderTarget->SizeY * (1.0f - Pivot.Y);
+					float U = -RenderTarget->SizeX * static_cast<float>(Pivot.X);
+					float V = -RenderTarget->SizeY * static_cast<float>(Pivot.Y);
+					float UL = RenderTarget->SizeX * (1.0f - static_cast<float>(Pivot.X));
+					float VL = RenderTarget->SizeY * (1.0f - static_cast<float>(Pivot.Y));
 
 					int32 VertexIndices[4];
 
@@ -408,13 +408,13 @@ public:
 					const int32 NumSegments = FMath::Lerp(4, 32, ArcAngle/PI);
 
 
-					const float Radius = RenderTarget->SizeX / ArcAngle;
-					const float Apothem = Radius * FMath::Cos(0.5f*ArcAngle);
-					const float ChordLength = 2.0f * Radius * FMath::Sin(0.5f*ArcAngle);
+					const double Radius = RenderTarget->SizeX / ArcAngle;
+					const double Apothem = Radius * FMath::Cos(0.5*ArcAngle);
+					const double ChordLength = 2.0f * Radius * FMath::Sin(0.5*ArcAngle);
 					
-					const float PivotOffsetX = ChordLength * (0.5-Pivot.X);
-					const float V = -RenderTarget->SizeY * Pivot.Y;
-					const float VL = RenderTarget->SizeY * (1.0f - Pivot.Y);
+					const double PivotOffsetX = ChordLength * (0.5-Pivot.X);
+					const double V = -RenderTarget->SizeY * Pivot.Y;
+					const double VL = RenderTarget->SizeY * (1.0 - Pivot.Y);
 
 					int32 VertexIndices[4];
 
@@ -424,7 +424,7 @@ public:
 
 						if (VisibilityMap & (1 << ViewIndex))
 						{
-							const float RadiansPerStep = ArcAngle / NumSegments;
+							const double RadiansPerStep = ArcAngle / NumSegments;
 
 							FVector LastTangentX;
 							FVector LastTangentY;
@@ -432,14 +432,14 @@ public:
 
 							for (int32 Segment = 0; Segment < NumSegments; Segment++ )
 							{
-								const float Angle = -ArcAngle / 2 + Segment * RadiansPerStep;
-								const float NextAngle = Angle + RadiansPerStep;
+								const double Angle = -ArcAngle / 2 + Segment * RadiansPerStep;
+								const double NextAngle = Angle + RadiansPerStep;
 								
 								// Polar to Cartesian
-								const float X0 = Radius * FMath::Cos(Angle) - Apothem;
-								const float Y0 = Radius * FMath::Sin(Angle);
-								const float X1 = Radius * FMath::Cos(NextAngle) - Apothem;
-								const float Y1 = Radius * FMath::Sin(NextAngle);
+								const double X0 = Radius * FMath::Cos(Angle) - Apothem;
+								const double Y0 = Radius * FMath::Sin(Angle);
+								const double X1 = Radius * FMath::Cos(NextAngle) - Apothem;
+								const double Y1 = Radius * FMath::Sin(NextAngle);
 
 								const float U0 = static_cast<float>(Segment) / NumSegments;
 								const float U1 = static_cast<float>(Segment+1) / NumSegments;
@@ -586,11 +586,11 @@ public:
 
 	virtual uint32 GetMemoryFootprint(void) const override { return(sizeof(*this) + GetAllocatedSize()); }
 
-	uint32 GetAllocatedSize(void) const { return( FPrimitiveSceneProxy::GetAllocatedSize() ); }
+	SIZE_T GetAllocatedSize(void) const { return( FPrimitiveSceneProxy::GetAllocatedSize() ); }
 
 private:
 	FVector Origin;
-	float ArcAngle;
+	double ArcAngle;
 	FVector2D Pivot;
 	ISlate3DRenderer& Renderer;
 	UTextureRenderTarget2D* RenderTarget;
@@ -848,7 +848,7 @@ FPrimitiveSceneProxy* UWidgetComponent::CreateSceneProxy()
 			return Result;
 		}
 		virtual uint32 GetMemoryFootprint(void) const override { return(sizeof(*this) + GetAllocatedSize()); }
-		uint32 GetAllocatedSize(void) const { return(FPrimitiveSceneProxy::GetAllocatedSize()); }
+		SIZE_T GetAllocatedSize(void) const { return(FPrimitiveSceneProxy::GetAllocatedSize()); }
 
 	private:
 		const FVector	BoxExtents;
@@ -1245,7 +1245,7 @@ void UWidgetComponent::TickComponent(float DeltaTime, enum ELevelTick TickType,
 				// Calculate the actual delta time since we last drew, this handles the case where we're ticking when
 				// the world is paused, this also takes care of the case where the widget component is rendering at
 				// a different rate than the rest of the world.
-				const float DeltaTimeFromLastDraw = LastWidgetRenderTime == 0 ? 0 : (GetCurrentTime() - LastWidgetRenderTime);
+				const float DeltaTimeFromLastDraw = static_cast<float>(LastWidgetRenderTime == 0 ? 0 : (GetCurrentTime() - LastWidgetRenderTime));
 				DrawWidgetToRenderTarget(DeltaTimeFromLastDraw);
 
 				// We draw an empty widget.
@@ -1392,8 +1392,8 @@ void UWidgetComponent::DrawWidgetToRenderTarget(float DeltaTime)
 	if ( bDrawAtDesiredSize )
 	{
 		FVector2D DesiredSize = SlateWindow->GetDesiredSize();
-		DesiredSize.X = FMath::RoundToInt(DesiredSize.X);
-		DesiredSize.Y = FMath::RoundToInt(DesiredSize.Y);
+		DesiredSize.X = FMath::RoundToDouble(DesiredSize.X);
+		DesiredSize.Y = FMath::RoundToDouble(DesiredSize.Y);
 		CurrentDrawSize = DesiredSize.IntPoint();
 
 		if (DesiredSize.X <= 0 || DesiredSize.Y <= 0)
@@ -1453,7 +1453,7 @@ void UWidgetComponent::DrawWidgetToRenderTarget(float DeltaTime)
 	}
 }
 
-float UWidgetComponent::ComputeComponentWidth() const
+double UWidgetComponent::ComputeComponentWidth() const
 {
 	switch (GeometryMode)
 	{
@@ -1463,8 +1463,8 @@ float UWidgetComponent::ComputeComponentWidth() const
 		break;
 
 		case EWidgetGeometryMode::Cylinder:
-			const float ArcAngleRadians = FMath::DegreesToRadians(GetCylinderArcAngle());
-			const float Radius = CurrentDrawSize.X / ArcAngleRadians;
+			const double ArcAngleRadians = FMath::DegreesToRadians(GetCylinderArcAngle());
+			const double Radius = CurrentDrawSize.X / ArcAngleRadians;
 			// Chord length is 2*R*Sin(Theta/2)
 			return 2.0f * Radius * FMath::Sin(0.5f*ArcAngleRadians);
 		break;
@@ -1922,15 +1922,15 @@ void UWidgetComponent::UpdateBodySetup( bool bDrawSizeChanged )
 
 		FKBoxElem* BoxElem = BodySetup->AggGeom.BoxElems.GetData();
 
-		const float Width = ComputeComponentWidth();
-		const float Height = CurrentDrawSize.Y;
+		const double Width = ComputeComponentWidth();
+		const double Height = CurrentDrawSize.Y;
 		const FVector Origin = FVector(.5f,
 			-( Width * 0.5f ) + ( Width * Pivot.X ),
 			-( Height * 0.5f ) + ( Height * Pivot.Y ));
 			
 		BoxElem->X = 0.01f;
-		BoxElem->Y = Width;
-		BoxElem->Z = Height;
+		BoxElem->Y = static_cast<float>(Width);
+		BoxElem->Z = static_cast<float>(Height);
 
 		BoxElem->SetTransform(FTransform::Identity);
 		BoxElem->Center = Origin;
@@ -1958,29 +1958,29 @@ void UWidgetComponent::GetLocalHitLocation(FVector WorldHitLocation, FVector2D&
 }
 
 
-TOptional<float> FindLineSphereIntersection(const FVector& Start, const FVector& Dir, float Radius)
+TOptional<double> FindLineSphereIntersection(const FVector& Start, const FVector& Dir, float Radius)
 {
 	// Solution exist at two possible locations:
 	// (Start + Dir * t) (dot) (Start + Dir * t) = Radius^2
 	// Dir(dot)Dir*t^2 + 2*Start(dot)Dir + Start(dot)Start - Radius^2 = 0
 	//
 	// Recognize quadratic form with:
-	const float a = FVector::DotProduct(Dir,Dir);
-	const float b = 2 * FVector::DotProduct(Start,Dir);
-	const float c = FVector::DotProduct(Start,Start) - Radius*Radius;
+	const double a = FVector::DotProduct(Dir,Dir);
+	const double b = 2 * FVector::DotProduct(Start,Dir);
+	const double c = FVector::DotProduct(Start,Start) - Radius*Radius;
 
-	const float Discriminant = b*b - 4 * a * c;
+	const double Discriminant = b*b - 4 * a * c;
 	
 	if (Discriminant >= 0)
 	{
-		const float SqrtDiscr = FMath::Sqrt(Discriminant);
-		const float Soln1 = (-b + SqrtDiscr) / (2 * a);
+		const double SqrtDiscr = FMath::Sqrt(Discriminant);
+		const double Soln1 = (-b + SqrtDiscr) / (2 * a);
 
 		return Soln1;
 	}
 	else
 	{
-		return TOptional<float>();
+		return TOptional<double>();
 	}
 }
 
@@ -1998,12 +1998,12 @@ TTuple<FVector, FVector2D> UWidgetComponent::GetCylinderHitLocation(FVector Worl
 	const FVector HitDirection_ComponentSpace = GetComponentTransform().InverseTransformVector(WorldHitDirection);
 
 
-	const float ArcAngleRadians = FMath::DegreesToRadians(GetCylinderArcAngle());
-	const float Radius = CurrentDrawSize.X / ArcAngleRadians;
-	const float Apothem = Radius * FMath::Cos(0.5f*ArcAngleRadians);
-	const float ChordLength = 2.0f * Radius * FMath::Sin(0.5f*ArcAngleRadians);
+	const double ArcAngleRadians = FMath::DegreesToRadians(GetCylinderArcAngle());
+	const double Radius = CurrentDrawSize.X / ArcAngleRadians;
+	const double Apothem = Radius * FMath::Cos(0.5f*ArcAngleRadians);
+	const double ChordLength = 2.0f * Radius * FMath::Sin(0.5f*ArcAngleRadians);
 
-	const float PivotOffsetX = ChordLength * (0.5-Pivot.X);
+	const double PivotOffsetX = ChordLength * (0.5-Pivot.X);
 
 	if (bDrawCollisionDebug)
 	{
@@ -2014,7 +2014,7 @@ TTuple<FVector, FVector2D> UWidgetComponent::GetCylinderHitLocation(FVector Worl
 
 		// Draw the imaginary circle which we use to describe the cylinder.
 		// Note that we transform all the hit locations into a space where the circle's origin is at (0,0).
-		UKismetSystemLibrary::DrawDebugCircle((UWidgetComponent*)(this), ToWorld.TransformPosition(FVector::ZeroVector), ToWorld.GetScale3D().X*Radius, 64, FLinearColor::Green,
+		UKismetSystemLibrary::DrawDebugCircle((UWidgetComponent*)(this), ToWorld.TransformPosition(FVector::ZeroVector), static_cast<float>(ToWorld.GetScale3D().X*Radius), 64, FLinearColor::Green,
 			0, 1.0f, FVector(0, 1, 0), FVector(1, 0, 0));
 		UKismetSystemLibrary::DrawDebugLine((UWidgetComponent*)(this), ToWorld.TransformPosition(FVector(-Apothem, -Radius, 0.0f)), ToWorld.TransformPosition(FVector(-Apothem, +Radius, 0.0f)), FLinearColor::Green);
 	}
@@ -2037,10 +2037,10 @@ TTuple<FVector, FVector2D> UWidgetComponent::GetCylinderHitLocation(FVector Worl
 	}
 
 	// Perform a ray vs. circle intersection test (effectively in 2D because Z coordinate is always 0)
-	const TOptional<float> Solution = FindLineSphereIntersection(HitLocation_CircleSpace, HitDirection_CircleSpace, Radius);
+	const TOptional<double> Solution = FindLineSphereIntersection(HitLocation_CircleSpace, HitDirection_CircleSpace, static_cast<float>(Radius));
 	if (Solution.IsSet())
 	{
-		const float Time = Solution.GetValue();
+		const double Time = Solution.GetValue();
 
 		const FVector TrueHitLocation_CircleSpace = HitLocation_CircleSpace + HitDirection_CircleSpace * Time;
 		if (bDrawCollisionDebug)
@@ -2052,15 +2052,15 @@ TTuple<FVector, FVector2D> UWidgetComponent::GetCylinderHitLocation(FVector Worl
 		 }
 			
 		// Determine the widget-space X hit coordinate.
-		const float Endpoint1 = FMath::Fmod(FMath::Atan2(-0.5f*ChordLength, -Apothem) + 2*PI, 2*PI);
-		const float Endpoint2 = FMath::Fmod(FMath::Atan2(+0.5f*ChordLength, -Apothem) + 2*PI, 2*PI);
-		const float HitAngleRads = FMath::Fmod((float)FMath::Atan2(TrueHitLocation_CircleSpace.Y, TrueHitLocation_CircleSpace.X) + 2*PI, 2*PI);
-		const float HitAngleZeroToOne = (HitAngleRads - FMath::Min(Endpoint1, Endpoint2)) / FMath::Abs(Endpoint2 - Endpoint1);
+		const double Endpoint1 = FMath::Fmod(FMath::Atan2(-0.5f*ChordLength, -Apothem) + 2*PI, 2*PI);
+		const double Endpoint2 = FMath::Fmod(FMath::Atan2(+0.5f*ChordLength, -Apothem) + 2*PI, 2*PI);
+		const double HitAngleRads = FMath::Fmod((float)FMath::Atan2(TrueHitLocation_CircleSpace.Y, TrueHitLocation_CircleSpace.X) + 2*PI, 2*PI);
+		const double HitAngleZeroToOne = (HitAngleRads - FMath::Min(Endpoint1, Endpoint2)) / FMath::Abs(Endpoint2 - Endpoint1);
 
 
 		// Determine the widget-space Y hit coordinate
 		const FVector CylinderHitLocation_ComponentSpace = HitLocation_ComponentSpace + HitDirection_ComponentSpace*Time;
-		const float YHitLocation = (-CylinderHitLocation_ComponentSpace.Z + CurrentDrawSize.Y*Pivot.Y);
+		const double YHitLocation = (-CylinderHitLocation_ComponentSpace.Z + CurrentDrawSize.Y*Pivot.Y);
 
 		const FVector2D WidgetSpaceHitCoord = FVector2D(HitAngleZeroToOne * CurrentDrawSize.X, YHitLocation);
 			
diff --git a/Engine/Source/Runtime/UMG/Public/Components/WidgetComponent.h b/Engine/Source/Runtime/UMG/Public/Components/WidgetComponent.h
--- a/Engine/Source/Runtime/UMG/Public/Components/WidgetComponent.h
+++ b/Engine/Source/Runtime/UMG/Public/Components/WidgetComponent.h
@@ -360,11 +360,11 @@ public:
 
 	/** Defines the curvature of the widget component when using EWidgetGeometryMode::Cylinder; ignored otherwise.  */
 	UFUNCTION(BlueprintCallable, Category = UserInterface)
-	float GetCylinderArcAngle() const { return CylinderArcAngle; }
+	double GetCylinderArcAngle() const { return CylinderArcAngle; }
 
 	/** Defines the curvature of the widget component when using EWidgetGeometryMode::Cylinder; ignored otherwise.  */
 	UFUNCTION(BlueprintCallable, Category = UserInterface)
-	void SetCylinderArcAngle(const float InCylinderArcAngle) 
+	void SetCylinderArcAngle(const double InCylinderArcAngle) 
 	{ 
 		CylinderArcAngle = InCylinderArcAngle; 
 	}
@@ -429,7 +429,7 @@ protected:
 	UMG_API virtual void DrawWidgetToRenderTarget(float DeltaTime);
 
 	/** Returns the width of the widget component taking GeometryMode into account. */
-	UMG_API float ComputeComponentWidth() const;
+	UMG_API double ComputeComponentWidth() const;
 
 	UMG_API void UpdateMaterialInstance();
 
@@ -622,7 +622,7 @@ protected:
 
 	/** Curvature of a cylindrical widget in degrees. */
 	UPROPERTY(EditAnywhere, Category=UserInterface, meta=(ClampMin=1.0f, ClampMax=180.0f))
-	float CylinderArcAngle;
+	double CylinderArcAngle;
 
 	UPROPERTY(EditAnywhere, Category = UserInterface)
 	ETickMode TickMode;
