MovieScene Truncation Warning fixes

diff --git a/Engine/Source/Editor/MovieSceneTools/Public/KeyframeTrackEditor.h b/Engine/Source/Editor/MovieSceneTools/Public/KeyframeTrackEditor.h
--- a/Engine/Source/Editor/MovieSceneTools/Public/KeyframeTrackEditor.h
+++ b/Engine/Source/Editor/MovieSceneTools/Public/KeyframeTrackEditor.h
@@ -401,7 +401,7 @@ struct TAddKeyImpl<FMovieSceneIntegerChannel, int32> : IImpl
 				TOptional<int32> OptInt = Channel->GetDefault();
 				LocalValue = OptInt.IsSet() ? OptInt.GetValue() : 0;
 			}
-			ValueToSet = (float)((ValueToSet - CurrentValue)) * Weight + LocalValue;
+			ValueToSet = FMath::TruncToInt32((float)((ValueToSet - CurrentValue)) * Weight + LocalValue);
 			return true;
 		}
 
diff --git a/Engine/Source/Editor/SequencerCore/Public/TimeToPixel.h b/Engine/Source/Editor/SequencerCore/Public/TimeToPixel.h
--- a/Engine/Source/Editor/SequencerCore/Public/TimeToPixel.h
+++ b/Engine/Source/Editor/SequencerCore/Public/TimeToPixel.h
@@ -24,7 +24,7 @@ public:
 		, TickResolution( InTickResolution )
 	{
 		ensureMsgf(WidthPx >= 0, TEXT("Got negative pixel width!"));
-		const double VisibleWidth = InLocalViewRange.Size<double>();
+		const float VisibleWidth = static_cast<float>(InLocalViewRange.Size<double>());
 
 		const float MaxPixelsPerSecond = 1000.f;
 		PixelsPerSecond = VisibleWidth > 0 ? WidthPx / VisibleWidth : MaxPixelsPerSecond;
@@ -38,7 +38,7 @@ public:
 	 */
 	float SecondsToPixel( double Time ) const
 	{
-		return (Time - ViewRangeStartSeconds) * PixelsPerSecond;
+		return static_cast<float>(Time - ViewRangeStartSeconds) * PixelsPerSecond;
 	}
 
 	/**
@@ -49,7 +49,7 @@ public:
 	 */
 	float SecondsDeltaToPixel( double TimeDelta ) const
 	{
-		return TimeDelta * PixelsPerSecond;
+		return static_cast<float>(TimeDelta) * PixelsPerSecond;
 	}
 
 	/**
@@ -71,7 +71,7 @@ public:
 	 */
 	float FrameToPixel( const FFrameTime& Time ) const
 	{
-		return (Time / TickResolution - ViewRangeStartSeconds) * PixelsPerSecond;
+		return static_cast<float>(Time / TickResolution - ViewRangeStartSeconds) * PixelsPerSecond;
 	}
 
 	/**
@@ -82,7 +82,7 @@ public:
 	 */
 	float FrameDeltaToPixel( const FFrameTime& TimeDelta) const
 	{
-		return (TimeDelta / TickResolution) * PixelsPerSecond;
+		return static_cast<float>(TimeDelta / TickResolution) * PixelsPerSecond;
 	}
 
 	/**
diff --git a/Engine/Source/Runtime/MovieScene/Public/EntitySystem/MovieSceneEntityIDs.h b/Engine/Source/Runtime/MovieScene/Public/EntitySystem/MovieSceneEntityIDs.h
--- a/Engine/Source/Runtime/MovieScene/Public/EntitySystem/MovieSceneEntityIDs.h
+++ b/Engine/Source/Runtime/MovieScene/Public/EntitySystem/MovieSceneEntityIDs.h
@@ -258,7 +258,7 @@ struct TComponentTypeID : public FComponentTypeID
 	static TComponentTypeID FromBitIndex(int32 BitIndex)
 	{
 		check( (BitIndex & 0xFFFF0000) == 0 );
-		return TComponentTypeID(BitIndex);
+		return TComponentTypeID(static_cast<uint16>(BitIndex));
 	}
 private:
 
diff --git a/Engine/Source/Runtime/MovieScene/Public/EntitySystem/MovieSceneInitialValueCache.h b/Engine/Source/Runtime/MovieScene/Public/EntitySystem/MovieSceneInitialValueCache.h
--- a/Engine/Source/Runtime/MovieScene/Public/EntitySystem/MovieSceneInitialValueCache.h
+++ b/Engine/Source/Runtime/MovieScene/Public/EntitySystem/MovieSceneInitialValueCache.h
@@ -294,7 +294,7 @@ private:
 			// Hash only considers the _type_ of the resolved property
 			// which defers the final comparison to the equality operator
 			uint32 Hash = GetTypeHash(InKey.Object);
-			Hash = HashCombine(Hash, InKey.Property.GetIndex());
+			Hash = HashCombine(Hash, static_cast<uint32>(InKey.Property.GetIndex()));
 			return Hash;
 		}
 		friend bool operator==(const FKeyType& A, const FKeyType& B)
diff --git a/Engine/Source/Runtime/MovieScene/Public/EntitySystem/MovieScenePropertyRegistry.h b/Engine/Source/Runtime/MovieScene/Public/EntitySystem/MovieScenePropertyRegistry.h
--- a/Engine/Source/Runtime/MovieScene/Public/EntitySystem/MovieScenePropertyRegistry.h
+++ b/Engine/Source/Runtime/MovieScene/Public/EntitySystem/MovieScenePropertyRegistry.h
@@ -258,8 +258,10 @@ private:
 	{
 		using StorageType = typename PropertyTraits::StorageType;
 
-		const int32 CompositeOffset = CompositeDefinitions.Num();
-		checkf(CompositeOffset <= MAX_uint16, TEXT("Maximum number of composite definitions reached"));
+		const int32 CompositeOffsetInt32 = CompositeDefinitions.Num();
+		checkf(CompositeOffsetInt32 <= MAX_uint16, TEXT("Maximum number of composite definitions reached"));
+
+		const uint16 CompositeOffset = static_cast<uint16>(CompositeOffsetInt32);
 
 		TStatId StatID;
 
diff --git a/Engine/Source/Runtime/MovieScene/Public/Evaluation/MovieScenePlaybackCapabilities.h b/Engine/Source/Runtime/MovieScene/Public/Evaluation/MovieScenePlaybackCapabilities.h
--- a/Engine/Source/Runtime/MovieScene/Public/Evaluation/MovieScenePlaybackCapabilities.h
+++ b/Engine/Source/Runtime/MovieScene/Public/Evaluation/MovieScenePlaybackCapabilities.h
@@ -312,13 +312,13 @@ protected:
 			for (; Index < NewCapabilityIndex; ++Index)
 			{
 				RequiredSizeof = Align(RequiredSizeof, (uint64)ExistingHeaders[Index].Alignment);
-				NewCapabilityOffsets[Index] = RequiredSizeof;
+				NewCapabilityOffsets[Index] = static_cast<uint16>(RequiredSizeof);
 				RequiredSizeof += ExistingHeaders[Index].Sizeof;
 			}
 
 			// Count up the size and alignment for the new capability
 			RequiredSizeof = Align(RequiredSizeof, alignof(StorageType));
-			NewCapabilityOffsets[Index] = RequiredSizeof;
+			NewCapabilityOffsets[Index] = static_cast<uint16>(RequiredSizeof);
 			RequiredSizeof += sizeof(StorageType);
 			++Index;
 
@@ -326,7 +326,7 @@ protected:
 			for (; Index < ExistingNum+1; ++Index)
 			{
 				RequiredSizeof = Align(RequiredSizeof, (uint64)ExistingHeaders[Index-1].Alignment);
-				NewCapabilityOffsets[Index] = RequiredSizeof;
+				NewCapabilityOffsets[Index] = static_cast<uint16>(RequiredSizeof);
 				RequiredSizeof += ExistingHeaders[Index-1].Sizeof;
 			}
 		}
@@ -342,7 +342,7 @@ protected:
 		if (bNeedsReallocation)
 		{
 			// Use the greater of the required size or double the current size to allow some additional capcity
-			Capacity = FMath::Max(RequiredSizeof, uint64(Capacity)*2);
+			Capacity = static_cast<uint16>(FMath::Max(RequiredSizeof, uint64(Capacity)*2));
 			Memory = reinterpret_cast<uint8*>(FMemory::Malloc(Capacity, RequiredAlignment));
 		}
 
diff --git a/Engine/Source/Runtime/MovieSceneTracks/Public/MovieSceneTracksPropertyTypes.h b/Engine/Source/Runtime/MovieSceneTracks/Public/MovieSceneTracksPropertyTypes.h
--- a/Engine/Source/Runtime/MovieSceneTracks/Public/MovieSceneTracksPropertyTypes.h
+++ b/Engine/Source/Runtime/MovieSceneTracks/Public/MovieSceneTracksPropertyTypes.h
@@ -129,13 +129,13 @@ struct FIntermediateColor
 	FColor GetColor() const
 	{
 		const bool bConvertBackToSRgb = true;
-		const FColor SRgbColor = FLinearColor(R, G, B, A).ToFColor(bConvertBackToSRgb);
+		const FColor SRgbColor = GetLinearColor().ToFColor(bConvertBackToSRgb);
 		return SRgbColor;
 	}
 
 	FLinearColor GetLinearColor() const
 	{
-		return FLinearColor(R, G, B, A);
+		return FLinearColor(static_cast<float>(R), static_cast<float>(G), static_cast<float>(B), static_cast<float>(A));
 	}
 
 	FSlateColor GetSlateColor() const
diff --git a/Engine/Source/Runtime/UMG/Private/Animation/MovieScene2DTransformSection.cpp b/Engine/Source/Runtime/UMG/Private/Animation/MovieScene2DTransformSection.cpp
--- a/Engine/Source/Runtime/UMG/Private/Animation/MovieScene2DTransformSection.cpp
+++ b/Engine/Source/Runtime/UMG/Private/Animation/MovieScene2DTransformSection.cpp
@@ -84,11 +84,11 @@ struct F2DTransformSectionEditorData
 
 	static TOptional<float> ExtractTranslationX(UObject& InObject, FTrackInstancePropertyBindings* Bindings)
 	{
-		return Bindings ? Bindings->GetCurrentValue<FWidgetTransform>(InObject).Translation.X : TOptional<float>();
+		return Bindings ? static_cast<float>(Bindings->GetCurrentValue<FWidgetTransform>(InObject).Translation.X) : TOptional<float>();
 	}
 	static TOptional<float> ExtractTranslationY(UObject& InObject, FTrackInstancePropertyBindings* Bindings)
 	{
-		return Bindings ? Bindings->GetCurrentValue<FWidgetTransform>(InObject).Translation.Y : TOptional<float>();
+		return Bindings ? static_cast<float>(Bindings->GetCurrentValue<FWidgetTransform>(InObject).Translation.Y) : TOptional<float>();
 	}
 
 	static TOptional<float> ExtractRotation(UObject& InObject, FTrackInstancePropertyBindings* Bindings)
@@ -98,20 +98,20 @@ struct F2DTransformSectionEditorData
 
 	static TOptional<float> ExtractScaleX(UObject& InObject, FTrackInstancePropertyBindings* Bindings)
 	{
-		return Bindings ? Bindings->GetCurrentValue<FWidgetTransform>(InObject).Scale.X : TOptional<float>();
+		return Bindings ? static_cast<float>(Bindings->GetCurrentValue<FWidgetTransform>(InObject).Scale.X) : TOptional<float>();
 	}
 	static TOptional<float> ExtractScaleY(UObject& InObject, FTrackInstancePropertyBindings* Bindings)
 	{
-		return Bindings ? Bindings->GetCurrentValue<FWidgetTransform>(InObject).Scale.Y : TOptional<float>();
+		return Bindings ? static_cast<float>(Bindings->GetCurrentValue<FWidgetTransform>(InObject).Scale.Y) : TOptional<float>();
 	}
 
 	static TOptional<float> ExtractShearX(UObject& InObject, FTrackInstancePropertyBindings* Bindings)
 	{
-		return Bindings ? Bindings->GetCurrentValue<FWidgetTransform>(InObject).Shear.X : TOptional<float>();
+		return Bindings ? static_cast<float>(Bindings->GetCurrentValue<FWidgetTransform>(InObject).Shear.X) : TOptional<float>();
 	}
 	static TOptional<float> ExtractShearY(UObject& InObject, FTrackInstancePropertyBindings* Bindings)
 	{
-		return Bindings ? Bindings->GetCurrentValue<FWidgetTransform>(InObject).Shear.Y : TOptional<float>();
+		return Bindings ? static_cast<float>(Bindings->GetCurrentValue<FWidgetTransform>(InObject).Shear.Y) : TOptional<float>();
 	}
 
 	FMovieSceneChannelMetaData      MetaData[7];
diff --git a/Engine/Source/Runtime/UMG/Private/Animation/MovieSceneUMGComponentTypes.cpp b/Engine/Source/Runtime/UMG/Private/Animation/MovieSceneUMGComponentTypes.cpp
--- a/Engine/Source/Runtime/UMG/Private/Animation/MovieSceneUMGComponentTypes.cpp
+++ b/Engine/Source/Runtime/UMG/Private/Animation/MovieSceneUMGComponentTypes.cpp
@@ -25,7 +25,7 @@ void ConvertOperationalProperty(const FIntermediateWidgetTransform& In, FWidgetT
 {
 	Out.Translation.X = In.TranslationX;
 	Out.Translation.Y = In.TranslationY;
-	Out.Angle = In.Rotation;
+	Out.Angle = static_cast<float>(In.Rotation);
 	Out.Scale.X = In.ScaleX;
 	Out.Scale.Y = In.ScaleY;
 	Out.Shear.X = In.ShearX;
@@ -44,10 +44,10 @@ void ConvertOperationalProperty(const FWidgetTransform& In, FIntermediateWidgetT
 
 void ConvertOperationalProperty(const FIntermediateMargin& In, FMargin& Out)
 {
-	Out.Top = In.Top;
-	Out.Right = In.Right;
-	Out.Bottom = In.Bottom;
-	Out.Left = In.Left;
+	Out.Top = static_cast<float>(In.Top);
+	Out.Right = static_cast<float>(In.Right);
+	Out.Bottom = static_cast<float>(In.Bottom);
+	Out.Left = static_cast<float>(In.Left);
 }
 
 void ConvertOperationalProperty(const FMargin& In, FIntermediateMargin& Out)
