Speculative fix crash custom drawer garbage collected by deferring drawer collection using a render command.

diff --git a/Engine/Source/Runtime/UMG/Private/Slate/SPostBufferUpdate.cpp b/Engine/Source/Runtime/UMG/Private/Slate/SPostBufferUpdate.cpp
--- a/Engine/Source/Runtime/UMG/Private/Slate/SPostBufferUpdate.cpp
+++ b/Engine/Source/Runtime/UMG/Private/Slate/SPostBufferUpdate.cpp
@@ -226,6 +226,13 @@ const TArrayView<const ESlatePostRT> SPostBufferUpdate::GetBuffersToUpdate() con
 UMG_API void SPostBufferUpdate::ReleasePostBufferUpdater()
 {
 #if !UE_SERVER
+	// Copy the pointer onto a lambda to defer the final deletion to after any pending uses on the renderthread
+	TSharedPtr<FPostBufferUpdater, ESPMode::ThreadSafe> ReleaseMe = PostBufferUpdater;
+	ENQUEUE_RENDER_COMMAND(ReleaseCommand)([ReleaseMe](FRHICommandList& RHICmdList) mutable
+	{
+		ReleaseMe.Reset();
+	});
+
 	PostBufferUpdater.Reset();
 #endif // !UE_SERVER
 }
