PR #11434: UMG: In Blueprint, allow setting InText inline for URichTextBlock::SetText, this aligns with UTextBlock::SetText

diff --git a/Engine/Source/Runtime/UMG/Public/Components/RichTextBlock.h b/Engine/Source/Runtime/UMG/Public/Components/RichTextBlock.h
--- a/Engine/Source/Runtime/UMG/Public/Components/RichTextBlock.h
+++ b/Engine/Source/Runtime/UMG/Public/Components/RichTextBlock.h
@@ -204,7 +204,7 @@ public:
 	 * Warning: This will wipe any binding created for the Text property!
 	 * @param InText The text to assign to the widget
 	 */
-	UFUNCTION(BlueprintCallable, Category = "Widget")
+	UFUNCTION(BlueprintCallable, Category = "Widget", meta = (AutoCreateRefTerm = "InText"))
 	UMG_API virtual void SetText(const FText& InText);
 
 	UFUNCTION(BlueprintCallable, Category = "Widget")
