Retainer: fix an issue with clipping that was not working in case Stencil clipping was used on elements inside the the retainer (which when using shear or non axis aligned rotation).
It was for instance preventing the display of the fill of the progress bar, if shear/rotation was applyed on it and it was used in a retainer.
It's not using the change of the PR as it was working only in some case: the clipping zone we use always start from 0,0, so using the size of the Retainer (like in the PR) if it was smaller than the window would lead to the element being clipped out if the retainer was too far from 0,0.
#11455
#tests In PIE, put a progress bar inside a retainer and played with shear and rotation on both the retainer and the progress bar: the blue fill is now always displayed. Engine test retainer test still succeed. Validated curve editor was still working as expected. Tested that animation are still played properly on windows (in case of drag n drop).

diff --git a/Engine/Source/Runtime/UMG/Private/Slate/SRetainerWidget.cpp b/Engine/Source/Runtime/UMG/Private/Slate/SRetainerWidget.cpp
--- a/Engine/Source/Runtime/UMG/Private/Slate/SRetainerWidget.cpp
+++ b/Engine/Source/Runtime/UMG/Private/Slate/SRetainerWidget.cpp
@@ -565,6 +565,16 @@ SRetainerWidget::EPaintRetainedContentResult SRetainerWidget::PaintRetainedConte
 				WidgetRenderer->ViewOffset = -ViewOffset;
 				WidgetRenderer->SetIsPrepassNeeded(false);
 
+				FVector2f WindowSize(RenderSize);
+				SWindow* PaintWindow = Context.WindowElementList->GetPaintWindow();
+				if (PaintWindow)
+				{
+					const FVector2f ViewportSize = PaintWindow->GetViewportSize();
+					WindowSize.X = FMath::Max(WindowSize.X, ViewportSize.X);
+					WindowSize.Y = FMath::Max(WindowSize.Y, ViewportSize.Y);
+				}
+				VirtualWindow->Resize(WindowSize);
+
 				bool bRepaintedWidgets = WidgetRenderer->DrawInvalidationRoot(VirtualWindow, RenderTarget, *this, Context, GDeferRetainedRenderingRenderThread != 0);
 				bRenderRequested = false;
 				Shared_WaitingToRender.Remove(this);
