Sequencer: reset director blueprints in UMG when stopping an animation

diff --git a/Engine/Source/Runtime/UMG/Private/Animation/UMGSequencePlayer.cpp b/Engine/Source/Runtime/UMG/Private/Animation/UMGSequencePlayer.cpp
--- a/Engine/Source/Runtime/UMG/Private/Animation/UMGSequencePlayer.cpp
+++ b/Engine/Source/Runtime/UMG/Private/Animation/UMGSequencePlayer.cpp
@@ -434,6 +434,10 @@ void UUMGSequencePlayer::HandleLatentStop()
 
 	PlayerStatus = EMovieScenePlayerStatus::Stopped;
 
+	PRAGMA_DISABLE_DEPRECATION_WARNINGS
+	RootTemplateInstance.ResetDirectorInstances();
+	PRAGMA_ENABLE_DEPRECATION_WARNINGS
+
 	if (bRestoreState)
 	{
 		RestorePreAnimatedState();
