Add QueueAnimation methods, which perform better when multiple animation transitions occur in one frame.

#tests PIE EUW

diff --git a/Engine/Source/Runtime/UMG/Private/UserWidget.cpp b/Engine/Source/Runtime/UMG/Private/UserWidget.cpp
--- a/Engine/Source/Runtime/UMG/Private/UserWidget.cpp
+++ b/Engine/Source/Runtime/UMG/Private/UserWidget.cpp
@@ -513,6 +513,58 @@ UUMGSequencePlayer* UUserWidget::GetOrAddSequencePlayer(UWidgetAnimation* InAnim
 	return nullptr;
 }
 
+void UUserWidget::ExecuteQueuedAnimationTransitions()
+{
+	// In case any users queue animations in respose to animation transitions, operate on a copy array
+	TArray<FQueuedWidgetAnimationTransition, TInlineAllocator<8>> CurrentWidgetAnimationTransitions(QueuedWidgetAnimationTransitions);
+
+	for (FQueuedWidgetAnimationTransition& QueuedWidgetAnimationTransition : CurrentWidgetAnimationTransitions)
+	{
+		switch (QueuedWidgetAnimationTransition.TransitionMode)
+		{
+		case EQueuedWidgetAnimationMode::Play:
+			PlayAnimation(QueuedWidgetAnimationTransition.WidgetAnimation
+				, QueuedWidgetAnimationTransition.StartAtTime.GetValue()
+				, QueuedWidgetAnimationTransition.NumLoopsToPlay.GetValue()
+				, QueuedWidgetAnimationTransition.PlayMode.GetValue()
+				, QueuedWidgetAnimationTransition.PlaybackSpeed.GetValue()
+				, QueuedWidgetAnimationTransition.bRestoreState.GetValue());
+			break;
+		case EQueuedWidgetAnimationMode::PlayTo:
+			PlayAnimationTimeRange(QueuedWidgetAnimationTransition.WidgetAnimation
+				, QueuedWidgetAnimationTransition.StartAtTime.GetValue()
+				, QueuedWidgetAnimationTransition.EndAtTime.GetValue()
+				, QueuedWidgetAnimationTransition.NumLoopsToPlay.GetValue()
+				, QueuedWidgetAnimationTransition.PlayMode.GetValue()
+				, QueuedWidgetAnimationTransition.PlaybackSpeed.GetValue()
+				, QueuedWidgetAnimationTransition.bRestoreState.GetValue());
+			break;
+		case EQueuedWidgetAnimationMode::Forward:
+			PlayAnimationForward(QueuedWidgetAnimationTransition.WidgetAnimation
+				, QueuedWidgetAnimationTransition.PlaybackSpeed.GetValue()
+				, QueuedWidgetAnimationTransition.bRestoreState.GetValue());
+			break;
+		case EQueuedWidgetAnimationMode::Reverse:
+			PlayAnimationReverse(QueuedWidgetAnimationTransition.WidgetAnimation
+				, QueuedWidgetAnimationTransition.PlaybackSpeed.GetValue()
+				, QueuedWidgetAnimationTransition.bRestoreState.GetValue());
+			break;
+		case EQueuedWidgetAnimationMode::Stop:
+			StopAnimation(QueuedWidgetAnimationTransition.WidgetAnimation);
+			break;
+		case EQueuedWidgetAnimationMode::Pause:
+			PauseAnimation(QueuedWidgetAnimationTransition.WidgetAnimation);
+			break;
+		}
+	}
+
+	if (QueuedWidgetAnimationTransitions.Num() > 0)
+	{
+		QueuedWidgetAnimationTransitions.Empty();
+		UpdateCanTick();
+	}
+}
+
 void UUserWidget::ConditionalTearDownAnimations()
 {
 	for (auto It = ActiveSequencePlayers.CreateIterator(); It; ++It)
@@ -575,6 +627,142 @@ void UUserWidget::Invalidate(EInvalidateWidgetReason InvalidateReason)
 	}
 }
 
+void UUserWidget::QueuePlayAnimation(UWidgetAnimation* InAnimation, float StartAtTime, int32 NumLoopsToPlay, EUMGSequencePlayMode::Type PlayMode, float PlaybackSpeed, bool bRestoreState)
+{
+	if (!InAnimation)
+	{
+		return;
+	}
+
+	FQueuedWidgetAnimationTransition* QueuedTransitionPtr = QueuedWidgetAnimationTransitions.FindByPredicate([&](const FQueuedWidgetAnimationTransition& QueuedTransition) { return QueuedTransition.WidgetAnimation == InAnimation; });
+	FQueuedWidgetAnimationTransition& QueuedTransition = QueuedTransitionPtr ? *QueuedTransitionPtr : QueuedWidgetAnimationTransitions.AddDefaulted_GetRef();
+
+	QueuedTransition = FQueuedWidgetAnimationTransition();
+	QueuedTransition.WidgetAnimation = InAnimation;
+	QueuedTransition.TransitionMode = EQueuedWidgetAnimationMode::Play;
+	QueuedTransition.StartAtTime = StartAtTime;
+	QueuedTransition.NumLoopsToPlay = NumLoopsToPlay;
+	QueuedTransition.PlayMode = PlayMode;
+	QueuedTransition.PlaybackSpeed = PlaybackSpeed;
+	QueuedTransition.bRestoreState = bRestoreState;
+
+	UpdateCanTick();
+}
+
+void UUserWidget::QueuePlayAnimationTimeRange(UWidgetAnimation* InAnimation, float StartAtTime, float EndAtTime, int32 NumLoopsToPlay, EUMGSequencePlayMode::Type PlayMode, float PlaybackSpeed, bool bRestoreState)
+{
+	if (!InAnimation)
+	{
+		return;
+	}
+
+	FQueuedWidgetAnimationTransition* QueuedTransitionPtr = QueuedWidgetAnimationTransitions.FindByPredicate([&](const FQueuedWidgetAnimationTransition& QueuedTransition) { return QueuedTransition.WidgetAnimation == InAnimation; });
+	FQueuedWidgetAnimationTransition& QueuedTransition = QueuedTransitionPtr ? *QueuedTransitionPtr : QueuedWidgetAnimationTransitions.AddDefaulted_GetRef();
+
+	QueuedTransition.WidgetAnimation = InAnimation;
+	QueuedTransition.TransitionMode = EQueuedWidgetAnimationMode::PlayTo;
+	QueuedTransition.StartAtTime = StartAtTime;
+	QueuedTransition.EndAtTime = EndAtTime;
+	QueuedTransition.NumLoopsToPlay = NumLoopsToPlay;
+	QueuedTransition.PlayMode = PlayMode;
+	QueuedTransition.PlaybackSpeed = PlaybackSpeed;
+	QueuedTransition.bRestoreState = bRestoreState;
+
+	UpdateCanTick();
+}
+
+void UUserWidget::QueuePlayAnimationForward(UWidgetAnimation* InAnimation, float PlaybackSpeed, bool bRestoreState)
+{
+	if (!InAnimation)
+	{
+		return;
+	}
+
+	FQueuedWidgetAnimationTransition* QueuedTransitionPtr = QueuedWidgetAnimationTransitions.FindByPredicate([&](const FQueuedWidgetAnimationTransition& QueuedTransition) { return QueuedTransition.WidgetAnimation == InAnimation; });
+	FQueuedWidgetAnimationTransition& QueuedTransition = QueuedTransitionPtr ? *QueuedTransitionPtr : QueuedWidgetAnimationTransitions.AddDefaulted_GetRef();
+
+	QueuedTransition.WidgetAnimation = InAnimation;
+	QueuedTransition.TransitionMode = EQueuedWidgetAnimationMode::Forward;
+	QueuedTransition.PlaybackSpeed = PlaybackSpeed;
+	QueuedTransition.bRestoreState = bRestoreState;
+
+	UpdateCanTick();
+}
+
+void UUserWidget::QueuePlayAnimationReverse(UWidgetAnimation* InAnimation, float PlaybackSpeed, bool bRestoreState)
+{
+	if (!InAnimation)
+	{
+		return;
+	}
+
+	FQueuedWidgetAnimationTransition* QueuedTransitionPtr = QueuedWidgetAnimationTransitions.FindByPredicate([&](const FQueuedWidgetAnimationTransition& QueuedTransition) { return QueuedTransition.WidgetAnimation == InAnimation; });
+	FQueuedWidgetAnimationTransition& QueuedTransition = QueuedTransitionPtr ? *QueuedTransitionPtr : QueuedWidgetAnimationTransitions.AddDefaulted_GetRef();
+
+	QueuedTransition.WidgetAnimation = InAnimation;
+	QueuedTransition.TransitionMode = EQueuedWidgetAnimationMode::Reverse;
+	QueuedTransition.PlaybackSpeed = PlaybackSpeed;
+	QueuedTransition.bRestoreState = bRestoreState;
+
+	UpdateCanTick();
+}
+
+void UUserWidget::QueueStopAnimation(const UWidgetAnimation* InAnimation)
+{
+	if (!InAnimation)
+	{
+		return;
+	}
+
+	FQueuedWidgetAnimationTransition* QueuedTransitionPtr = QueuedWidgetAnimationTransitions.FindByPredicate([&](const FQueuedWidgetAnimationTransition& QueuedTransition) { return QueuedTransition.WidgetAnimation == InAnimation; });
+	FQueuedWidgetAnimationTransition& QueuedTransition = QueuedTransitionPtr ? *QueuedTransitionPtr : QueuedWidgetAnimationTransitions.AddDefaulted_GetRef();
+
+	QueuedTransition.WidgetAnimation = const_cast<UWidgetAnimation*>(InAnimation);
+	QueuedTransition.TransitionMode = EQueuedWidgetAnimationMode::Stop;
+
+	UpdateCanTick();
+}
+
+void UUserWidget::QueueStopAllAnimations()
+{
+	for (FQueuedWidgetAnimationTransition& QueuedWidgetAnimationTransition : QueuedWidgetAnimationTransitions)
+	{
+		QueuedWidgetAnimationTransition.TransitionMode = EQueuedWidgetAnimationMode::Stop;
+	}
+
+	TArray<UUMGSequencePlayer*, TInlineAllocator<8>> CurrentActivePlayers(ActiveSequencePlayers);
+	for (UUMGSequencePlayer* FoundPlayer : ActiveSequencePlayers)
+	{
+		if (FoundPlayer->GetPlaybackStatus() == EMovieScenePlayerStatus::Playing)
+		{
+			QueueStopAnimation(FoundPlayer->GetAnimation());
+		}
+	}
+
+	UpdateCanTick();
+}
+
+float UUserWidget::QueuePauseAnimation(const UWidgetAnimation* InAnimation)
+{
+	if (InAnimation)
+	{
+		FQueuedWidgetAnimationTransition* QueuedTransitionPtr = QueuedWidgetAnimationTransitions.FindByPredicate([&](const FQueuedWidgetAnimationTransition& QueuedTransition) { return QueuedTransition.WidgetAnimation == InAnimation; });
+		FQueuedWidgetAnimationTransition& QueuedTransition = QueuedTransitionPtr ? *QueuedTransitionPtr : QueuedWidgetAnimationTransitions.AddDefaulted_GetRef();
+
+		QueuedTransition.WidgetAnimation = const_cast<UWidgetAnimation*>(InAnimation);
+		QueuedTransition.TransitionMode = EQueuedWidgetAnimationMode::Pause;
+
+		UpdateCanTick();
+
+		if (UUMGSequencePlayer* FoundPlayer = GetSequencePlayer(InAnimation))
+		{
+			return (float)FoundPlayer->GetCurrentTime().AsSeconds();
+		}
+	}
+
+	return 0;
+}
+
 UUMGSequencePlayer* UUserWidget::PlayAnimation(UWidgetAnimation* InAnimation, float StartAtTime, int32 NumberOfLoops, EUMGSequencePlayMode::Type PlayMode, float PlaybackSpeed, bool bRestoreState)
 {
 	SCOPED_NAMED_EVENT_TEXT("Widget::PlayAnimation", FColor::Emerald);
@@ -587,7 +775,7 @@ UUMGSequencePlayer* UUserWidget::PlayAnimation(UWidgetAnimation* InAnimation, fl
 		OnAnimationStartedPlaying(*Player);
 
 		UpdateCanTick();
-}
+	}
 
 	return Player;
 }
@@ -1615,6 +1803,8 @@ void UUserWidget::NativeTick(const FGeometry& MyGeometry, float InDeltaTime)
 #endif
 		if (bTickAnimations)
 		{
+			ExecuteQueuedAnimationTransitions();
+
 			if (AnimationTickManager)
 			{
 				AnimationTickManager->OnWidgetTicked(this);
@@ -1853,6 +2043,7 @@ void UUserWidget::UpdateCanTick()
 			bCanTick |= bHasScriptImplementedTick;
 			bCanTick |= World->GetLatentActionManager().GetNumActionsForObject(this) != 0;
 			bCanTick |= ActiveSequencePlayers.Num() > 0;
+			bCanTick |= QueuedWidgetAnimationTransitions.Num() > 0;
 
 			if (!bCanTick && bAreExtensionsConstructed)
 			{
diff --git a/Engine/Source/Runtime/UMG/Public/Blueprint/UserWidget.h b/Engine/Source/Runtime/UMG/Public/Blueprint/UserWidget.h
--- a/Engine/Source/Runtime/UMG/Public/Blueprint/UserWidget.h
+++ b/Engine/Source/Runtime/UMG/Public/Blueprint/UserWidget.h
@@ -45,6 +45,64 @@ class UWidgetTree;
 class UNamedSlot;
 class UUserWidgetExtension;
 
+/** Describes overall action driving this animation transition. */
+enum class EQueuedWidgetAnimationMode : uint8
+{
+	/** Animation plays with given params. */
+	Play,
+	/** Animation plays with given params to given point. */
+	PlayTo,
+	/** Animation plays from current position forward. */
+	Forward,
+	/** Animation plays from current position reverse. */
+	Reverse,
+	/** Animation stops playing. */
+	Stop,
+	/** Animation stops playing. */
+	Pause,
+	/** Default state, should not be used. */
+	None,
+};
+
+/**
+ * Struct that maintains state of currently queued animation transtions to be evaluated next frame.
+ */
+USTRUCT()
+struct UMG_API FQueuedWidgetAnimationTransition
+{
+	GENERATED_BODY()
+
+	/** Animation with a queued transition */
+	UPROPERTY(Transient)
+	TObjectPtr<UWidgetAnimation> WidgetAnimation;
+
+	/** Overall action driving this animation transition */
+	EQueuedWidgetAnimationMode TransitionMode;
+
+	/** The time in the animation from which to start playing, relative to the start position. For looped animations, this will only affect the first playback of the animation */
+	TOptional<float> StartAtTime;
+
+	/** The absolute time in the animation where to stop, this is only considered in the last loop. */
+	TOptional<float> EndAtTime;
+
+	/** The number of times to loop this animation (0 to loop indefinitely) */
+	TOptional<int32> NumLoopsToPlay;
+
+	/** Specifies the playback mode (Forward, Reverse) */
+	TOptional<EUMGSequencePlayMode::Type> PlayMode;
+
+	/** The speed at which the animation should play */
+	TOptional<float> PlaybackSpeed;
+
+	/** Restores widgets to their pre-animated state when the animation stops */
+	TOptional<bool> bRestoreState;
+
+	FQueuedWidgetAnimationTransition()
+		: WidgetAnimation(nullptr)
+		, TransitionMode(EQueuedWidgetAnimationMode::None)
+	{}
+};
+
 /** Determines what strategy we use to determine when and if the widget ticks. */
 UENUM()
 enum class EWidgetTickFrequency : uint8
@@ -997,6 +1055,82 @@ public:
 
 	UMG_API void SetIsFocusable(bool InIsFocusable);
 
+	/**
+	 * Plays an animation in this widget a specified number of times
+	 * 
+	 * @param InAnimation The animation to play
+	 * @param StartAtTime The time in the animation from which to start playing, relative to the start position. For looped animations, this will only affect the first playback of the animation.
+	 * @param NumLoopsToPlay The number of times to loop this animation (0 to loop indefinitely)
+	 * @param PlaybackSpeed The speed at which the animation should play
+	 * @param PlayMode Specifies the playback mode
+	 * @param bRestoreState Restores widgets to their pre-animated state when the animation stops
+	 */
+	UFUNCTION(BlueprintCallable, BlueprintCosmetic, Category = "User Interface|Animation")
+	UMG_API void QueuePlayAnimation(UWidgetAnimation* InAnimation, float StartAtTime = 0.0f, int32 NumLoopsToPlay = 1, EUMGSequencePlayMode::Type PlayMode = EUMGSequencePlayMode::Forward, float PlaybackSpeed = 1.0f, bool bRestoreState = false);
+
+	/**
+	 * Plays an animation in this widget a specified number of times stopping at a specified time
+	 * 
+	 * @param InAnimation The animation to play
+	 * @param StartAtTime The time in the animation from which to start playing, relative to the start position. For looped animations, this will only affect the first playback of the animation.
+	 * @param EndAtTime The absolute time in the animation where to stop, this is only considered in the last loop.
+	 * @param NumLoopsToPlay The number of times to loop this animation (0 to loop indefinitely)
+	 * @param PlayMode Specifies the playback mode
+	 * @param PlaybackSpeed The speed at which the animation should play
+	 * @param bRestoreState Restores widgets to their pre-animated state when the animation stops
+	 */
+	UFUNCTION(BlueprintCallable, BlueprintCosmetic, Category = "User Interface|Animation")
+	UMG_API void QueuePlayAnimationTimeRange(UWidgetAnimation* InAnimation, float StartAtTime = 0.0f, float EndAtTime = 0.0f, int32 NumLoopsToPlay = 1, EUMGSequencePlayMode::Type PlayMode = EUMGSequencePlayMode::Forward, float PlaybackSpeed = 1.0f, bool bRestoreState = false);
+
+	/**
+	 * Plays an animation on this widget relative to it's current state forward.  You should use this version in situations where
+	 * say a user can click a button and that causes a panel to slide out, and you want to reverse that same animation to begin sliding
+	 * in the opposite direction.
+	 * 
+	 * @param InAnimation The animation to play
+	 * @param PlayMode Specifies the playback mode
+	 * @param PlaybackSpeed The speed at which the animation should play
+	 * @param bRestoreState Restores widgets to their pre-animated state when the animation stops
+	 */
+	UFUNCTION(BlueprintCallable, BlueprintCosmetic, Category = "User Interface|Animation")
+	UMG_API void QueuePlayAnimationForward(UWidgetAnimation* InAnimation, float PlaybackSpeed = 1.0f, bool bRestoreState = false);
+
+	/**
+	 * Plays an animation on this widget relative to it's current state in reverse.  You should use this version in situations where
+	 * say a user can click a button and that causes a panel to slide out, and you want to reverse that same animation to begin sliding
+	 * in the opposite direction.
+	 *
+	 * @param InAnimation The animation to play
+	 * @param PlayMode Specifies the playback mode
+	 * @param PlaybackSpeed The speed at which the animation should play
+	 * @param bRestoreState Restores widgets to their pre-animated state when the animation stops
+	 */
+	UFUNCTION(BlueprintCallable, BlueprintCosmetic, Category = "User Interface|Animation")
+	UMG_API void QueuePlayAnimationReverse(UWidgetAnimation* InAnimation, float PlaybackSpeed = 1.0f, bool bRestoreState = false);
+
+	/**
+	 * Stops an already running animation in this widget
+	 * 
+	 * @param The name of the animation to stop
+	 */
+	UFUNCTION(BlueprintCallable, Category="User Interface|Animation")
+	UMG_API void QueueStopAnimation(const UWidgetAnimation* InAnimation);
+
+	/**
+	 * Stop All actively running animations.
+	 */
+	UFUNCTION(BlueprintCallable, Category="User Interface|Animation")
+	UMG_API void QueueStopAllAnimations();
+
+	/**
+	 * Pauses an already running animation in this widget
+	 * 
+	 * @param The name of the animation to pause
+	 * @return the time point the animation was at when it was paused, relative to its start position.  Use this as the StartAtTime when you trigger PlayAnimation.
+	 */
+	UFUNCTION(BlueprintCallable, BlueprintCosmetic, Category="User Interface|Animation")
+	UMG_API float QueuePauseAnimation(const UWidgetAnimation* InAnimation);
+
 	/**
 	 * Plays an animation in this widget a specified number of times
 	 * 
@@ -1286,6 +1420,10 @@ private:
 
 public:
 
+	/** Animation transitions to trigger on next tick */
+	UPROPERTY(Transient)
+	TArray<FQueuedWidgetAnimationTransition> QueuedWidgetAnimationTransitions;
+
 	/** All the sequence players currently playing */
 	UPROPERTY(Transient)
 	TArray<TObjectPtr<UUMGSequencePlayer>> ActiveSequencePlayers;
@@ -1428,6 +1566,8 @@ protected:
 	UMG_API UUMGSequencePlayer* GetSequencePlayer(const UWidgetAnimation* InAnimation) const;
 	UMG_API UUMGSequencePlayer* GetOrAddSequencePlayer(UWidgetAnimation* InAnimation);
 
+	UMG_API void ExecuteQueuedAnimationTransitions();
+
 	UMG_API void ConditionalTearDownAnimations();
 
 	UMG_API void TearDownAnimations();
