Sequencer: clear cache of UMG widget director instances when invalidating cached data.

diff --git a/Engine/Source/Runtime/UMG/Private/Animation/UMGSequencePlayer.cpp b/Engine/Source/Runtime/UMG/Private/Animation/UMGSequencePlayer.cpp
--- a/Engine/Source/Runtime/UMG/Private/Animation/UMGSequencePlayer.cpp
+++ b/Engine/Source/Runtime/UMG/Private/Animation/UMGSequencePlayer.cpp
@@ -564,6 +564,10 @@ void UUMGSequencePlayer::RemoveEvaluationData()
 		SequenceInstance->Ledger.UnlinkEverything(Linker);
 		SequenceInstance->InvalidateCachedData();
 	}
+
+	PRAGMA_DISABLE_DEPRECATION_WARNINGS
+	RootTemplateInstance.ResetDirectorInstances();
+	PRAGMA_ENABLE_DEPRECATION_WARNINGS
 }
 
 void UUMGSequencePlayer::TearDown()
