UMG: Update remaining UTextLayoutWidget properties to work with Style Sheets

diff --git a/Engine/Plugins/Runtime/CommonUI/Source/CommonUI/Private/CommonTextBlock.cpp b/Engine/Plugins/Runtime/CommonUI/Source/CommonUI/Private/CommonTextBlock.cpp
--- a/Engine/Plugins/Runtime/CommonUI/Source/CommonUI/Private/CommonTextBlock.cpp
+++ b/Engine/Plugins/Runtime/CommonUI/Source/CommonUI/Private/CommonTextBlock.cpp
@@ -349,8 +349,7 @@ void UCommonTextBlock::Serialize(FArchive& Ar)
 
 void UCommonTextBlock::SetWrapTextWidth(int32 InWrapTextAt)
 {
-	WrapTextAt = InWrapTextAt;
-	SynchronizeProperties();
+	SetWrapTextAt(InWrapTextAt);
 }
 
 void UCommonTextBlock::SetTextCase(bool bUseAllCaps)
diff --git a/Engine/Source/Runtime/UMG/Private/Components/MultiLineEditableText.cpp b/Engine/Source/Runtime/UMG/Private/Components/MultiLineEditableText.cpp
--- a/Engine/Source/Runtime/UMG/Private/Components/MultiLineEditableText.cpp
+++ b/Engine/Source/Runtime/UMG/Private/Components/MultiLineEditableText.cpp
@@ -114,16 +114,51 @@ void UMultiLineEditableText::SynchronizeProperties()
 }
 
 
-void UMultiLineEditableText::SetJustification(ETextJustify::Type InJustification)
+void UMultiLineEditableText::OnShapedTextOptionsChanged(FShapedTextOptions InShapedTextOptions)
 {
-	Super::SetJustification(InJustification);
+	Super::OnShapedTextOptionsChanged(InShapedTextOptions);
+	if (MyMultiLineEditableText.IsValid())
+	{
+		InShapedTextOptions.SynchronizeShapedTextProperties(*MyMultiLineEditableText);
+	}
+}
 
+void UMultiLineEditableText::OnJustificationChanged(ETextJustify::Type InJustification)
+{
+	Super::OnJustificationChanged(InJustification);
 	if (MyMultiLineEditableText.IsValid())
 	{
 		MyMultiLineEditableText->SetJustification(InJustification);
 	}
 }
 
+void UMultiLineEditableText::OnWrappingPolicyChanged(ETextWrappingPolicy InWrappingPolicy)
+{
+	Super::OnWrappingPolicyChanged(InWrappingPolicy);
+	if (MyMultiLineEditableText.IsValid())
+	{
+		MyMultiLineEditableText->SetWrappingPolicy(InWrappingPolicy);
+	}
+}
+
+void UMultiLineEditableText::OnAutoWrapTextChanged(bool InAutoWrapText)
+{
+	Super::OnAutoWrapTextChanged(InAutoWrapText);
+	if (MyMultiLineEditableText.IsValid())
+	{
+		MyMultiLineEditableText->SetAutoWrapText(InAutoWrapText);
+	}
+}
+
+void UMultiLineEditableText::OnWrapTextAtChanged(float InWrapTextAt)
+{
+	Super::OnWrapTextAtChanged(InWrapTextAt);
+	if (MyMultiLineEditableText.IsValid())
+	{
+		MyMultiLineEditableText->SetWrapTextAt(InWrapTextAt);
+	}
+}
+
 void UMultiLineEditableText::OnLineHeightPercentageChanged(float InLineHeightPercentage)
 {
 	Super::OnLineHeightPercentageChanged(InLineHeightPercentage);
diff --git a/Engine/Source/Runtime/UMG/Private/Components/MultiLineEditableTextBox.cpp b/Engine/Source/Runtime/UMG/Private/Components/MultiLineEditableTextBox.cpp
--- a/Engine/Source/Runtime/UMG/Private/Components/MultiLineEditableTextBox.cpp
+++ b/Engine/Source/Runtime/UMG/Private/Components/MultiLineEditableTextBox.cpp
@@ -125,16 +125,51 @@ void UMultiLineEditableTextBox::SynchronizeProperties()
 	Super::SynchronizeTextLayoutProperties(*MyEditableTextBlock);
 }
 
-void UMultiLineEditableTextBox::SetJustification(ETextJustify::Type InJustification)
+void UMultiLineEditableTextBox::OnShapedTextOptionsChanged(FShapedTextOptions InShapedTextOptions)
 {
-	Super::SetJustification(InJustification);
+	Super::OnShapedTextOptionsChanged(InShapedTextOptions);
+	if (MyEditableTextBlock.IsValid())
+	{
+		InShapedTextOptions.SynchronizeShapedTextProperties(*MyEditableTextBlock);
+	}
+}
 
+void UMultiLineEditableTextBox::OnJustificationChanged(ETextJustify::Type InJustification)
+{
+	Super::OnJustificationChanged(InJustification);
 	if (MyEditableTextBlock.IsValid())
 	{
 		MyEditableTextBlock->SetJustification(InJustification);
 	}
 }
 
+void UMultiLineEditableTextBox::OnWrappingPolicyChanged(ETextWrappingPolicy InWrappingPolicy)
+{
+	Super::OnWrappingPolicyChanged(InWrappingPolicy);
+	if (MyEditableTextBlock.IsValid())
+	{
+		MyEditableTextBlock->SetWrappingPolicy(InWrappingPolicy);
+	}
+}
+
+void UMultiLineEditableTextBox::OnAutoWrapTextChanged(bool InAutoWrapText)
+{
+	Super::OnAutoWrapTextChanged(InAutoWrapText);
+	if (MyEditableTextBlock.IsValid())
+	{
+		MyEditableTextBlock->SetAutoWrapText(InAutoWrapText);
+	}
+}
+
+void UMultiLineEditableTextBox::OnWrapTextAtChanged(float InWrapTextAt)
+{
+	Super::OnWrapTextAtChanged(InWrapTextAt);
+	if (MyEditableTextBlock.IsValid())
+	{
+		MyEditableTextBlock->SetWrapTextAt(InWrapTextAt);
+	}
+}
+
 void UMultiLineEditableTextBox::OnLineHeightPercentageChanged(float InLineHeightPercentage)
 {
 	Super::OnLineHeightPercentageChanged(InLineHeightPercentage);
diff --git a/Engine/Source/Runtime/UMG/Private/Components/RichTextBlock.cpp b/Engine/Source/Runtime/UMG/Private/Components/RichTextBlock.cpp
--- a/Engine/Source/Runtime/UMG/Private/Components/RichTextBlock.cpp
+++ b/Engine/Source/Runtime/UMG/Private/Components/RichTextBlock.cpp
@@ -391,18 +391,51 @@ void URichTextBlock::SetDefaultStrikeBrush(const FSlateBrush& InStrikeBrush)
 	DefaultTextStyleOverride.StrikeBrush = InStrikeBrush;
 	ApplyUpdatedDefaultTextStyle();
 }
-PRAGMA_ENABLE_DEPRECATION_WARNINGS
 
-void URichTextBlock::SetJustification(ETextJustify::Type InJustification)
+void URichTextBlock::OnShapedTextOptionsChanged(FShapedTextOptions InShapedTextOptions)
+{
+	Super::OnShapedTextOptionsChanged(InShapedTextOptions);
+	if (MyRichTextBlock.IsValid())
+	{
+		InShapedTextOptions.SynchronizeShapedTextProperties(*MyRichTextBlock);
+	}
+}
+
+void URichTextBlock::OnJustificationChanged(ETextJustify::Type InJustification)
 {
-	Super::SetJustification(InJustification);
+	Super::OnJustificationChanged(InJustification);
 	if (MyRichTextBlock.IsValid())
 	{
 		MyRichTextBlock->SetJustification(InJustification);
 	}
 }
 
-PRAGMA_DISABLE_DEPRECATION_WARNINGS
+void URichTextBlock::OnWrappingPolicyChanged(ETextWrappingPolicy InWrappingPolicy)
+{
+	Super::OnWrappingPolicyChanged(InWrappingPolicy);
+	if (MyRichTextBlock.IsValid())
+	{
+		MyRichTextBlock->SetWrappingPolicy(InWrappingPolicy);
+	}
+}
+
+void URichTextBlock::OnAutoWrapTextChanged(bool InAutoWrapText)
+{
+	Super::OnAutoWrapTextChanged(InAutoWrapText);
+	if (MyRichTextBlock.IsValid())
+	{
+		MyRichTextBlock->SetAutoWrapText(InAutoWrapText);
+	}
+}
+
+void URichTextBlock::OnWrapTextAtChanged(float InWrapTextAt)
+{
+	Super::OnWrapTextAtChanged(InWrapTextAt);
+	if (MyRichTextBlock.IsValid())
+	{
+		MyRichTextBlock->SetWrapTextAt(InWrapTextAt);
+	}
+}
 
 void URichTextBlock::OnLineHeightPercentageChanged(float InLineHeightPercentage)
 {
diff --git a/Engine/Source/Runtime/UMG/Private/Components/TextBlock.cpp b/Engine/Source/Runtime/UMG/Private/Components/TextBlock.cpp
--- a/Engine/Source/Runtime/UMG/Private/Components/TextBlock.cpp
+++ b/Engine/Source/Runtime/UMG/Private/Components/TextBlock.cpp
@@ -143,15 +143,51 @@ void UTextBlock::SetStrikeBrush(FSlateBrush InStrikeBrush)
 	}
 }
 
-void UTextBlock::SetJustification(ETextJustify::Type InJustification)
+void UTextBlock::OnShapedTextOptionsChanged(FShapedTextOptions InShapedTextOptions)
 {
-	Super::SetJustification(InJustification);
+	Super::OnShapedTextOptionsChanged(InShapedTextOptions);
+	if (MyTextBlock.IsValid())
+	{
+		InShapedTextOptions.SynchronizeShapedTextProperties(*MyTextBlock);
+	}
+}
+
+void UTextBlock::OnJustificationChanged(ETextJustify::Type InJustification)
+{
+	Super::OnJustificationChanged(InJustification);
 	if (MyTextBlock.IsValid())
 	{
 		MyTextBlock->SetJustification(InJustification);
 	}
 }
 
+void UTextBlock::OnWrappingPolicyChanged(ETextWrappingPolicy InWrappingPolicy)
+{
+	Super::OnWrappingPolicyChanged(InWrappingPolicy);
+	if (MyTextBlock.IsValid())
+	{
+		MyTextBlock->SetWrappingPolicy(InWrappingPolicy);
+	}
+}
+
+void UTextBlock::OnAutoWrapTextChanged(bool InAutoWrapText)
+{
+	Super::OnAutoWrapTextChanged(InAutoWrapText);
+	if (MyTextBlock.IsValid())
+	{
+		MyTextBlock->SetAutoWrapText(InAutoWrapText);
+	}
+}
+
+void UTextBlock::OnWrapTextAtChanged(float InWrapTextAt)
+{
+	Super::OnWrapTextAtChanged(InWrapTextAt);
+	if (MyTextBlock.IsValid())
+	{
+		MyTextBlock->SetWrapTextAt(InWrapTextAt);
+	}
+}
+
 void UTextBlock::OnLineHeightPercentageChanged(float InLineHeightPercentage)
 {
 	Super::OnLineHeightPercentageChanged(InLineHeightPercentage);
diff --git a/Engine/Source/Runtime/UMG/Private/Components/TextWidgetTypes.cpp b/Engine/Source/Runtime/UMG/Private/Components/TextWidgetTypes.cpp
--- a/Engine/Source/Runtime/UMG/Private/Components/TextWidgetTypes.cpp
+++ b/Engine/Source/Runtime/UMG/Private/Components/TextWidgetTypes.cpp
@@ -26,6 +26,51 @@ UTextLayoutWidget::UTextLayoutWidget(const FObjectInitializer& ObjectInitializer
 	ApplyLineHeightToBottomLine = true;
 }
 
+void UTextLayoutWidget::SetJustification(ETextJustify::Type InJustification)
+{
+	if (Justification != InJustification)
+	{
+		Justification = InJustification;
+		OnJustificationChanged(InJustification);
+	}
+}
+
+void UTextLayoutWidget::SetShapedTextOptions(FShapedTextOptions InShapedTextOptions)
+{
+	if (ShapedTextOptions != InShapedTextOptions)
+	{
+		ShapedTextOptions = InShapedTextOptions;
+		OnShapedTextOptionsChanged(InShapedTextOptions);
+	}
+}
+
+void UTextLayoutWidget::SetWrappingPolicy(ETextWrappingPolicy InWrappingPolicy)
+{
+	if (WrappingPolicy != InWrappingPolicy)
+	{
+		WrappingPolicy = InWrappingPolicy;
+		OnWrappingPolicyChanged(InWrappingPolicy);
+	}
+}
+
+void UTextLayoutWidget::SetAutoWrapText(bool InAutoWrapText)
+{
+	if (AutoWrapText != InAutoWrapText)
+	{
+		AutoWrapText = InAutoWrapText;
+		OnAutoWrapTextChanged(InAutoWrapText);
+	}
+}
+
+void UTextLayoutWidget::SetWrapTextAt(float InWrapTextAt)
+{
+	if (WrapTextAt != InWrapTextAt)
+	{
+		WrapTextAt = InWrapTextAt;
+		OnWrapTextAtChanged(InWrapTextAt);
+	}
+}
+
 void UTextLayoutWidget::SetLineHeightPercentage(float InLineHeightPercentage)
 {
 	if (LineHeightPercentage != InLineHeightPercentage)
diff --git a/Engine/Source/Runtime/UMG/Public/Components/MultiLineEditableText.h b/Engine/Source/Runtime/UMG/Public/Components/MultiLineEditableText.h
--- a/Engine/Source/Runtime/UMG/Public/Components/MultiLineEditableText.h
+++ b/Engine/Source/Runtime/UMG/Public/Components/MultiLineEditableText.h
@@ -165,10 +165,6 @@ public:
 	UFUNCTION(BlueprintCallable, Category = "Appearance")
 	UMG_API void SetFontOutlineMaterial(UMaterialInterface* InMaterial);
 
-	//~ Begin UTextLayoutWidget Interface
-	UMG_API virtual void SetJustification(ETextJustify::Type InJustification) override;
-	//~ End UTextLayoutWidget Interface
-
 	//~ Begin UWidget Interface
 	UMG_API virtual void SynchronizeProperties() override;
 	//~ End UWidget Interface
@@ -183,6 +179,11 @@ public:
 
 protected:
 	//~ Begin UTextLayoutWidget Interface
+	UMG_API virtual void OnShapedTextOptionsChanged(FShapedTextOptions InShapedTextOptions) override;
+	UMG_API virtual void OnJustificationChanged(ETextJustify::Type InJustification) override;
+	UMG_API virtual void OnWrappingPolicyChanged(ETextWrappingPolicy InWrappingPolicy) override;
+	UMG_API virtual void OnAutoWrapTextChanged(bool InAutoWrapText) override;
+	UMG_API virtual void OnWrapTextAtChanged(float InWrapTextAt) override;
 	UMG_API virtual void OnLineHeightPercentageChanged(float InLineHeightPercentage) override;
 	UMG_API virtual void OnApplyLineHeightToBottomLineChanged(bool InApplyLineHeightToBottomLine) override;
 	UMG_API virtual void OnMarginChanged(const FMargin& InMargin) override;
diff --git a/Engine/Source/Runtime/UMG/Public/Components/MultiLineEditableTextBox.h b/Engine/Source/Runtime/UMG/Public/Components/MultiLineEditableTextBox.h
--- a/Engine/Source/Runtime/UMG/Public/Components/MultiLineEditableTextBox.h
+++ b/Engine/Source/Runtime/UMG/Public/Components/MultiLineEditableTextBox.h
@@ -132,10 +132,6 @@ public:
 	//TODO UMG Add Set Font
 
 public:
-	//~ Begin UTextLayoutWidget Interface
-	UMG_API virtual void SetJustification(ETextJustify::Type InJustification) override;
-	//~ End UTextLayoutWidget Interface
-
 	//~ Begin UWidget Interface
 	UMG_API virtual void SynchronizeProperties() override;
 	//~ End UWidget Interface
@@ -151,6 +147,11 @@ public:
 
 protected:
 	//~ Begin UTextLayoutWidget Interface
+	UMG_API virtual void OnShapedTextOptionsChanged(FShapedTextOptions InShapedTextOptions) override;
+	UMG_API virtual void OnJustificationChanged(ETextJustify::Type InJustification) override;
+	UMG_API virtual void OnWrappingPolicyChanged(ETextWrappingPolicy InWrappingPolicy) override;
+	UMG_API virtual void OnAutoWrapTextChanged(bool InAutoWrapText) override;
+	UMG_API virtual void OnWrapTextAtChanged(float InWrapTextAt) override;
 	UMG_API virtual void OnLineHeightPercentageChanged(float InLineHeightPercentage) override;
 	UMG_API virtual void OnApplyLineHeightToBottomLineChanged(bool InApplyLineHeightToBottomLine) override;
 	UMG_API virtual void OnMarginChanged(const FMargin& InMargin) override;
diff --git a/Engine/Source/Runtime/UMG/Public/Components/RichTextBlock.h b/Engine/Source/Runtime/UMG/Public/Components/RichTextBlock.h
--- a/Engine/Source/Runtime/UMG/Public/Components/RichTextBlock.h
+++ b/Engine/Source/Runtime/UMG/Public/Components/RichTextBlock.h
@@ -173,10 +173,6 @@ public:
 public:
 	UMG_API URichTextBlock(const FObjectInitializer& ObjectInitializer);
 
-	//~ Begin UTextLayoutWidget Interface
-	UMG_API virtual void SetJustification(ETextJustify::Type InJustification) override;
-	//~ End UTextLayoutWidget Interface
-
 	// UWidget interface
 	UMG_API virtual void SynchronizeProperties() override;
 	// End of UWidget interface
@@ -228,6 +224,11 @@ public:
 
 protected:
 	//~ Begin UTextLayoutWidget Interface
+	UMG_API virtual void OnShapedTextOptionsChanged(FShapedTextOptions InShapedTextOptions) override;
+	UMG_API virtual void OnJustificationChanged(ETextJustify::Type InJustification) override;
+	UMG_API virtual void OnWrappingPolicyChanged(ETextWrappingPolicy InWrappingPolicy) override;
+	UMG_API virtual void OnAutoWrapTextChanged(bool InAutoWrapText) override;
+	UMG_API virtual void OnWrapTextAtChanged(float InWrapTextAt) override;
 	UMG_API virtual void OnLineHeightPercentageChanged(float InLineHeightPercentage) override;
 	UMG_API virtual void OnApplyLineHeightToBottomLineChanged(bool InApplyLineHeightToBottomLine) override;
 	UMG_API virtual void OnMarginChanged(const FMargin& InMargin) override;
diff --git a/Engine/Source/Runtime/UMG/Public/Components/TextBlock.h b/Engine/Source/Runtime/UMG/Public/Components/TextBlock.h
--- a/Engine/Source/Runtime/UMG/Public/Components/TextBlock.h
+++ b/Engine/Source/Runtime/UMG/Public/Components/TextBlock.h
@@ -231,10 +231,6 @@ protected:
 	}
 
 public:
-	//~ Begin UTextLayoutWidget Interface
-	UMG_API virtual void SetJustification(ETextJustify::Type InJustification) override;
-	//~ End UTextLayoutWidget Interface
-
 	//~ Begin UWidget Interface
 	UMG_API virtual void SynchronizeProperties() override;
 	//~ End UWidget Interface
@@ -257,6 +253,11 @@ public:
 
 protected:
 	//~ Begin UTextLayoutWidget Interface
+	UMG_API virtual void OnShapedTextOptionsChanged(FShapedTextOptions InShapedTextOptions) override;
+	UMG_API virtual void OnJustificationChanged(ETextJustify::Type InJustification) override;
+	UMG_API virtual void OnWrappingPolicyChanged(ETextWrappingPolicy InWrappingPolicy) override;
+	UMG_API virtual void OnAutoWrapTextChanged(bool InAutoWrapText) override;
+	UMG_API virtual void OnWrapTextAtChanged(float InWrapTextAt) override;
 	UMG_API virtual void OnLineHeightPercentageChanged(float InLineHeightPercentage) override;
 	UMG_API virtual void OnApplyLineHeightToBottomLineChanged(bool InApplyLineHeightToBottomLine) override;
 	UMG_API virtual void OnMarginChanged(const FMargin& InMargin) override;
diff --git a/Engine/Source/Runtime/UMG/Public/Components/TextWidgetTypes.h b/Engine/Source/Runtime/UMG/Public/Components/TextWidgetTypes.h
--- a/Engine/Source/Runtime/UMG/Public/Components/TextWidgetTypes.h
+++ b/Engine/Source/Runtime/UMG/Public/Components/TextWidgetTypes.h
@@ -25,12 +25,25 @@ struct FShapedTextOptions
 
 	/** Synchronize the properties with the given widget. A template as the Slate widgets conform to the same API, but don't derive from a common base. */
 	template <typename TWidgetType>
-	void SynchronizeShapedTextProperties(TWidgetType& InWidget)
+	void SynchronizeShapedTextProperties(TWidgetType& InWidget) const
 	{
 		InWidget.SetTextShapingMethod(bOverride_TextShapingMethod ? TOptional<ETextShapingMethod>(TextShapingMethod) : TOptional<ETextShapingMethod>());
 		InWidget.SetTextFlowDirection(bOverride_TextFlowDirection ? TOptional<ETextFlowDirection>(TextFlowDirection) : TOptional<ETextFlowDirection>());
 	}
 
+	inline bool operator==(const FShapedTextOptions& Other) const
+	{
+		return bOverride_TextShapingMethod == Other.bOverride_TextShapingMethod
+			&& bOverride_TextFlowDirection == Other.bOverride_TextFlowDirection
+			&& (!bOverride_TextShapingMethod || TextShapingMethod == Other.TextShapingMethod)
+			&& (!bOverride_TextFlowDirection || TextFlowDirection == Other.TextFlowDirection);
+	}
+
+	inline bool operator!=(const FShapedTextOptions& Other) const
+	{
+		return !operator==(Other);
+	}
+
 	/**  */
 	UPROPERTY(EditAnywhere, Category=Localization, meta=(InlineEditConditionToggle))
 	uint8 bOverride_TextShapingMethod : 1;
@@ -63,13 +76,22 @@ public:
 	float GetWrapTextAt() const { return WrapTextAt; }
 
 	UFUNCTION(BlueprintSetter)
-	virtual void SetJustification(ETextJustify::Type InJustification) { Justification = InJustification; }
+	UMG_API virtual void SetJustification(ETextJustify::Type InJustification);
 
+	UMG_API void SetShapedTextOptions(FShapedTextOptions InShapedTextOptions);
+	UMG_API void SetWrappingPolicy(ETextWrappingPolicy InWrappingPolicy);
+	UMG_API void SetAutoWrapText(bool InAutoWrapText);
+	UMG_API void SetWrapTextAt(float InWrapTextAt);
 	UMG_API void SetLineHeightPercentage(float InLineHeightPercentage);
 	UMG_API void SetApplyLineHeightToBottomLine(bool InApplyLineHeightToBottomLine);
 	UMG_API void SetMargin(const FMargin& InMargin);
 
 protected:
+	virtual void OnShapedTextOptionsChanged(FShapedTextOptions InShapedTextOptions) {};
+	virtual void OnJustificationChanged(ETextJustify::Type InJustification) {};
+	virtual void OnWrappingPolicyChanged(ETextWrappingPolicy InWrappingPolicy) {};
+	virtual void OnAutoWrapTextChanged(bool InAutoWrapText) {};
+	virtual void OnWrapTextAtChanged(float InWrapTextAt) {};
 	virtual void OnLineHeightPercentageChanged(float InLineHeightPercentage) {};
 	virtual void OnApplyLineHeightToBottomLineChanged(bool InApplyLineHeightToBottomLine) {};
 	virtual void OnMarginChanged(const FMargin& InMargin) {};
@@ -90,23 +112,23 @@ protected:
 	}
 
 	/** Controls how the text within this widget should be shaped. */
-	UPROPERTY(EditAnywhere, BlueprintReadOnly, Category=Localization, AdvancedDisplay, meta=(ShowOnlyInnerProperties))
+	UPROPERTY(EditAnywhere, BlueprintReadWrite, Setter, Category=Localization, AdvancedDisplay, meta=(ShowOnlyInnerProperties))
 	FShapedTextOptions ShapedTextOptions;
 
 	/** How the text should be aligned with the margin. */
-	UPROPERTY(EditAnywhere, BlueprintReadWrite, BlueprintSetter=SetJustification, Category=Appearance)
+	UPROPERTY(EditAnywhere, BlueprintReadWrite, BlueprintSetter=SetJustification, Setter, Category=Appearance)
 	TEnumAsByte<ETextJustify::Type> Justification;
 
 	/** The wrapping policy to use. */
-	UPROPERTY(EditAnywhere, BlueprintReadOnly, Category=Wrapping, AdvancedDisplay)
+	UPROPERTY(EditAnywhere, BlueprintReadWrite, Setter, Category=Wrapping, AdvancedDisplay)
 	ETextWrappingPolicy WrappingPolicy;
 
 	/** True if we're wrapping text automatically based on the computed horizontal space for this widget. */
-	UPROPERTY(EditAnywhere, BlueprintReadOnly, Category=Wrapping)
+	UPROPERTY(EditAnywhere, BlueprintReadWrite, Setter, Category=Wrapping)
 	uint8 AutoWrapText:1;
 
 	/** Whether text wraps onto a new line when it's length exceeds this width; if this value is zero or negative, no wrapping occurs. */
-	UPROPERTY(EditAnywhere, BlueprintReadOnly, Category=Wrapping)
+	UPROPERTY(EditAnywhere, BlueprintReadWrite, Setter, Category=Wrapping)
 	float WrapTextAt;
 
 	/** The amount of blank space left around the edges of text area. */
