UMG: Added defensive checks to avoid nullptr access caused by erroneous GC reference elimination

diff --git a/Engine/Source/Runtime/UMG/Private/UserWidget.cpp b/Engine/Source/Runtime/UMG/Private/UserWidget.cpp
--- a/Engine/Source/Runtime/UMG/Private/UserWidget.cpp
+++ b/Engine/Source/Runtime/UMG/Private/UserWidget.cpp
@@ -600,7 +600,14 @@ void UUserWidget::TearDownAnimations()
 
 	for (UUMGSequencePlayer* Player : StoppedSequencePlayers)
 	{
-		Player->TearDown();
+		// These null checks should not be necessary since StoppedSequencePlayers should always keep them alive
+		//   but there are some rare cases where the GC may eliminate and null these references if 'this' is an
+		//   instance of a BP class. This is a bug in the reference collection semantics, but is too intricate
+		//   to fix in a timely manner so we have to be defensive here.
+		if (Player)
+		{
+			Player->TearDown();
+		}
 	}
 
 	ActiveSequencePlayers.Empty();
