PR #11546: Fixing UMG Image and Border's GetDynamicMaterial functions not checking MID outer

diff --git a/Engine/Source/Runtime/UMG/Private/Components/Border.cpp b/Engine/Source/Runtime/UMG/Private/Components/Border.cpp
--- a/Engine/Source/Runtime/UMG/Private/Components/Border.cpp
+++ b/Engine/Source/Runtime/UMG/Private/Components/Border.cpp
@@ -313,7 +313,7 @@ UMaterialInstanceDynamic* UBorder::GetDynamicMaterial()
 	{
 		UMaterialInstanceDynamic* DynamicMaterial = Cast<UMaterialInstanceDynamic>(Material);
 
-		if ( !DynamicMaterial )
+		if ( !DynamicMaterial || DynamicMaterial->GetOuter() != this)
 		{
 			DynamicMaterial = UMaterialInstanceDynamic::Create(Material, this);
 			Background.SetResourceObject(DynamicMaterial);
@@ -427,4 +427,3 @@ const FText UBorder::GetPaletteCategory()
 /////////////////////////////////////////////////////
 
 #undef LOCTEXT_NAMESPACE
-
diff --git a/Engine/Source/Runtime/UMG/Private/Components/Image.cpp b/Engine/Source/Runtime/UMG/Private/Components/Image.cpp
--- a/Engine/Source/Runtime/UMG/Private/Components/Image.cpp
+++ b/Engine/Source/Runtime/UMG/Private/Components/Image.cpp
@@ -391,7 +391,7 @@ void UImage::SetBrushFromSoftMaterial(TSoftObjectPtr<UMaterialInterface> SoftMat
 PRAGMA_DISABLE_DEPRECATION_WARNINGS
 UMaterialInstanceDynamic* UImage::GetDynamicMaterial()
 {
-	UMaterialInterface* Material = NULL;
+	UMaterialInterface* Material = nullptr;
 
 	UObject* Resource = Brush.GetResourceObject();
 	Material = Cast<UMaterialInterface>(Resource);
@@ -400,7 +400,7 @@ UMaterialInstanceDynamic* UImage::GetDynamicMaterial()
 	{
 		UMaterialInstanceDynamic* DynamicMaterial = Cast<UMaterialInstanceDynamic>(Material);
 
-		if ( !DynamicMaterial )
+		if ( !DynamicMaterial || DynamicMaterial->GetOuter() != this )
 		{
 			DynamicMaterial = UMaterialInstanceDynamic::Create(Material, this);
 			SetBrushResourceObject(DynamicMaterial);
@@ -415,7 +415,7 @@ UMaterialInstanceDynamic* UImage::GetDynamicMaterial()
 
 	//TODO UMG can we do something for textures?  General purpose dynamic material for them?
 
-	return NULL;
+	return nullptr;
 }
 PRAGMA_ENABLE_DEPRECATION_WARNINGS
 
@@ -449,4 +449,3 @@ const FText UImage::GetPaletteCategory()
 /////////////////////////////////////////////////////
 
 #undef LOCTEXT_NAMESPACE
-
