Panel Slot:
- Add GetContent BP callable function to retrieve the content of a Slot

diff --git a/Engine/Source/Runtime/UMG/Private/Components/PanelSlot.cpp b/Engine/Source/Runtime/UMG/Private/Components/PanelSlot.cpp
--- a/Engine/Source/Runtime/UMG/Private/Components/PanelSlot.cpp
+++ b/Engine/Source/Runtime/UMG/Private/Components/PanelSlot.cpp
@@ -13,6 +13,11 @@ UPanelSlot::UPanelSlot(const FObjectInitializer& ObjectInitializer)
 {
 }
 
+UWidget* UPanelSlot::GetContent() const
+{
+	return Content;
+}
+
 #if WITH_EDITOR
 bool UPanelSlot::IsDesignTime() const
 {
diff --git a/Engine/Source/Runtime/UMG/Public/Components/PanelSlot.h b/Engine/Source/Runtime/UMG/Public/Components/PanelSlot.h
--- a/Engine/Source/Runtime/UMG/Public/Components/PanelSlot.h
+++ b/Engine/Source/Runtime/UMG/Public/Components/PanelSlot.h
@@ -20,6 +20,9 @@ public:
 
 	UPROPERTY(Instanced)
 	TObjectPtr<class UWidget> Content;
+
+	UFUNCTION(BlueprintCallable, Category = "Layout|Panel Slot")
+	UMG_API UWidget* GetContent() const;
 	
 #if WITH_EDITOR
 	UMG_API bool IsDesignTime() const;
