UMG: Prevent ticking partially destroyed or soon-to-be-destroyed tick managers

diff --git a/Engine/Source/Runtime/UMG/Private/Animation/UMGSequenceTickManager.cpp b/Engine/Source/Runtime/UMG/Private/Animation/UMGSequenceTickManager.cpp
--- a/Engine/Source/Runtime/UMG/Private/Animation/UMGSequenceTickManager.cpp
+++ b/Engine/Source/Runtime/UMG/Private/Animation/UMGSequenceTickManager.cpp
@@ -129,6 +129,12 @@ void UUMGSequenceTickManager::TickWidgetAnimations(float DeltaSeconds)
 		return;
 	}
 
+	if (IsUnreachable() || HasAnyFlags(RF_BeginDestroyed) || Linker == nullptr || Linker->IsUnreachable() || Linker->HasAnyFlags(RF_BeginDestroyed))
+	{
+		// Speculatively ignore any kinds of updates if any of the required objects are in the process of being torn down
+		return;
+	}
+
 	// Don't tick the animation if inside of a PostLoad
 	if (FUObjectThreadContext::Get().IsRoutingPostLoad)
 	{
