Fix Screen Space widgets not working in split screen

#tests PIE Splitscreen

diff --git a/Engine/Source/Runtime/UMG/Private/Components/WidgetComponent.cpp b/Engine/Source/Runtime/UMG/Private/Components/WidgetComponent.cpp
--- a/Engine/Source/Runtime/UMG/Private/Components/WidgetComponent.cpp
+++ b/Engine/Source/Runtime/UMG/Private/Components/WidgetComponent.cpp
@@ -2,6 +2,8 @@
 
 #include "Components/WidgetComponent.h"
 #include "Engine/GameInstance.h"
+#include "GameFramework/Pawn.h"
+#include "GameFramework/PlayerController.h"
 #include "Materials/Material.h"
 #include "PrimitiveViewRelevance.h"
 #include "PrimitiveSceneProxy.h"
@@ -1713,6 +1715,14 @@ ULocalPlayer* UWidgetComponent::GetOwnerPlayer() const
 	{
 		return OwnerPlayer;
 	}
+
+	if (APawn* OwningPawn = Cast<APawn>(GetOwner()))
+	{
+		if (APlayerController* OwningController = Cast<APlayerController>(OwningPawn->GetController()))
+		{
+			return OwningController->GetLocalPlayer();
+		}
+	}
 	
 	if (UWorld* LocalWorld = GetWorld())
 	{
