Added UUserWidget::NativeOnFinishedPointerInput which is called from SObjectWidget::OnFinishedPointerInput after all input events have already been processed for a widget (this callback existed for Slate widgets but not UMG)

diff --git a/Engine/Source/Runtime/UMG/Private/Slate/SObjectWidget.cpp b/Engine/Source/Runtime/UMG/Private/Slate/SObjectWidget.cpp
--- a/Engine/Source/Runtime/UMG/Private/Slate/SObjectWidget.cpp
+++ b/Engine/Source/Runtime/UMG/Private/Slate/SObjectWidget.cpp
@@ -514,6 +514,11 @@ FReply SObjectWidget::OnTouchForceChanged(const FGeometry& MyGeometry, const FPo
 	return FReply::Unhandled();
 }
 
+void SObjectWidget::OnFinishedPointerInput()
+{
+	WidgetObject->NativeOnFinishedPointerInput();
+}
+
 FNavigationReply SObjectWidget::OnNavigation(const FGeometry& MyGeometry, const FNavigationEvent& InNavigationEvent)
 {
 	if (WidgetObject && WidgetObject->NativeSupportsCustomNavigation())
diff --git a/Engine/Source/Runtime/UMG/Public/Blueprint/UserWidget.h b/Engine/Source/Runtime/UMG/Public/Blueprint/UserWidget.h
--- a/Engine/Source/Runtime/UMG/Public/Blueprint/UserWidget.h
+++ b/Engine/Source/Runtime/UMG/Public/Blueprint/UserWidget.h
@@ -1552,6 +1552,7 @@ protected:
 	UMG_API virtual FCursorReply NativeOnCursorQuery( const FGeometry& InGeometry, const FPointerEvent& InCursorEvent );
 	UMG_API virtual FNavigationReply NativeOnNavigation(const FGeometry& InGeometry, const FNavigationEvent& InNavigationEvent);
 	UMG_API virtual void NativeOnMouseCaptureLost(const FCaptureLostEvent& CaptureLostEvent);
+	UMG_API virtual void NativeOnFinishedPointerInput() {}
 
 protected:
 
diff --git a/Engine/Source/Runtime/UMG/Public/Slate/SObjectWidget.h b/Engine/Source/Runtime/UMG/Public/Slate/SObjectWidget.h
--- a/Engine/Source/Runtime/UMG/Public/Slate/SObjectWidget.h
+++ b/Engine/Source/Runtime/UMG/Public/Slate/SObjectWidget.h
@@ -95,6 +95,7 @@ class SObjectWidget : public SCompoundWidget, public FGCObject
 	UMG_API virtual FReply OnTouchEnded(const FGeometry& MyGeometry, const FPointerEvent& InTouchEvent) override;
 	UMG_API virtual FReply OnMotionDetected(const FGeometry& MyGeometry, const FMotionEvent& InMotionEvent) override;
 	UMG_API virtual FReply OnTouchForceChanged(const FGeometry& MyGeometry, const FPointerEvent& InTouchEvent) override;
+	UMG_API virtual void OnFinishedPointerInput() override;
 
 	UMG_API virtual FNavigationReply OnNavigation(const FGeometry& MyGeometry, const FNavigationEvent& InNavigationEvent) override;
 
