Remove enum state functionality
#tests PIE

diff --git a/Engine/Plugins/Runtime/CommonUI/Source/CommonUI/Private/CommonButtonBase.cpp b/Engine/Plugins/Runtime/CommonUI/Source/CommonUI/Private/CommonButtonBase.cpp
--- a/Engine/Plugins/Runtime/CommonUI/Source/CommonUI/Private/CommonButtonBase.cpp
+++ b/Engine/Plugins/Runtime/CommonUI/Source/CommonUI/Private/CommonButtonBase.cpp
@@ -491,19 +491,20 @@ void UCommonButtonBase::SetIsEnabled(bool bInIsEnabled)
 {
 	bool bValueChanged = bButtonEnabled != bInIsEnabled;
 
+	// Change the underlying enabled bool but do not broadcast because we don't want to propagate it to the underlying SWidget
+	bool bOldBroadcastState = bShouldBroadcastState;
+	bShouldBroadcastState = false;
 	if (bInIsEnabled)
 	{
-		TGuardValue<bool> StateBroadcastGuard(bShouldBroadcastState, false);
 		Super::SetIsEnabled(bInIsEnabled);
 		EnableButton();
 	}
 	else
 	{
-		// Change the underlying enabled bool but do not call the case because we don't want to propogate it to the underlying SWidget
-		TGuardValue<bool> StateBroadcastGuard(bShouldBroadcastState, false);
 		Super::SetIsEnabled(bInIsEnabled);
 		DisableButton();
 	}
+	bShouldBroadcastState = bOldBroadcastState;
 
 	if (bValueChanged)
 	{
diff --git a/Engine/Source/Runtime/UMG/Private/Binding/States/Tests/WidgetStateBitfieldTests.cpp b/Engine/Source/Runtime/UMG/Private/Binding/States/Tests/WidgetStateBitfieldTests.cpp
--- a/Engine/Source/Runtime/UMG/Private/Binding/States/Tests/WidgetStateBitfieldTests.cpp
+++ b/Engine/Source/Runtime/UMG/Private/Binding/States/Tests/WidgetStateBitfieldTests.cpp
@@ -50,22 +50,6 @@ bool FOperatorBoolTest::RunTest(const FString& Parameters)
 		GPassing &= TestEqual(TEXT("Has binary is false"), Test.HasBinaryStates(), false);
 	}
 
-	{
-		FWidgetStateBitfield Test = {};
-
-		Test.SetEnumStateSlow(FName("CheckedState"), 1);
-		GPassing &= TestEqual(TEXT("Any enum is true"), Test, true);
-
-		GPassing &= TestEqual(TEXT("Has enum is true"), Test.HasEnumStates(), true);
-		GPassing &= TestEqual(TEXT("Has empty unused is false"), Test.HasEmptyUsedEnumStates(), false);
-
-		Test.ClearEnumState(FName("CheckedState"));
-		GPassing &= TestEqual(TEXT("No enum is false"), Test, false);
-
-		GPassing &= TestEqual(TEXT("Has enum is false"), Test.HasEnumStates(), false);
-		GPassing &= TestEqual(TEXT("Has empty unused is false - none in use"), Test.HasEmptyUsedEnumStates(), false);
-	}
-
 	return GPassing;
 }
 
@@ -106,39 +90,6 @@ bool FOperatorBitwiseAndTest::RunTest(const FString& Parameters)
 		GPassing &= TestEqual(TEXT("Bit: Hovered & PressedHovered is true"), UWidgetHoveredStateRegistration::Bit.Intersect(PressedHovered), true);
 	}
 
-	{
-		FWidgetStateBitfield Checked = {};
-		FWidgetStateBitfield Unchecked = {};
-
-		Checked.SetEnumState(UWidgetCheckedStateRegistration::Checked);
-		Unchecked.SetEnumState(UWidgetCheckedStateRegistration::Unchecked);
-
-		GPassing &= TestEqual(TEXT("Checked & Unchecked is false"), Checked.Intersect(Unchecked), false);
-
-		GPassing &= TestEqual(TEXT("Bit: Checked & Unchecked is false"), UWidgetCheckedStateRegistration::Checked.Intersect(UWidgetCheckedStateRegistration::Unchecked), false);
-		GPassing &= TestEqual(TEXT("Bit: Checked & Unchecked is false"), Checked.Intersect(UWidgetCheckedStateRegistration::Unchecked), false);
-		GPassing &= TestEqual(TEXT("Bit: Checked & Unchecked is false"), UWidgetCheckedStateRegistration::Checked.Intersect(Unchecked), false);
-	}
-
-	{
-		FWidgetStateBitfield Checked = {};
-		FWidgetStateBitfield Hovered = {};
-		FWidgetStateBitfield CheckedHovered = {};
-
-		Checked.SetEnumState(UWidgetCheckedStateRegistration::Checked);
-		Hovered.SetBinaryState(UWidgetHoveredStateRegistration::Bit, true);
-
-		CheckedHovered.SetEnumState(UWidgetCheckedStateRegistration::Checked);
-		CheckedHovered.SetBinaryState(UWidgetHoveredStateRegistration::Bit, true);
-
-		GPassing &= TestEqual(TEXT("Checked & Hovered is false"), Checked.Intersect(Hovered), false);
-		GPassing &= TestEqual(TEXT("Checked & CheckedHovered is false"), Checked.Intersect(CheckedHovered), false);
-		GPassing &= TestEqual(TEXT("CheckedHovered & Checked is true"), CheckedHovered.Intersect(Checked), true);
-		GPassing &= TestEqual(TEXT("Hovered & CheckedHovered is false"), Hovered.Intersect(CheckedHovered), false);
-		GPassing &= TestEqual(TEXT("CheckedHovered & Hovered is true"), CheckedHovered.Intersect(Hovered), true);
-		GPassing &= TestEqual(TEXT("CheckedHovered & CheckedHovered is true"), CheckedHovered.Intersect(CheckedHovered), true);
-	}
-
 	return GPassing;
 }
 
@@ -170,25 +121,6 @@ bool FOperatorBitwiseOrTest::RunTest(const FString& Parameters)
 		GPassing &= TestEqual(TEXT("Bit: Hovered | Pressed is PressedHovered"), UWidgetHoveredStateRegistration::Bit.Union(UWidgetPressedStateRegistration::Bit), PressedHovered);
 	}
 
-	{
-		FWidgetStateBitfield Checked = {};
-		FWidgetStateBitfield Unchecked = {};
-		FWidgetStateBitfield CheckedUnchecked = {};
-
-		Checked.SetEnumState(UWidgetCheckedStateRegistration::Checked);
-		Unchecked.SetEnumState(UWidgetCheckedStateRegistration::Unchecked);
-		CheckedUnchecked.SetEnumState(UWidgetCheckedStateRegistration::Checked.Union(UWidgetCheckedStateRegistration::Unchecked));
-
-		GPassing &= TestEqual(TEXT("Checked | Unchecked is CheckedUnchecked"), Checked.Union(Unchecked), CheckedUnchecked);
-		GPassing &= TestEqual(TEXT("Checked | Unchecked is true"), Checked.Union(Unchecked), true);
-
-		GPassing &= TestEqual(TEXT("CheckedUnchecked & Checked is true"), CheckedUnchecked.Intersect(Checked), true);
-		GPassing &= TestEqual(TEXT("CheckedUnchecked & Unchecked is true"), CheckedUnchecked.Intersect(Unchecked), true);
-
-		GPassing &= TestEqual(TEXT("Checked & CheckedUnchecked is true"), Checked.Intersect(CheckedUnchecked), true);
-		GPassing &= TestEqual(TEXT("Unchecked & CheckedUnchecked is true"), Unchecked.Intersect(CheckedUnchecked), true);
-	}
-
 	return GPassing;
 }
 
@@ -213,20 +145,6 @@ bool FOperatorBitwiseNotTest::RunTest(const FString& Parameters)
 		GPassing &= TestEqual(TEXT("~Pressed does not have any flags Pressed"), (~Pressed).HasAnyFlags(Pressed), false);
 	}
 
-	{
-		FWidgetStateBitfield Checked = {};
-		FWidgetStateBitfield Unchecked = {};
-		FWidgetStateBitfield CheckedUnchecked = {};
-
-		Checked.SetEnumState(UWidgetCheckedStateRegistration::Checked);
-		Unchecked.SetEnumState(UWidgetCheckedStateRegistration::Unchecked);
-		CheckedUnchecked.SetEnumState(UWidgetCheckedStateRegistration::Checked.Union(UWidgetCheckedStateRegistration::Unchecked));
-
-		GPassing &= TestEqual(TEXT("Checked & ~Checked is false"), Checked.Intersect(~Checked), false);
-		GPassing &= TestEqual(TEXT("Unchecked & ~Unchecked is false"), Unchecked.Intersect(~Unchecked), false);
-		GPassing &= TestEqual(TEXT("CheckedUnchecked & ~CheckedUnchecked is false"), CheckedUnchecked.Intersect(~CheckedUnchecked), false);
-	}
-
 	return GPassing;
 }
 
@@ -303,32 +221,6 @@ bool FHasFlagTest::RunTest(const FString& Parameters)
 		GPassing &= TestEqual(TEXT("PressedHovered all flag true"), Test.HasAllFlags(PressedHovered), false);
 	}
 
-	{
-		FWidgetStateBitfield Test = {};
-
-		Test.SetEnumState(UWidgetCheckedStateRegistration::Checked);
-		GPassing &= TestEqual(TEXT("Checked any flag true"), Test.HasAnyFlags(UWidgetCheckedStateRegistration::Checked), true);
-		GPassing &= TestEqual(TEXT("Checked all flag true"), Test.HasAllFlags(UWidgetCheckedStateRegistration::Checked), true);
-		GPassing &= TestEqual(TEXT("Checked any enum flag true"), Test.HasAnyEnumFlags(UWidgetCheckedStateRegistration::Checked), true);
-		GPassing &= TestEqual(TEXT("Checked all enum flag true"), Test.HasAllEnumFlags(UWidgetCheckedStateRegistration::Checked), true);
-
-		Test.SetEnumState(UWidgetCheckedStateRegistration::Unchecked);
-		GPassing &= TestEqual(TEXT("Unchecked any flag true"), Test.HasAnyFlags(UWidgetCheckedStateRegistration::Unchecked), true);
-		GPassing &= TestEqual(TEXT("Unchecked all flag true"), Test.HasAllFlags(UWidgetCheckedStateRegistration::Unchecked), true);
-		GPassing &= TestEqual(TEXT("Unchecked any enum flag true"), Test.HasAnyEnumFlags(UWidgetCheckedStateRegistration::Unchecked), true);
-		GPassing &= TestEqual(TEXT("Unchecked all enum flag true"), Test.HasAllEnumFlags(UWidgetCheckedStateRegistration::Unchecked), true);
-		GPassing &= TestEqual(TEXT("Checked any flag false"), Test.HasAnyFlags(UWidgetCheckedStateRegistration::Checked), false);
-		GPassing &= TestEqual(TEXT("Checked all flag false"), Test.HasAllFlags(UWidgetCheckedStateRegistration::Checked), false);
-		GPassing &= TestEqual(TEXT("Checked any enum flag false"), Test.HasAnyEnumFlags(UWidgetCheckedStateRegistration::Checked), false);
-		GPassing &= TestEqual(TEXT("Checked all enum flag false"), Test.HasAllEnumFlags(UWidgetCheckedStateRegistration::Checked), false);
-
-		Test.ClearEnumState(UWidgetCheckedStateRegistration::Unchecked);
-		GPassing &= TestEqual(TEXT("Unchecked any flag false"), Test.HasAnyFlags(UWidgetCheckedStateRegistration::Unchecked), false);
-		GPassing &= TestEqual(TEXT("Unchecked all flag false"), Test.HasAllFlags(UWidgetCheckedStateRegistration::Unchecked), false);
-		GPassing &= TestEqual(TEXT("Unchecked any enum flag false"), Test.HasAnyEnumFlags(UWidgetCheckedStateRegistration::Unchecked), false);
-		GPassing &= TestEqual(TEXT("Unchecked all enum flag false"), Test.HasAllEnumFlags(UWidgetCheckedStateRegistration::Unchecked), false);
-	}
-
 	return GPassing;
 }
 
diff --git a/Engine/Source/Runtime/UMG/Private/Binding/States/WidgetStateBitfield.cpp b/Engine/Source/Runtime/UMG/Private/Binding/States/WidgetStateBitfield.cpp
--- a/Engine/Source/Runtime/UMG/Private/Binding/States/WidgetStateBitfield.cpp
+++ b/Engine/Source/Runtime/UMG/Private/Binding/States/WidgetStateBitfield.cpp
@@ -17,7 +17,6 @@ FWidgetStateBitfield::FWidgetStateBitfield(const FName InStateName)
 
 FWidgetStateBitfield::FWidgetStateBitfield(const FName InStateName, const uint8 InValue)
 {
-	SetEnumStateSlow(InStateName, InValue);
 }
 
 FWidgetStateBitfield FWidgetStateBitfield::operator~() const
@@ -31,58 +30,6 @@ FWidgetStateBitfield FWidgetStateBitfield::Intersect(const FWidgetStateBitfield&
 {
 	FWidgetStateBitfield Temp;
 	Temp.BinaryStates = BinaryStates & Rhs.BinaryStates;
-	Temp.IsEnumStateNotAllowed = IsEnumStateUsed & Rhs.IsEnumStateUsed;
-
-	// See'IsEnumStateUsed' in header, we want to track enum state usage additively even with '&'
-	Temp.IsEnumStateUsed = IsEnumStateUsed | Rhs.IsEnumStateUsed;
-
-	// Only set enum states if all RHS binary state tests passed, else enum pass can shadow binary failures
-	bool bBinaryStatesPassed = !Rhs.BinaryStates || Temp.BinaryStates;
-	if (bBinaryStatesPassed && Temp.IsEnumStateUsed != 0)
-	{
-		for (uint8 EnumStateIndex = 0; EnumStateIndex < EnumStates.Num(); EnumStateIndex++)
-		{
-			bool bIsStateUsed = Temp.IsEnumStateUsed & (1 << EnumStateIndex);
-
-			if (bIsStateUsed)
-			{
-				bool bIsNotAllowedLhs = IsEnumStateNotAllowed & (1 << EnumStateIndex);
-				bool bIsNotAllowedRhs = Rhs.IsEnumStateNotAllowed & (1 << EnumStateIndex);
-				bool bIsAllowanceSame = bIsNotAllowedLhs == bIsNotAllowedRhs;
-
-				if (bIsAllowanceSame)
-				{
-					// If the RHS doesn't use this enum state, carry over the LHS enum state set
-					// Allows for tests against bitfields where RHS doesn't use LHS enum states
-					// to pass (Ex: {checked, hovered} & {hovered} should pass).
-
-					bool bRhsNotUsed = (Rhs.IsEnumStateUsed & (1 << EnumStateIndex)) == 0;
-
-					if (bRhsNotUsed)
-					{
-						Temp.EnumStates[EnumStateIndex] = EnumStates[EnumStateIndex];
-					}
-					else
-					{
-						Temp.EnumStates[EnumStateIndex] = EnumStates[EnumStateIndex].Intersect(Rhs.EnumStates[EnumStateIndex]);
-					}
-				}
-				else if (bIsNotAllowedRhs)
-				{
-					Temp.EnumStates[EnumStateIndex] = EnumStates[EnumStateIndex].Difference(Rhs.EnumStates[EnumStateIndex]);
-				}
-				else if (bIsNotAllowedLhs)
-				{
-					Temp.EnumStates[EnumStateIndex] = Rhs.EnumStates[EnumStateIndex].Difference(EnumStates[EnumStateIndex]);
-				}
-				else
-				{
-					checkf(false, TEXT("Should not be able to reach, all allowance combinations covered above"));
-				}
-			}
-		}
-	}
-
 	return Temp;
 }
 
@@ -90,50 +37,13 @@ FWidgetStateBitfield FWidgetStateBitfield::Union(const FWidgetStateBitfield& Rhs
 {
 	FWidgetStateBitfield Temp;
 	Temp.BinaryStates = BinaryStates | Rhs.BinaryStates;
-	Temp.IsEnumStateUsed = IsEnumStateUsed | Rhs.IsEnumStateUsed;
-
-	// See'IsEnumStateNotAllowed' in header, we wain to maintain the allowed set so shrink the not allowed
-	Temp.IsEnumStateNotAllowed = IsEnumStateNotAllowed & Rhs.IsEnumStateNotAllowed;
-
-	if (Temp.IsEnumStateUsed != 0)
-	{
-		for (uint8 EnumStateIndex = 0; EnumStateIndex < EnumStates.Num(); EnumStateIndex++)
-		{
-			bool bIsStateUsed = Temp.IsEnumStateUsed & (1 << EnumStateIndex);
-			bool bIsStateAllowedUsed = Temp.IsEnumStateUsed & (1 << EnumStateIndex);
-			if (bIsStateUsed)
-			{
-				bool bIsNotAllowedLhs = IsEnumStateNotAllowed & (1 << EnumStateIndex);
-				bool bIsNotAllowedRhs = Rhs.IsEnumStateNotAllowed & (1 << EnumStateIndex);
-				bool bIsAllowanceSame = bIsNotAllowedLhs == bIsNotAllowedRhs;
-
-				if (bIsAllowanceSame)
-				{
-					Temp.EnumStates[EnumStateIndex] = EnumStates[EnumStateIndex].Union(Rhs.EnumStates[EnumStateIndex]);
-				}
-				else if (bIsNotAllowedRhs)
-				{
-					Temp.EnumStates[EnumStateIndex] = EnumStates[EnumStateIndex];
-				}
-				else if (bIsNotAllowedLhs)
-				{
-					Temp.EnumStates[EnumStateIndex] = Rhs.EnumStates[EnumStateIndex];
-				}
-				else
-				{
-					checkf(false, TEXT("Should not be able to reach, all allowance combinations covered above"));
-				}
-			}
-		}
-	}
-
 	return Temp;
 }
 
 FWidgetStateBitfield::operator bool() const
 {
 	// We have to check 'IsEnumStateUsed' here since otherwise a completely empty statefield would pass
-	return (BinaryStates != 0 || IsEnumStateUsed != 0) && !HasEmptyUsedEnumStates();
+	return BinaryStates != 0;
 }
 
 bool FWidgetStateBitfield::HasBinaryStates() const
@@ -143,35 +53,22 @@ bool FWidgetStateBitfield::HasBinaryStates() const
 
 bool FWidgetStateBitfield::HasEnumStates() const
 {
-	return IsEnumStateUsed != 0;
+	return false;
 }
 
 bool FWidgetStateBitfield::HasEmptyUsedEnumStates() const
 {
-	if (IsEnumStateUsed != 0)
-	{
-		for (uint8 EnumStateIndex = 0; EnumStateIndex < EnumStates.Num(); EnumStateIndex++)
-		{
-			if (IsEnumStateUsed & (1 << EnumStateIndex))
-			{
-				if (EnumStates[EnumStateIndex].IsEmpty())
-				{
-					return true;
-				}
-			}
-		}
-	}
 	return false;
 }
 
 bool FWidgetStateBitfield::HasAnyFlags(const FWidgetStateBitfield& InBitfield) const
 {
-	return HasAnyBinaryFlags(InBitfield) || HasAnyEnumFlags(InBitfield);
+	return HasAnyBinaryFlags(InBitfield);
 }
 
 bool FWidgetStateBitfield::HasAllFlags(const FWidgetStateBitfield& InBitfield) const
 {
-	return HasAllBinaryFlags(InBitfield) && HasAllEnumFlags(InBitfield);
+	return HasAllBinaryFlags(InBitfield);
 }
 
 bool FWidgetStateBitfield::HasAnyBinaryFlags(const FWidgetStateBitfield& InBitfield) const
@@ -186,44 +83,11 @@ bool FWidgetStateBitfield::HasAllBinaryFlags(const FWidgetStateBitfield& InBitfi
 
 bool FWidgetStateBitfield::HasAnyEnumFlags(const FWidgetStateBitfield& InBitfield) const
 {
-	const uint16 IsEnumStateShared = IsEnumStateUsed & InBitfield.IsEnumStateUsed;
-
-	if (IsEnumStateShared != 0)
-	{
-		for (uint8 EnumStateIndex = 0; EnumStateIndex < EnumStates.Num(); EnumStateIndex++)
-		{
-			if (IsEnumStateShared & (1 << EnumStateIndex))
-			{
-				if (!EnumStates[EnumStateIndex].Intersect(InBitfield.EnumStates[EnumStateIndex]).IsEmpty())
-				{
-					return true;
-				}
-			}
-		}
-	}
 	return false;
 }
 
 bool FWidgetStateBitfield::HasAllEnumFlags(const FWidgetStateBitfield& InBitfield) const
 {
-	const uint16 IsEnumStateShared = IsEnumStateUsed & InBitfield.IsEnumStateUsed;
-
-	if (IsEnumStateShared == InBitfield.IsEnumStateUsed)
-	{
-		for (uint8 EnumStateIndex = 0; EnumStateIndex < EnumStates.Num(); EnumStateIndex++)
-		{
-			if (InBitfield.IsEnumStateUsed & (1 << EnumStateIndex))
-			{
-				if (!EnumStates[EnumStateIndex].Includes(InBitfield.EnumStates[EnumStateIndex]))
-				{
-					return false;
-				}
-			}
-		}
-
-		return true;
-	}
-
 	return false;
 }
 
@@ -235,7 +99,6 @@ void FWidgetStateBitfield::SetState(const FWidgetStateBitfield& InBitfield)
 void FWidgetStateBitfield::NegateStates()
 {
 	NegateBinaryStates();
-	NegateEnumStates();
 }
 
 void FWidgetStateBitfield::SetBinaryState(uint8 BinaryStateIndex, bool BinaryStateValue)
@@ -268,55 +131,28 @@ void FWidgetStateBitfield::NegateBinaryStates()
 
 void FWidgetStateBitfield::SetEnumState(uint8 EnumStateIndex, uint8 EnumStateValue)
 {
-	IsEnumStateUsed = IsEnumStateUsed | (1 << EnumStateIndex);
-	EnumStates[EnumStateIndex] = { EnumStateValue };
 }
 
 void FWidgetStateBitfield::SetEnumState(const FWidgetStateBitfield& EnumStateBitfield)
 {
-	IsEnumStateUsed |= EnumStateBitfield.IsEnumStateUsed;
-
-	for (uint8 EnumStateIndex = 0; EnumStateIndex < EnumStates.Num(); EnumStateIndex++)
-	{
-		if (EnumStateBitfield.IsEnumStateUsed & (1 << EnumStateIndex))
-		{
-			EnumStates[EnumStateIndex] = EnumStateBitfield.EnumStates[EnumStateIndex];
-		}
-	}
 }
 
 void FWidgetStateBitfield::SetEnumStateSlow(FName EnumStateName, uint8 EnumStateValue)
 {
-	uint8 EnumStateIndex = UWidgetStateSettings::Get()->GetEnumStateIndex(EnumStateName);
-	SetEnumState(EnumStateIndex, EnumStateValue);
 }
 
 void FWidgetStateBitfield::ClearEnumState(const FWidgetStateBitfield& EnumStateBitfield)
 {
-	IsEnumStateUsed &= ~EnumStateBitfield.IsEnumStateUsed;
-
-	for (uint8 EnumStateIndex = 0; EnumStateIndex < EnumStates.Num(); EnumStateIndex++)
-	{
-		if (EnumStateBitfield.IsEnumStateUsed & (1 << EnumStateIndex))
-		{
-			EnumStates[EnumStateIndex].Empty();
-		}
-	}
 }
 
 void FWidgetStateBitfield::ClearEnumState(uint8 EnumStateIndex)
 {
-	IsEnumStateUsed = IsEnumStateUsed & ~(1 << EnumStateIndex);
-	EnumStates[EnumStateIndex].Empty();
 }
 
 void FWidgetStateBitfield::ClearEnumState(FName EnumStateName)
 {
-	uint8 EnumStateIndex = UWidgetStateSettings::Get()->GetEnumStateIndex(EnumStateName);
-	ClearEnumState(EnumStateIndex);
 }
 
 void FWidgetStateBitfield::NegateEnumStates()
 {
-	IsEnumStateNotAllowed = ~IsEnumStateNotAllowed;
 }
diff --git a/Engine/Source/Runtime/UMG/Private/Binding/States/WidgetStateSettings.cpp b/Engine/Source/Runtime/UMG/Private/Binding/States/WidgetStateSettings.cpp
--- a/Engine/Source/Runtime/UMG/Private/Binding/States/WidgetStateSettings.cpp
+++ b/Engine/Source/Runtime/UMG/Private/Binding/States/WidgetStateSettings.cpp
@@ -109,8 +109,12 @@ FWidgetStateBitfield UWidgetStateSettings::GetInitialRegistrationBitfield(const
 		bool InitBinaryState = BinaryStateRegistrationCDO->GetRegisteredWidgetState(InWidget);
 		InitBitfield.SetBinaryState(BinaryStateIndex, InitBinaryState);
 		BinaryStateIndex++;
+
+		ensureMsgf(BinaryStateIndex <= 64, TEXT("Num Binary states should not exceed 64, as it is stored in a uint64. Remove 'UWidgetBinaryStateRegistration' subclasses from your project to address."));
 	}
 
+PRAGMA_DISABLE_DEPRECATION_WARNINGS
+	
 	uint8 EnumStateIndex = 0;
 	for (const TObjectPtr<UWidgetEnumStateRegistration>& EnumStateRegistrationCDO : EnumStateRegistrationCDOs)
 	{
@@ -127,5 +131,7 @@ FWidgetStateBitfield UWidgetStateSettings::GetInitialRegistrationBitfield(const
 		EnumStateIndex++;
 	}
 
+PRAGMA_ENABLE_DEPRECATION_WARNINGS
+	
 	return InitBitfield;
 }
diff --git a/Engine/Source/Runtime/UMG/Private/Components/CheckBox.cpp b/Engine/Source/Runtime/UMG/Private/Components/CheckBox.cpp
--- a/Engine/Source/Runtime/UMG/Private/Components/CheckBox.cpp
+++ b/Engine/Source/Runtime/UMG/Private/Components/CheckBox.cpp
@@ -329,9 +329,11 @@ const FWidgetStateBitfield& UWidgetCheckedStateRegistration::GetBitfieldFromValu
 
 void UWidgetCheckedStateRegistration::InitializeStaticBitfields() const
 {
+PRAGMA_DISABLE_DEPRECATION_WARNINGS
 	Unchecked = FWidgetStateBitfield(GetStateName(), (uint8)ECheckBoxState::Unchecked);
 	Checked = FWidgetStateBitfield(GetStateName(), (uint8)ECheckBoxState::Checked);
 	Undetermined = FWidgetStateBitfield(GetStateName(), (uint8)ECheckBoxState::Undetermined);
+PRAGMA_ENABLE_DEPRECATION_WARNINGS
 }
 
 /////////////////////////////////////////////////////
diff --git a/Engine/Source/Runtime/UMG/Private/Components/Widget.cpp b/Engine/Source/Runtime/UMG/Private/Components/Widget.cpp
--- a/Engine/Source/Runtime/UMG/Private/Components/Widget.cpp
+++ b/Engine/Source/Runtime/UMG/Private/Components/Widget.cpp
@@ -208,6 +208,7 @@ PRAGMA_DISABLE_DEPRECATION_WARNINGS
 	AccessibleWidgetData = nullptr;
 
 	bShouldBroadcastState = true;
+	bWidgetStateInitialized = false;
 
 PRAGMA_ENABLE_DEPRECATION_WARNINGS
 
@@ -1341,13 +1342,6 @@ void UWidget::PreSave(FObjectPreSaveContext ObjectSaveContext)
 	SynchronizeAccessibleData();
 }
 
-void UWidget::ReleaseSlateResources(bool bReleaseChildren)
-{
-	UVisual::ReleaseSlateResources(bReleaseChildren);
-
-	MyWidgetStateBitfield.Reset();
-}
-
 #if WITH_EDITOR
 bool UWidget::Modify(bool bAlwaysMarkDirty)
 {
@@ -1803,14 +1797,15 @@ bool UWidget::AddBinding(FDelegateProperty* DelegateProperty, UObject* SourceObj
 
 FDelegateHandle UWidget::RegisterPostStateListener(const FOnWidgetStateBroadcast::FDelegate& ListenerDelegate, bool bBroadcastCurrentState)
 {
-	if (!MyWidgetStateBitfield.IsValid())
+	if (!bWidgetStateInitialized)
 	{
-		MyWidgetStateBitfield = MakeShared<FWidgetStateBitfield>(UWidgetStateSettings::Get()->GetInitialRegistrationBitfield(this));
+		MyWidgetStateBitfield = UWidgetStateSettings::Get()->GetInitialRegistrationBitfield(this);
+		bWidgetStateInitialized = true;
 	}
 
 	if (bBroadcastCurrentState)
 	{
-		ListenerDelegate.ExecuteIfBound(this, *MyWidgetStateBitfield);
+		ListenerDelegate.ExecuteIfBound(this, MyWidgetStateBitfield);
 	}
 
 	return PostWidgetStateChanged.Add(ListenerDelegate);
@@ -1819,11 +1814,6 @@ FDelegateHandle UWidget::RegisterPostStateListener(const FOnWidgetStateBroadcast
 void UWidget::UnregisterPostStateListener(const FDelegateHandle& ListenerDelegate)
 {
 	PostWidgetStateChanged.Remove(ListenerDelegate);
-
-	if (!PostWidgetStateChanged.IsBound())
-	{
-		MyWidgetStateBitfield.Reset();
-	}
 }
 
 void UWidget::OnBindingChanged(const FName& Property)
@@ -1833,20 +1823,15 @@ void UWidget::OnBindingChanged(const FName& Property)
 
 void UWidget::BroadcastBinaryPostStateChange(const FWidgetStateBitfield& StateChange, bool bInValue)
 {
-	if (bShouldBroadcastState && MyWidgetStateBitfield.IsValid())
+	if (bShouldBroadcastState && bWidgetStateInitialized)
 	{
-		MyWidgetStateBitfield->SetBinaryState(StateChange, bInValue);
-		PostWidgetStateChanged.Broadcast(this, *MyWidgetStateBitfield);
+		MyWidgetStateBitfield.SetBinaryState(StateChange, bInValue);
+		PostWidgetStateChanged.Broadcast(this, MyWidgetStateBitfield);
 	}
 }
 
 void UWidget::BroadcastEnumPostStateChange(const FWidgetStateBitfield& StateChange)
 {
-	if (bShouldBroadcastState && MyWidgetStateBitfield.IsValid())
-	{
-		MyWidgetStateBitfield->SetEnumState(StateChange);
-		PostWidgetStateChanged.Broadcast(this, *MyWidgetStateBitfield);
-	}
 }
 
 namespace UE::UMG::Private
diff --git a/Engine/Source/Runtime/UMG/Public/Binding/States/WidgetStateBitfield.h b/Engine/Source/Runtime/UMG/Public/Binding/States/WidgetStateBitfield.h
--- a/Engine/Source/Runtime/UMG/Public/Binding/States/WidgetStateBitfield.h
+++ b/Engine/Source/Runtime/UMG/Public/Binding/States/WidgetStateBitfield.h
@@ -32,6 +32,7 @@ public:
 	UMG_API FWidgetStateBitfield(const FName InStateName);
 
 	/** Interprets name value pair as enum state */
+	UE_DEPRECATED(5.5, "FWidgetStateBitfield currently no longer supports enum states")
 	UMG_API FWidgetStateBitfield(const FName InStateName, const uint8 InValue);
 
 public:
@@ -65,9 +66,11 @@ public:
 	UMG_API bool HasBinaryStates() const;
 
 	/** True if any enum state is used, does not indicate if state is actually set. Use 'HasEmptyUsedEnumStates' for that. */
+	UE_DEPRECATED(5.5, "FWidgetStateBitfield currently no longer supports enum states")
 	UMG_API bool HasEnumStates() const;
 
 	/** True if any enum used enum states are empty, usually indicating a failed '&' test */
+	UE_DEPRECATED(5.5, "FWidgetStateBitfield currently no longer supports enum states")
 	UMG_API bool HasEmptyUsedEnumStates() const;
 
 public:
@@ -89,9 +92,11 @@ public:
 	UMG_API bool HasAllBinaryFlags(const FWidgetStateBitfield& InBitfield) const;
 
 	/** True if any enum state flag is met */
+	UE_DEPRECATED(5.5, "FWidgetStateBitfield currently no longer supports enum states")
 	UMG_API bool HasAnyEnumFlags(const FWidgetStateBitfield& InBitfield) const;
 
 	/** True if all enum state flags are met */
+	UE_DEPRECATED(5.5, "FWidgetStateBitfield currently no longer supports enum states")
 	UMG_API bool HasAllEnumFlags(const FWidgetStateBitfield& InBitfield) const;
 
 public:
@@ -130,6 +135,7 @@ public:
 	 * @param EnumStateIndex Index for enum state to assign value to
 	 * @param EnumStateValue Value to assign for given enum state
 	 */
+	UE_DEPRECATED(5.5, "FWidgetStateBitfield currently no longer supports enum states")
 	UMG_API void SetEnumState(uint8 EnumStateIndex, uint8 EnumStateValue);
 
 	/**
@@ -137,9 +143,11 @@ public:
 	 *
 	 * @param EnumStateBitfield state bitfield, all used states will copy-overwrite current enum state
 	 */
+	UE_DEPRECATED(5.5, "FWidgetStateBitfield currently no longer supports enum states")
 	UMG_API void SetEnumState(const FWidgetStateBitfield& EnumStateBitfield);
 
 	/** Set enum state to given value, clearing all existing values for that state, slow version w/ FName map lookup */
+	UE_DEPRECATED(5.5, "FWidgetStateBitfield currently no longer supports enum states")
 	UMG_API void SetEnumStateSlow(FName EnumStateName, uint8 EnumStateValue);
 
 	/**
@@ -147,6 +155,7 @@ public:
 	 *
 	 * @param EnumStateBitfield state bitfield, all used states will be cleared
 	 */
+	UE_DEPRECATED(5.5, "FWidgetStateBitfield currently no longer supports enum states")
 	UMG_API void ClearEnumState(const FWidgetStateBitfield& EnumStateBitfield);
 
 	/**
@@ -154,41 +163,19 @@ public:
 	 *
 	 * @param EnumStateIndex Index for enum state to mark as unused
 	 */
+	UE_DEPRECATED(5.5, "FWidgetStateBitfield currently no longer supports enum states")
 	UMG_API void ClearEnumState(uint8 EnumStateIndex);
 
 	/** Clear all existing values for given state & mark the state unused, slow version w/ map lookup */
+	UE_DEPRECATED(5.5, "FWidgetStateBitfield currently no longer supports enum states")
 	UMG_API void ClearEnumState(FName EnumStateName);
 
 	/** Negate only enum states, does not change usage, but marks existing values as not allowed */
+	UE_DEPRECATED(5.5, "FWidgetStateBitfield currently no longer supports enum states")
 	UMG_API void NegateEnumStates();
 
 private:
 
-	/** 
-	 * Note: This implementation is private, we choose to limit 64 binary and 16 enum states. 
-	 * We believe that is more than enum, but if not it's possible to use:
-	 * 
-	 * TBitArray<TInlineAllocator<32>> BinaryStates;
-	 * TArray<TSet<uint8, TInlineAllocator<1>>, TInlineAllocator<8>> EnumStates;
-	 * 
-	 * With some additional modifications to the .cpp impl.
-	 */
-
 	/** Binary states captured by this state field */
 	uint64 BinaryStates = 0;
-
-	/** 
-	 * An empty intersection of enum states may / may not be a reason to convert this field to true, so track states used 
-	 * In particular for '&' IsEnumStateUsed will still grow by '|'. Ex: ECheckboxState & EAssetViewType => false.
-	 */
-	uint16 IsEnumStateUsed = 0;
-
-	/**
-	 * A negation of enum states does not change usage, but indicates that states should be removed on intersection.
-	 * On union, if allowance is the same the state set is union'd, else the allowed set is maintained
-	 */
-	uint16 IsEnumStateNotAllowed = 0;
-
-	/** Binary states captured by this state field. Note: 16 == sizeof(IsEnumStateUsed) * CHAR_BIT */
-	TStaticArray<TSet<uint8, DefaultKeyFuncs<uint8, false>, TInlineSetAllocator<1>>, 16> EnumStates;
 };
diff --git a/Engine/Source/Runtime/UMG/Public/Components/Widget.h b/Engine/Source/Runtime/UMG/Public/Components/Widget.h
--- a/Engine/Source/Runtime/UMG/Public/Components/Widget.h
+++ b/Engine/Source/Runtime/UMG/Public/Components/Widget.h
@@ -3,6 +3,7 @@
 #pragma once
 
 #include "CoreMinimal.h"
+#include "Binding/States/WidgetStateBitfield.h"
 #include "UObject/ObjectMacros.h"
 #include "FieldNotificationDeclaration.h"
 #include "INotifyFieldValueChanged.h"
@@ -41,7 +42,6 @@ class UPanelSlot;
 class UPropertyBinding;
 class UUserWidget;
 struct FDynamicPropertyPath;
-struct FWidgetStateBitfield;
 enum class ECheckBoxState : uint8;
 
 namespace UMWidget
@@ -389,6 +389,12 @@ protected:
 	/** Cached value that indicate if the widget was added to the GameViewportSubsystem. */
 	uint8 bIsManagedByGameViewportSubsystem:1;
 
+	/** False will skip state broadcasts. Useful for child classes to call Super methods without broadcasting early / late. */
+	bool bShouldBroadcastState : 1;
+
+	/** True implies widget state has been initialized. */
+	bool bWidgetStateInitialized : 1;
+
 public:
 #if WITH_EDITORONLY_DATA
 	/** Stores the design time flag setting if the widget is hidden inside the designer */
@@ -1003,10 +1009,6 @@ public:
 	UMG_API virtual void PreSave(FObjectPreSaveContext ObjectSaveContext) override;
 	//~ End UObject
 
-	//~ Begin UVisual
-	UMG_API virtual void ReleaseSlateResources(bool bReleaseChildren) override;
-	//~ End UVisual
-
 	FORCEINLINE bool CanSafelyRouteEvent()
 	{
 		return !IsDesignTime() && CanSafelyRouteCall();
@@ -1110,6 +1112,7 @@ protected:
 	 *
 	 * @param StateChange bitfield marking states that should be changed
 	 */
+	UE_DEPRECATED(5.5, "FWidgetStateBitfield currently no longer supports enum states")
 	UMG_API void BroadcastEnumPostStateChange(const FWidgetStateBitfield& StateChange);
 
 protected:
@@ -1169,10 +1172,7 @@ protected:
 	TWeakPtr<SObjectWidget> MyGCWidget;
 
 	/** The bitfield for this widget's state */
-	TSharedPtr<FWidgetStateBitfield> MyWidgetStateBitfield;
-
-	/** False will skip state broadcasts. Useful for child classes to call Super methods without broadcasting early / late. */
-	bool bShouldBroadcastState;
+	FWidgetStateBitfield MyWidgetStateBitfield;
 
 	/** Delegate that broadcasts after current widget state has fully changed, including all state-related side effects */
 	FOnWidgetStateBroadcast PostWidgetStateChanged;
