Add option to DynamicEntryBoxBase to remove widgets while deleting slate.
#tests PIE

diff --git a/Engine/Source/Runtime/UMG/Private/Components/DynamicEntryBoxBase.cpp b/Engine/Source/Runtime/UMG/Private/Components/DynamicEntryBoxBase.cpp
--- a/Engine/Source/Runtime/UMG/Private/Components/DynamicEntryBoxBase.cpp
+++ b/Engine/Source/Runtime/UMG/Private/Components/DynamicEntryBoxBase.cpp
@@ -92,7 +92,7 @@ const FRadialBoxSettings& UDynamicEntryBoxBase::GetRadialBoxSettings() const
 	return RadialBoxSettings;
 }
 
-void UDynamicEntryBoxBase::RemoveEntryInternal(UUserWidget* EntryWidget)
+void UDynamicEntryBoxBase::RemoveEntryInternal(UUserWidget* EntryWidget, bool bReleaseSlate)
 {
 	if (EntryWidget)
 	{
@@ -120,7 +120,7 @@ void UDynamicEntryBoxBase::RemoveEntryInternal(UUserWidget* EntryWidget)
 				}
 			}
 		}
-		EntryWidgetPool.Release(EntryWidget);
+		EntryWidgetPool.Release(EntryWidget, bReleaseSlate);
 	}
 }
 
diff --git a/Engine/Source/Runtime/UMG/Public/Components/DynamicEntryBoxBase.h b/Engine/Source/Runtime/UMG/Public/Components/DynamicEntryBoxBase.h
--- a/Engine/Source/Runtime/UMG/Public/Components/DynamicEntryBoxBase.h
+++ b/Engine/Source/Runtime/UMG/Public/Components/DynamicEntryBoxBase.h
@@ -129,7 +129,7 @@ protected:
 
 	UMG_API bool IsEntryClassValid(TSubclassOf<UUserWidget> InEntryClass) const;
 	UMG_API UUserWidget* CreateEntryInternal(TSubclassOf<UUserWidget> InEntryClass);
-	UMG_API void RemoveEntryInternal(UUserWidget* EntryWidget);
+	UMG_API void RemoveEntryInternal(UUserWidget* EntryWidget, bool bReleaseSlate = false);
 	UMG_API FMargin BuildEntryPadding(const FVector2D& DesiredSpacing);
 
 	/** Clear out the box entries, optionally deleting the underlying Slate widgets entirely as well. */
