Slate: UWidget::ForceLayoutPrepass for a SlatePrepass when the invalidation enabled. In slowpath, the ForceLayoutPrepass would execute and it has a different behavior compare to fastpath. Now the fastpath is like the slowpath.

diff --git a/Engine/Source/Runtime/UMG/Private/Components/Widget.cpp b/Engine/Source/Runtime/UMG/Private/Components/Widget.cpp
--- a/Engine/Source/Runtime/UMG/Private/Components/Widget.cpp
+++ b/Engine/Source/Runtime/UMG/Private/Components/Widget.cpp
@@ -764,6 +764,7 @@ void UWidget::ForceLayoutPrepass()
 	TSharedPtr<SWidget> SafeWidget = GetCachedWidget();
 	if (SafeWidget.IsValid())
 	{
+		SafeWidget->MarkPrepassAsDirty();
 		SafeWidget->SlatePrepass(SafeWidget->GetTickSpaceGeometry().Scale);
 	}
 }
