UMG: Remove call to PostEdit for editor utility widget. The previous bug could not be reproduced and it cause other issues with widget creation in editor.

diff --git a/Engine/Plugins/Runtime/AudioWidgets/Source/AudioWidgets/Private/AudioRadialSlider.cpp b/Engine/Plugins/Runtime/AudioWidgets/Source/AudioWidgets/Private/AudioRadialSlider.cpp
--- a/Engine/Plugins/Runtime/AudioWidgets/Source/AudioWidgets/Private/AudioRadialSlider.cpp
+++ b/Engine/Plugins/Runtime/AudioWidgets/Source/AudioWidgets/Private/AudioRadialSlider.cpp
@@ -37,9 +37,6 @@ UAudioRadialSlider::UAudioRadialSlider(const FObjectInitializer& ObjectInitializ
 	if (IsEditorWidget())
 	{
 		WidgetStyle = *EditorAudioRadialSliderStyle;
-
-		// The CDO isn't an editor widget and thus won't use the editor style, call post edit change to mark difference from CDO
-		PostEditChange();
 	}
 #endif // WITH_EDITOR
 
diff --git a/Engine/Plugins/Runtime/AudioWidgets/Source/AudioWidgets/Private/AudioSlider.cpp b/Engine/Plugins/Runtime/AudioWidgets/Source/AudioWidgets/Private/AudioSlider.cpp
--- a/Engine/Plugins/Runtime/AudioWidgets/Source/AudioWidgets/Private/AudioSlider.cpp
+++ b/Engine/Plugins/Runtime/AudioWidgets/Source/AudioWidgets/Private/AudioSlider.cpp
@@ -40,9 +40,6 @@ UAudioSliderBase::UAudioSliderBase(const FObjectInitializer& ObjectInitializer)
 	if (IsEditorWidget())
 	{
 		WidgetStyle = *EditorAudioSliderStyle;
-
-		// The CDO isn't an editor widget and thus won't use the editor style, call post edit change to mark difference from CDO
-		PostEditChange();
 	}
 #endif // WITH_EDITOR
 
diff --git a/Engine/Source/Runtime/UMG/Private/Components/Button.cpp b/Engine/Source/Runtime/UMG/Private/Components/Button.cpp
--- a/Engine/Source/Runtime/UMG/Private/Components/Button.cpp
+++ b/Engine/Source/Runtime/UMG/Private/Components/Button.cpp
@@ -32,9 +32,6 @@ PRAGMA_ENABLE_DEPRECATION_WARNINGS
 PRAGMA_DISABLE_DEPRECATION_WARNINGS
 		WidgetStyle = UE::Slate::Private::FDefaultStyleCache::GetEditor().GetButtonStyle();
 PRAGMA_ENABLE_DEPRECATION_WARNINGS
-
-		// The CDO isn't an editor widget and thus won't use the editor style, call post edit change to mark difference from CDO
-		PostEditChange();
 	}
 #endif // WITH_EDITOR
 
diff --git a/Engine/Source/Runtime/UMG/Private/Components/CheckBox.cpp b/Engine/Source/Runtime/UMG/Private/Components/CheckBox.cpp
--- a/Engine/Source/Runtime/UMG/Private/Components/CheckBox.cpp
+++ b/Engine/Source/Runtime/UMG/Private/Components/CheckBox.cpp
@@ -27,9 +27,6 @@ PRAGMA_DISABLE_DEPRECATION_WARNINGS
 	if (IsEditorWidget())
 	{
 		WidgetStyle = UE::Slate::Private::FDefaultStyleCache::GetEditor().GetCheckboxStyle();
-
-		// The CDO isn't an editor widget and thus won't use the editor style, call post edit change to mark difference from CDO
-		PostEditChange();
 	}
 #endif // WITH_EDITOR
 
diff --git a/Engine/Source/Runtime/UMG/Private/Components/CircularThrobber.cpp b/Engine/Source/Runtime/UMG/Private/Components/CircularThrobber.cpp
--- a/Engine/Source/Runtime/UMG/Private/Components/CircularThrobber.cpp
+++ b/Engine/Source/Runtime/UMG/Private/Components/CircularThrobber.cpp
@@ -26,9 +26,6 @@ PRAGMA_DISABLE_DEPRECATION_WARNINGS
 	if (IsEditorWidget())
 	{
 		Image = UE::Slate::Private::FDefaultStyleCache::GetEditor().GetCircularThrobberBrushStyle();
-
-		// The CDO isn't an editor widget and thus won't use the editor style, call post edit change to mark difference from CDO
-		PostEditChange();
 	}
 #endif // WITH_EDITOR
 
diff --git a/Engine/Source/Runtime/UMG/Private/Components/ComboBoxKey.cpp b/Engine/Source/Runtime/UMG/Private/Components/ComboBoxKey.cpp
--- a/Engine/Source/Runtime/UMG/Private/Components/ComboBoxKey.cpp
+++ b/Engine/Source/Runtime/UMG/Private/Components/ComboBoxKey.cpp
@@ -27,12 +27,6 @@ PRAGMA_DISABLE_DEPRECATION_WARNINGS
 		ItemStyle = IsEditorWidget() ? UE::Slate::Private::FDefaultStyleCache::GetEditor().GetTableRowStyle() : UE::Slate::Private::FDefaultStyleCache::GetRuntime().GetTableRowStyle();
 		ScrollBarStyle = IsEditorWidget() ? UE::Slate::Private::FDefaultStyleCache::GetEditor().GetScrollBarStyle() : UE::Slate::Private::FDefaultStyleCache::GetRuntime().GetScrollBarStyle();
 PRAGMA_ENABLE_DEPRECATION_WARNINGS
-
-		if (IsEditorWidget())
-		{
-			// The CDO isn't an editor widget and thus won't use the editor style, call post edit change to mark difference from CDO
-			PostEditChange();
-		}
 #else
 PRAGMA_DISABLE_DEPRECATION_WARNINGS
 		WidgetStyle = UE::Slate::Private::FDefaultStyleCache::GetRuntime().GetComboBoxStyle();
diff --git a/Engine/Source/Runtime/UMG/Private/Components/ComboBoxString.cpp b/Engine/Source/Runtime/UMG/Private/Components/ComboBoxString.cpp
--- a/Engine/Source/Runtime/UMG/Private/Components/ComboBoxString.cpp
+++ b/Engine/Source/Runtime/UMG/Private/Components/ComboBoxString.cpp
@@ -30,9 +30,6 @@ PRAGMA_DISABLE_DEPRECATION_WARNINGS
 		WidgetStyle = UE::Slate::Private::FDefaultStyleCache::GetEditor().GetComboBoxStyle();
 		ItemStyle = UE::Slate::Private::FDefaultStyleCache::GetEditor().GetTableRowStyle();
 		ScrollBarStyle = UE::Slate::Private::FDefaultStyleCache::GetEditor().GetScrollBarStyle();
-
-		// The CDO isn't an editor widget and thus won't use the editor style, call post edit change to mark difference from CDO
-		PostEditChange();
 	}
 #endif
 	
diff --git a/Engine/Source/Runtime/UMG/Private/Components/EditableText.cpp b/Engine/Source/Runtime/UMG/Private/Components/EditableText.cpp
--- a/Engine/Source/Runtime/UMG/Private/Components/EditableText.cpp
+++ b/Engine/Source/Runtime/UMG/Private/Components/EditableText.cpp
@@ -38,9 +38,6 @@ PRAGMA_DISABLE_DEPRECATION_WARNINGS
 	if (IsEditorWidget())
 	{
 		WidgetStyle = UE::Slate::Private::FDefaultStyleCache::GetEditor().GetEditableTextStyle();
-
-		// The CDO isn't an editor widget and thus won't use the editor style, call post edit change to mark difference from CDO
-		PostEditChange();
 	}
 #endif // WITH_EDITOR
 
diff --git a/Engine/Source/Runtime/UMG/Private/Components/EditableTextBox.cpp b/Engine/Source/Runtime/UMG/Private/Components/EditableTextBox.cpp
--- a/Engine/Source/Runtime/UMG/Private/Components/EditableTextBox.cpp
+++ b/Engine/Source/Runtime/UMG/Private/Components/EditableTextBox.cpp
@@ -53,9 +53,6 @@ PRAGMA_DISABLE_DEPRECATION_WARNINGS
 	if (IsEditorWidget())
 	{
 		WidgetStyle = UE::Slate::Private::FDefaultStyleCache::GetEditor().GetEditableTextBoxStyle();
-
-		// The CDO isn't an editor widget and thus won't use the editor style, call post edit change to mark difference from CDO
-		PostEditChange();
 	}
 #endif // WITH_EDITOR
 PRAGMA_ENABLE_DEPRECATION_WARNINGS
diff --git a/Engine/Source/Runtime/UMG/Private/Components/ExpandableArea.cpp b/Engine/Source/Runtime/UMG/Private/Components/ExpandableArea.cpp
--- a/Engine/Source/Runtime/UMG/Private/Components/ExpandableArea.cpp
+++ b/Engine/Source/Runtime/UMG/Private/Components/ExpandableArea.cpp
@@ -33,9 +33,6 @@ PRAGMA_DISABLE_DEPRECATION_WARNINGS
 	{
 		Style = UE::Slate::Private::FDefaultStyleCache::GetEditor().GetExpandableAreaStyle();
 		BorderBrush = UE::Slate::Private::FDefaultStyleCache::GetEditor().GetExpandableAreaBorderBrush();
-
-		// The CDO isn't an editor widget and thus won't use the editor style, call post edit change to mark difference from CDO
-		PostEditChange();
 	}
 #endif // WITH_EDITOR
 
diff --git a/Engine/Source/Runtime/UMG/Private/Components/InputKeySelector.cpp b/Engine/Source/Runtime/UMG/Private/Components/InputKeySelector.cpp
--- a/Engine/Source/Runtime/UMG/Private/Components/InputKeySelector.cpp
+++ b/Engine/Source/Runtime/UMG/Private/Components/InputKeySelector.cpp
@@ -26,9 +26,6 @@ PRAGMA_DISABLE_DEPRECATION_WARNINGS
 	{
 		WidgetStyle = UE::Slate::Private::FDefaultStyleCache::GetEditor().GetButtonStyle();
 		TextStyle = UE::Slate::Private::FDefaultStyleCache::GetEditor().GetTextBlockStyle();
-
-		// The CDO isn't an editor widget and thus won't use the editor style, call post edit change to mark difference from CDO
-		PostEditChange();
 	}
 #endif // WITH_EDITOR
 
diff --git a/Engine/Source/Runtime/UMG/Private/Components/ListView.cpp b/Engine/Source/Runtime/UMG/Private/Components/ListView.cpp
--- a/Engine/Source/Runtime/UMG/Private/Components/ListView.cpp
+++ b/Engine/Source/Runtime/UMG/Private/Components/ListView.cpp
@@ -28,9 +28,6 @@ UListView::UListView(const FObjectInitializer& ObjectInitializer)
 	{
 		WidgetStyle = UE::Slate::Private::FDefaultStyleCache::GetEditor().GetListViewStyle();
 		ScrollBarStyle = UE::Slate::Private::FDefaultStyleCache::GetEditor().GetScrollBarStyle();
-
-		// The CDO isn't an editor widget and thus won't use the editor style, call post edit change to mark difference from CDO
-		PostEditChange();
 	}
 #endif // WITH_EDITOR
 }
diff --git a/Engine/Source/Runtime/UMG/Private/Components/MultiLineEditableText.cpp b/Engine/Source/Runtime/UMG/Private/Components/MultiLineEditableText.cpp
--- a/Engine/Source/Runtime/UMG/Private/Components/MultiLineEditableText.cpp
+++ b/Engine/Source/Runtime/UMG/Private/Components/MultiLineEditableText.cpp
@@ -26,9 +26,6 @@ UMultiLineEditableText::UMultiLineEditableText(const FObjectInitializer& ObjectI
 	if (IsEditorWidget())
 	{
 		WidgetStyle = UE::Slate::Private::FDefaultStyleCache::GetEditor().GetTextBlockStyle();
-
-		// The CDO isn't an editor widget and thus won't use the editor style, call post edit change to mark difference from CDO
-		PostEditChange();
 	}
 #endif // WITH_EDITOR
 
diff --git a/Engine/Source/Runtime/UMG/Private/Components/MultiLineEditableTextBox.cpp b/Engine/Source/Runtime/UMG/Private/Components/MultiLineEditableTextBox.cpp
--- a/Engine/Source/Runtime/UMG/Private/Components/MultiLineEditableTextBox.cpp
+++ b/Engine/Source/Runtime/UMG/Private/Components/MultiLineEditableTextBox.cpp
@@ -40,9 +40,6 @@ UMultiLineEditableTextBox::UMultiLineEditableTextBox(const FObjectInitializer& O
 	{
 		WidgetStyle = UE::Slate::Private::FDefaultStyleCache::GetEditor().GetEditableTextBoxStyle();
 		TextStyle_DEPRECATED = WidgetStyle.TextStyle;
-
-		// The CDO isn't an editor widget and thus won't use the editor style, call post edit change to mark difference from CDO
-		PostEditChange();
 	}
 
 	bIsFontDeprecationDone = false;
diff --git a/Engine/Source/Runtime/UMG/Private/Components/ProgressBar.cpp b/Engine/Source/Runtime/UMG/Private/Components/ProgressBar.cpp
--- a/Engine/Source/Runtime/UMG/Private/Components/ProgressBar.cpp
+++ b/Engine/Source/Runtime/UMG/Private/Components/ProgressBar.cpp
@@ -22,9 +22,6 @@ PRAGMA_DISABLE_DEPRECATION_WARNINGS
 	if (IsEditorWidget())
 	{
 		WidgetStyle = UE::Slate::Private::FDefaultStyleCache::GetEditor().GetProgressBarStyle();
-
-		// The CDO isn't an editor widget and thus won't use the editor style, call post edit change to mark difference from CDO
-		PostEditChange();
 	}
 #endif // WITH_EDITOR
 
diff --git a/Engine/Source/Runtime/UMG/Private/Components/ScrollBar.cpp b/Engine/Source/Runtime/UMG/Private/Components/ScrollBar.cpp
--- a/Engine/Source/Runtime/UMG/Private/Components/ScrollBar.cpp
+++ b/Engine/Source/Runtime/UMG/Private/Components/ScrollBar.cpp
@@ -30,9 +30,6 @@ PRAGMA_DISABLE_DEPRECATION_WARNINGS
 	if (IsEditorWidget())
 	{
 		WidgetStyle = UE::Slate::Private::FDefaultStyleCache::GetEditor().GetScrollBarStyle();
-
-		// The CDO isn't an editor widget and thus won't use the editor style, call post edit change to mark difference from CDO
-		PostEditChange();
 	}
 #endif // WITH_EDITOR
 PRAGMA_ENABLE_DEPRECATION_WARNINGS
diff --git a/Engine/Source/Runtime/UMG/Private/Components/ScrollBox.cpp b/Engine/Source/Runtime/UMG/Private/Components/ScrollBox.cpp
--- a/Engine/Source/Runtime/UMG/Private/Components/ScrollBox.cpp
+++ b/Engine/Source/Runtime/UMG/Private/Components/ScrollBox.cpp
@@ -47,9 +47,6 @@ PRAGMA_DISABLE_DEPRECATION_WARNINGS
 		WidgetStyle = UE::Slate::Private::FDefaultStyleCache::GetEditor().GetScrollBoxStyle();
 		WidgetBarStyle = UE::Slate::Private::FDefaultStyleCache::GetEditor().GetScrollBarStyle();
 PRAGMA_ENABLE_DEPRECATION_WARNINGS
-
-		// The CDO isn't an editor widget and thus won't use the editor style, call post edit change to mark difference from CDO
-		PostEditChange();
 	}
 #endif // WITH_EDITOR
 PRAGMA_DISABLE_DEPRECATION_WARNINGS
diff --git a/Engine/Source/Runtime/UMG/Private/Components/Slider.cpp b/Engine/Source/Runtime/UMG/Private/Components/Slider.cpp
--- a/Engine/Source/Runtime/UMG/Private/Components/Slider.cpp
+++ b/Engine/Source/Runtime/UMG/Private/Components/Slider.cpp
@@ -33,9 +33,6 @@ PRAGMA_DISABLE_DEPRECATION_WARNINGS
 	if (IsEditorWidget())
 	{
 		WidgetStyle = UE::Slate::Private::FDefaultStyleCache::GetEditor().GetSliderStyle();
-
-		// The CDO isn't an editor widget and thus won't use the editor style, call post edit change to mark difference from CDO
-		PostEditChange();
 	}
 #endif // WITH_EDITOR
 
diff --git a/Engine/Source/Runtime/UMG/Private/Components/SpinBox.cpp b/Engine/Source/Runtime/UMG/Private/Components/SpinBox.cpp
--- a/Engine/Source/Runtime/UMG/Private/Components/SpinBox.cpp
+++ b/Engine/Source/Runtime/UMG/Private/Components/SpinBox.cpp
@@ -45,9 +45,6 @@ PRAGMA_DISABLE_DEPRECATION_WARNINGS
 	if (IsEditorWidget())
 	{
 		WidgetStyle = UE::Slate::Private::FDefaultStyleCache::GetEditor().GetSpinBoxStyle();
-
-		// The CDO isn't an editor widget and thus won't use the editor style, call post edit change to mark difference from CDO
-		PostEditChange();
 	}
 #endif // WITH_EDITOR
 
diff --git a/Engine/Source/Runtime/UMG/Private/Components/Throbber.cpp b/Engine/Source/Runtime/UMG/Private/Components/Throbber.cpp
--- a/Engine/Source/Runtime/UMG/Private/Components/Throbber.cpp
+++ b/Engine/Source/Runtime/UMG/Private/Components/Throbber.cpp
@@ -30,9 +30,6 @@ PRAGMA_DISABLE_DEPRECATION_WARNINGS
 	if (IsEditorWidget())
 	{
 		Image = UE::Slate::Private::FDefaultStyleCache::GetEditor().GetThrobberBrush();
-
-		// The CDO isn't an editor widget and thus won't use the editor style, call post edit change to mark difference from CDO
-		PostEditChange();
 	}
 #endif // WITH_EDITOR
 PRAGMA_ENABLE_DEPRECATION_WARNINGS
diff --git a/Engine/Source/Runtime/UMG/Private/Components/TreeView.cpp b/Engine/Source/Runtime/UMG/Private/Components/TreeView.cpp
--- a/Engine/Source/Runtime/UMG/Private/Components/TreeView.cpp
+++ b/Engine/Source/Runtime/UMG/Private/Components/TreeView.cpp
@@ -18,9 +18,6 @@ UTreeView::UTreeView(const FObjectInitializer& ObjectInitializer)
 	if (IsEditorWidget())
 	{
 		WidgetStyle = UE::Slate::Private::FDefaultStyleCache::GetEditor().GetTreeViewStyle();
-
-		// The CDO isn't an editor widget and thus won't use the editor style, call post edit change to mark difference from CDO
-		PostEditChange();
 	}
 #endif // WITH_EDITOR
 }
