Fixed UWidgetComponent::GetOwnerPlayer crash when owning world has no game instance.

diff --git a/Engine/Source/Runtime/UMG/Private/Components/WidgetComponent.cpp b/Engine/Source/Runtime/UMG/Private/Components/WidgetComponent.cpp
--- a/Engine/Source/Runtime/UMG/Private/Components/WidgetComponent.cpp
+++ b/Engine/Source/Runtime/UMG/Private/Components/WidgetComponent.cpp
@@ -1724,12 +1724,16 @@ ULocalPlayer* UWidgetComponent::GetOwnerPlayer() const
 		}
 	}
 	
-	if (UWorld* LocalWorld = GetWorld())
+	if (const UWorld* LocalWorld = GetWorld())
 	{
-		UGameInstance* GameInstance = LocalWorld->GetGameInstance();
-		check(GameInstance);
-
-		return GameInstance->GetFirstGamePlayer();
+		if (const UGameInstance* GameInstance = LocalWorld->GetGameInstance())
+		{
+			return GameInstance->GetFirstGamePlayer();
+		}
+		else
+		{
+			int32 a = 5;
+		}
 	}
 
 	return nullptr;
