Removing debug code accidentally included in CL 33796879.

#rnx

diff --git a/Engine/Source/Runtime/UMG/Private/Components/WidgetComponent.cpp b/Engine/Source/Runtime/UMG/Private/Components/WidgetComponent.cpp
--- a/Engine/Source/Runtime/UMG/Private/Components/WidgetComponent.cpp
+++ b/Engine/Source/Runtime/UMG/Private/Components/WidgetComponent.cpp
@@ -1730,10 +1730,6 @@ ULocalPlayer* UWidgetComponent::GetOwnerPlayer() const
 		{
 			return GameInstance->GetFirstGamePlayer();
 		}
-		else
-		{
-			int32 a = 5;
-		}
 	}
 
 	return nullptr;
