AutoRTFM-ifiy these Modules

diff --git a/Engine/Source/Runtime/UMG/UMG.Build.cs b/Engine/Source/Runtime/UMG/UMG.Build.cs
--- a/Engine/Source/Runtime/UMG/UMG.Build.cs
+++ b/Engine/Source/Runtime/UMG/UMG.Build.cs
@@ -57,5 +57,7 @@ public class UMG : ModuleRules
 				}
 			);
 		}
+
+		bAllowAutoRTFMInstrumentation = true;
 	}
 }
