Remove AutoRTFM

diff --git a/Engine/Source/Runtime/UMG/UMG.Build.cs b/Engine/Source/Runtime/UMG/UMG.Build.cs
--- a/Engine/Source/Runtime/UMG/UMG.Build.cs
+++ b/Engine/Source/Runtime/UMG/UMG.Build.cs
@@ -57,7 +57,5 @@ public class UMG : ModuleRules
 				}
 			);
 		}
-
-		bAllowAutoRTFMInstrumentation = true;
 	}
 }
