Add CSV Substats to Slate

diff --git a/Engine/Source/Runtime/Slate/Private/Framework/Application/SlateApplication.cpp b/Engine/Source/Runtime/Slate/Private/Framework/Application/SlateApplication.cpp
--- a/Engine/Source/Runtime/Slate/Private/Framework/Application/SlateApplication.cpp
+++ b/Engine/Source/Runtime/Slate/Private/Framework/Application/SlateApplication.cpp
@@ -1272,6 +1272,7 @@ static void PrepassWindowAndChildren(TSharedRef<SWindow> WindowToPrepass, const
 void FSlateApplication::DrawPrepass( TSharedPtr<SWindow> DrawOnlyThisWindow )
 {
 	SCOPED_NAMED_EVENT_TEXT("Slate::Prepass", FColor::Magenta);
+	CSV_SCOPED_TIMING_STAT(Slate, DrawPrePass);
 
 	TSharedPtr<SWindow> CurrentDebuggingWindowPinned = CurrentDebuggingWindow.Pin();
 
@@ -1575,6 +1576,7 @@ void FSlateApplication::TickTime()
 void FSlateApplication::TickPlatform(float DeltaTime)
 {
 	SCOPED_NAMED_EVENT_TEXT("Slate::TickPlatform", FColor::Magenta);
+	CSV_SCOPED_TIMING_STAT(Slate, TickPlatform);
 
 #if WITH_ACCESSIBILITY
 	{
diff --git a/Engine/Source/Runtime/SlateCore/Private/FastUpdate/SlateInvalidationRoot.cpp b/Engine/Source/Runtime/SlateCore/Private/FastUpdate/SlateInvalidationRoot.cpp
--- a/Engine/Source/Runtime/SlateCore/Private/FastUpdate/SlateInvalidationRoot.cpp
+++ b/Engine/Source/Runtime/SlateCore/Private/FastUpdate/SlateInvalidationRoot.cpp
@@ -718,6 +718,7 @@ bool FSlateInvalidationRoot::PaintFastPath_UpdateNextWidget(const FSlateInvalida
 bool FSlateInvalidationRoot::PaintFastPath(const FSlateInvalidationContext& Context)
 {
 	SCOPED_NAMED_EVENT(SWidget_FastPathUpdate, FColor::Green);
+	CSV_SCOPED_TIMING_STAT(Slate, PaintFastPath);
 
 	check(!bNeedsSlowPath);
 
diff --git a/Engine/Source/Runtime/SlateRHIRenderer/Private/SlateRHIRenderer.cpp b/Engine/Source/Runtime/SlateRHIRenderer/Private/SlateRHIRenderer.cpp
--- a/Engine/Source/Runtime/SlateRHIRenderer/Private/SlateRHIRenderer.cpp
+++ b/Engine/Source/Runtime/SlateRHIRenderer/Private/SlateRHIRenderer.cpp
@@ -66,6 +66,8 @@ DECLARE_CYCLE_STAT(TEXT("Total Render Thread time including dependent waits"), S
 
 CSV_DEFINE_CATEGORY(RenderThreadIdle, true);
 
+CSV_DECLARE_CATEGORY_MODULE_EXTERN(SLATECORE_API, Slate);
+
 DECLARE_GPU_DRAWCALL_STAT_NAMED(SlateUI, TEXT("Slate UI"));
 
 // Defines the maximum size that a slate viewport will create
@@ -1693,6 +1695,7 @@ void FSlateRHIRenderer::PrepareToTakeHDRScreenshot(const FIntRect& Rect, TArray<
 void FSlateRHIRenderer::DrawWindows_Private(FSlateDrawBuffer& WindowDrawBuffer)
 {
 	checkSlow(IsThreadSafeForSlateRendering());
+	CSV_SCOPED_TIMING_STAT(Slate, DrawWindows_Private);
 
 	if (bUpdateHDRDisplayInformation && IsHDRAllowed() && IsInGameThread())
 	{
diff --git a/Engine/Source/Runtime/UMG/Private/Slate/SObjectWidget.cpp b/Engine/Source/Runtime/UMG/Private/Slate/SObjectWidget.cpp
--- a/Engine/Source/Runtime/UMG/Private/Slate/SObjectWidget.cpp
+++ b/Engine/Source/Runtime/UMG/Private/Slate/SObjectWidget.cpp
@@ -7,6 +7,8 @@
 #include "Slate/UMGDragDropOp.h"
 #include "SlateGlobals.h"
 
+CSV_DECLARE_CATEGORY_MODULE_EXTERN(SLATECORE_API, Slate);
+
 void SObjectWidget::Construct(const FArguments& InArgs, UUserWidget* InWidgetObject)
 {
 	WidgetObject = InWidgetObject;
@@ -109,8 +111,10 @@ void SObjectWidget::SetPadding(const TAttribute<FMargin>& InMargin)
 
 void SObjectWidget::Tick( const FGeometry& AllottedGeometry, const double InCurrentTime, const float InDeltaTime )
 {
+	// Note: This tick will not execute unless the UserWidget itself ticks.
 #if SLATE_VERBOSE_NAMED_EVENTS
 	SCOPED_NAMED_EVENT_FSTRING(DebugTickEventName, FColor::Turquoise);
+	CSV_SCOPED_TIMING_STAT(Slate, SObjectWidget_Tick);
 #endif
 
 #if WITH_VERY_VERBOSE_SLATE_STATS
