Sequencer- Dynamic Binding Refactor

diff --git a/Engine/Source/Runtime/UMG/Public/Animation/WidgetAnimation.h b/Engine/Source/Runtime/UMG/Public/Animation/WidgetAnimation.h
--- a/Engine/Source/Runtime/UMG/Public/Animation/WidgetAnimation.h
+++ b/Engine/Source/Runtime/UMG/Public/Animation/WidgetAnimation.h
@@ -13,6 +13,10 @@
 class UMovieScene;
 class UUserWidget;
 
+#if WITH_EDITOR
+struct FSlateBrush;
+#endif
+
 
 /**
  * 
@@ -97,10 +101,16 @@ public:
 	virtual void UnbindPossessableObjects(const FGuid& ObjectId) override;
 	virtual void UnbindObjects(const FGuid& ObjectId, const TArray<UObject*>& InObjects, UObject* InContext) override {}
 	virtual void UnbindInvalidObjects(const FGuid& ObjectId, UObject* InContext) override {}
-	virtual void LocateBoundObjects(const FGuid& ObjectId, UObject* Context, TArray<UObject*, TInlineAllocator<1>>& OutObjects) const override;
+	virtual void LocateBoundObjects(const FGuid& ObjectId, const UE::UniversalObjectLocator::FResolveParams& ResolveParams, TSharedPtr<const FSharedPlaybackState> SharedPlaybackState, TArray<UObject*, TInlineAllocator<1>>& OutObjects) const override;
 	virtual UObject* CreateDirectorInstance(TSharedRef<const FSharedPlaybackState> SharedPlaybackState, FMovieSceneSequenceID SequenceID) override;
+
+	virtual void IterateDynamicBindings(const TSharedRef<UE::MovieScene::FSharedPlaybackState> SharedPlaybackState, TFunction<void(const FGuid&, FMovieSceneDynamicBinding&)> InCallback) override;
 #if WITH_EDITOR
 	virtual ETrackSupport IsTrackSupported(TSubclassOf<class UMovieSceneTrack> InTrackClass) const override;
+	virtual const FSlateBrush* GetCustomBrushForBinding(FGuid BindingID) const override;
+
+	DECLARE_MULTICAST_DELEGATE_OneParam(FFixupWidgetDynamicBindingsEvent, UWidgetAnimation*);
+	static UMG_API FFixupWidgetDynamicBindingsEvent FixupWidgetDynamicBindingsEvent;
 #endif
 	// ~UMovieSceneAnimation overrides
 
diff --git a/Engine/Source/Runtime/UMG/Public/Animation/WidgetAnimationBinding.h b/Engine/Source/Runtime/UMG/Public/Animation/WidgetAnimationBinding.h
--- a/Engine/Source/Runtime/UMG/Public/Animation/WidgetAnimationBinding.h
+++ b/Engine/Source/Runtime/UMG/Public/Animation/WidgetAnimationBinding.h
@@ -5,10 +5,19 @@
 #include "CoreMinimal.h"
 #include "UObject/ObjectMacros.h"
 #include "Misc/Guid.h"
+#include "MovieSceneDynamicBinding.h"
+#include "UObject/FortniteMainBranchObjectVersion.h"
 #include "WidgetAnimationBinding.generated.h"
 
 class UUserWidget;
 class UWidgetTree;
+class UMovieSceneSequence;
+
+namespace UE::MovieScene
+{
+	struct FSharedPlaybackState;
+}
+
 
 /**
  * A single object bound to a UMG sequence.
@@ -30,14 +39,20 @@ struct FWidgetAnimationBinding
 	UPROPERTY()
 	bool bIsRootWidget = false;
 
+	UPROPERTY()
+	FMovieSceneDynamicBinding DynamicBinding;
+
 public:
 
 	/**
 	 * Locates a runtime object to animate from the provided tree of widgets.
 	 * @return the runtime object to animate or null if not found
 	 */
+	UE_DEPRECATED(5.5, "Please use the version that takes a SharedPlaybackState and Sequence")
 	UMG_API UObject* FindRuntimeObject(const UWidgetTree& WidgetTree, UUserWidget& UserWidget) const;
 
+	UMG_API UObject* FindRuntimeObject(const UWidgetTree& WidgetTree, UUserWidget& UserWidget, const UMovieSceneSequence* Sequence, TSharedPtr<const UE::MovieScene::FSharedPlaybackState> SharedPlaybackState) const;
+
 	/**
 	 * Compares two widget animation bindings for equality.
 	 *
@@ -47,30 +62,6 @@ public:
 	 */
 	friend bool operator==(const FWidgetAnimationBinding& X, const FWidgetAnimationBinding& Y)
 	{
-		return (X.WidgetName == Y.WidgetName) && (X.SlotWidgetName == Y.SlotWidgetName) && (X.AnimationGuid == Y.AnimationGuid) && (X.bIsRootWidget == Y.bIsRootWidget);
-	}
-
-	/**
-	 * Serializes a widget animation binding from/to the given archive.
-	 *
-	 * @param Ar The archive to serialize to/from.
-	 * @param Binding the binding to serialize.
-	 */
-	friend FArchive& operator<<(FArchive& Ar, FWidgetAnimationBinding& Binding)
-	{
-		Ar << Binding.WidgetName;
-		Ar << Binding.SlotWidgetName;
-		Ar << Binding.AnimationGuid;
-		Ar << Binding.bIsRootWidget;
-		return Ar;
-	}
-
-	friend void operator<<(FStructuredArchive::FSlot Slot, FWidgetAnimationBinding& Binding)
-	{
-		FStructuredArchive::FRecord Record = Slot.EnterRecord();
-		Record << SA_VALUE(TEXT("WidgetName"), Binding.WidgetName);
-		Record << SA_VALUE(TEXT("SlotWidgetName"), Binding.SlotWidgetName);
-		Record << SA_VALUE(TEXT("AnimationGuid"), Binding.AnimationGuid);
-		Record << SA_VALUE(TEXT("bIsRootWidget"), Binding.bIsRootWidget);
+		return (X.WidgetName == Y.WidgetName) && (X.SlotWidgetName == Y.SlotWidgetName) && (X.AnimationGuid == Y.AnimationGuid) && (X.bIsRootWidget == Y.bIsRootWidget) && (X.DynamicBinding.Function == Y.DynamicBinding.Function);
 	}
 };
