PR #11951: Improve error message when focusing on a widget with invalid player controller

diff --git a/Engine/Source/Runtime/UMG/Private/Components/Widget.cpp b/Engine/Source/Runtime/UMG/Private/Components/Widget.cpp
--- a/Engine/Source/Runtime/UMG/Private/Components/Widget.cpp
+++ b/Engine/Source/Runtime/UMG/Private/Components/Widget.cpp
@@ -709,7 +709,7 @@ void UWidget::SetUserFocus(APlayerController* PlayerController)
 	if ( PlayerController == nullptr || !PlayerController->IsLocalPlayerController() || PlayerController->Player == nullptr )
 	{
 #if !(UE_BUILD_SHIPPING || UE_BUILD_TEST)
-		FMessageLog("PIE").Error(LOCTEXT("NoPlayerControllerToFocus", "The PlayerController is not a valid local player so it can't focus the widget."));
+		FMessageLog("PIE").Error()->AddToken(FTextToken::Create(LOCTEXT("NoPlayerControllerToFocus", "The PlayerController is not a valid local player so it can't focus on ")))->AddToken(FUObjectToken::Create(this));
 #endif
 		return;
 	}
