UMG: expose animation lerp speed as a property in UMG for scrollbox and tableviewbase

diff --git a/Engine/Source/Runtime/Slate/Private/Widgets/Layout/SScrollBox.cpp b/Engine/Source/Runtime/Slate/Private/Widgets/Layout/SScrollBox.cpp
--- a/Engine/Source/Runtime/Slate/Private/Widgets/Layout/SScrollBox.cpp
+++ b/Engine/Source/Runtime/Slate/Private/Widgets/Layout/SScrollBox.cpp
@@ -107,6 +107,7 @@ void SScrollBox::Construct( const FArguments& InArgs )
 	DesiredScrollOffset = 0;
 	bIsScrolling = false;
 	bAnimateScroll = false;
+	ScrollingAnimationInterpolationSpeed = InArgs._ScrollAnimationInterpSpeed;
 	AmountScrolledWhileRightMouseDown = 0;
 	PendingScrollTriggerAmount = 0;
 	bShowSoftwareCursor = false;
@@ -646,7 +647,7 @@ void SScrollBox::Tick( const FGeometry& AllottedGeometry, const double InCurrent
 	const float ViewFraction = GetViewFraction();
 	const float TargetViewOffset = GetViewOffsetFraction();
 	
-	const float CurrentViewOffset = bAnimateScroll ? FMath::FInterpTo(ScrollBar->DistanceFromTop(), TargetViewOffset, InDeltaTime, 15.f) : TargetViewOffset;
+	const float CurrentViewOffset = bAnimateScroll ? FMath::FInterpTo(ScrollBar->DistanceFromTop(), TargetViewOffset, InDeltaTime, ScrollingAnimationInterpolationSpeed) : TargetViewOffset;
 
 	// Update the scrollbar with the clamped version of the offset
 	float NewPhysicalOffset = GetScrollComponentFromVector(CurrentViewOffset * ScrollPanel->GetDesiredSize());
@@ -1179,6 +1180,11 @@ void SScrollBox::SetAnimateWheelScrolling(bool bInAnimateWheelScrolling)
 	bAnimateWheelScrolling = bInAnimateWheelScrolling;
 }
 
+void SScrollBox::SetScrollingAnimationInterpolationSpeed(float NewScrollingAnimationInterpolationSpeed)
+{
+	ScrollingAnimationInterpolationSpeed = NewScrollingAnimationInterpolationSpeed;
+}
+
 void SScrollBox::SetWheelScrollMultiplier(float NewWheelScrollMultiplier)
 {
 	WheelScrollMultiplier = NewWheelScrollMultiplier;
diff --git a/Engine/Source/Runtime/Slate/Private/Widgets/Views/STableViewBase.cpp b/Engine/Source/Runtime/Slate/Private/Widgets/Views/STableViewBase.cpp
--- a/Engine/Source/Runtime/Slate/Private/Widgets/Views/STableViewBase.cpp
+++ b/Engine/Source/Runtime/Slate/Private/Widgets/Views/STableViewBase.cpp
@@ -299,7 +299,7 @@ void STableViewBase::Tick( const FGeometry& AllottedGeometry, const double InCur
 
 			if((bStartedTouchInteraction && bEnableTouchAnimatedScrolling) || (!bStartedTouchInteraction && bEnableAnimatedScrolling))
 			{
-				CurrentScrollOffset = FMath::FInterpTo(CurrentScrollOffset, TargetScrollOffset, (double)InDeltaTime, 12.0);
+				CurrentScrollOffset = FMath::FInterpTo(CurrentScrollOffset, TargetScrollOffset, (double)InDeltaTime, ScrollingAnimationInterpolationSpeed);
 				if (FMath::IsNearlyEqual(CurrentScrollOffset, TargetScrollOffset, 0.01))
 				{
 					CurrentScrollOffset = TargetScrollOffset;
@@ -912,6 +912,11 @@ void STableViewBase::SetIsScrollAnimationEnabled(bool bInEnableScrollAnimation)
 	bEnableAnimatedScrolling = bInEnableScrollAnimation;
 }
 
+void STableViewBase::SetScrollingAnimationInterpolationSpeed(float InScrollingAnimationInterpolationSpeed)
+{
+	ScrollingAnimationInterpolationSpeed = InScrollingAnimationInterpolationSpeed;
+}
+
 void STableViewBase::SetEnableTouchAnimatedScrolling(bool bInEnableTouchAnimatedScrolling)
 {
 	bEnableTouchAnimatedScrolling = bInEnableTouchAnimatedScrolling;
diff --git a/Engine/Source/Runtime/Slate/Public/Widgets/Layout/SScrollBox.h b/Engine/Source/Runtime/Slate/Public/Widgets/Layout/SScrollBox.h
--- a/Engine/Source/Runtime/Slate/Public/Widgets/Layout/SScrollBox.h
+++ b/Engine/Source/Runtime/Slate/Public/Widgets/Layout/SScrollBox.h
@@ -246,6 +246,7 @@ public:
 		, _BackPadScrolling(false)
 		, _FrontPadScrolling(false)
 		, _AnimateWheelScrolling(false)
+		, _ScrollAnimationInterpSpeed(15.f)
 		, _WheelScrollMultiplier(1.f)
 		, _NavigationDestination(EDescendantScrollDestination::IntoView)
 		, _NavigationScrollPadding(0.0f)
@@ -290,6 +291,8 @@ public:
 
 		SLATE_ARGUMENT(bool, AnimateWheelScrolling);
 
+		SLATE_ARGUMENT(float, ScrollAnimationInterpSpeed);
+
 		SLATE_ARGUMENT(float, WheelScrollMultiplier);
 
 		SLATE_ARGUMENT(EDescendantScrollDestination, NavigationDestination);
@@ -339,6 +342,8 @@ public:
 
 	SLATE_API void SetAnimateWheelScrolling(bool bInAnimateWheelScrolling);
 
+	SLATE_API void SetScrollingAnimationInterpolationSpeed(float NewScrollingAnimationInterpolationSpeed);
+
 	SLATE_API void SetWheelScrollMultiplier(float NewWheelScrollMultiplier);
 	
 	SLATE_API void SetScrollWhenFocusChanges(EScrollWhenFocusChanges NewScrollWhenFocusChanges);
@@ -580,6 +585,9 @@ protected:
 	/** Multiplier applied to each click of the scroll wheel (applied alongside the global scroll amount) */
 	float WheelScrollMultiplier = 1.f;
 
+	/** The speed of interpolation for the scrolling animation */
+	float ScrollingAnimationInterpolationSpeed = 15.f;
+
 	/** Whether to animate wheel scrolling */
 	bool bAnimateWheelScrolling : 1;
 
diff --git a/Engine/Source/Runtime/Slate/Public/Widgets/Views/STableViewBase.h b/Engine/Source/Runtime/Slate/Public/Widgets/Views/STableViewBase.h
--- a/Engine/Source/Runtime/Slate/Public/Widgets/Views/STableViewBase.h
+++ b/Engine/Source/Runtime/Slate/Public/Widgets/Views/STableViewBase.h
@@ -170,6 +170,9 @@ public:
 	/** Sets whether the list should lerp between scroll offsets or jump instantly between them. */
 	SLATE_API void SetIsScrollAnimationEnabled(bool bInEnableScrollAnimation);
 
+	/** Sets the speed of lerp between scroll offsets. */
+	SLATE_API void SetScrollingAnimationInterpolationSpeed(float InScrollingAnimationInterpolationSpeed);
+
 	/** Sets whether the list should lerp between scroll offsets or jump instantly between them with touch. */
 	SLATE_API void SetEnableTouchAnimatedScrolling(bool bInEnableTouchAnimatedScrolling);
 
@@ -421,6 +424,9 @@ protected:
 	/** True to lerp smoothly between offsets when the desired scroll offset changes. */
 	bool bEnableAnimatedScrolling = false;
 
+	/** The speed of interpolation for the scrolling animation */
+	float ScrollingAnimationInterpolationSpeed = 12.f;
+
 	/** True to lerp smoothly between offsets when the desired scroll offset changes with touch. */
 	bool bEnableTouchAnimatedScrolling = false;
 
diff --git a/Engine/Source/Runtime/UMG/Private/Components/ListViewBase.cpp b/Engine/Source/Runtime/UMG/Private/Components/ListViewBase.cpp
--- a/Engine/Source/Runtime/UMG/Private/Components/ListViewBase.cpp
+++ b/Engine/Source/Runtime/UMG/Private/Components/ListViewBase.cpp
@@ -154,6 +154,7 @@ TSharedRef<SWidget> UListViewBase::RebuildWidget()
 
 	MyTableViewBase = RebuildListWidget();
 	MyTableViewBase->SetIsScrollAnimationEnabled(bEnableScrollAnimation);
+	MyTableViewBase->SetScrollingAnimationInterpolationSpeed(ScrollingAnimationInterpolationSpeed);
 	MyTableViewBase->SetEnableTouchAnimatedScrolling(bInEnableTouchAnimatedScrolling);
 	MyTableViewBase->SetIsRightClickScrollingEnabled(bEnableRightClickScrolling);
 	MyTableViewBase->SetIsTouchScrollingEnabled(bEnableTouchScrolling);
@@ -180,6 +181,7 @@ void UListViewBase::SynchronizeProperties()
 	if (MyTableViewBase)
 	{
 		MyTableViewBase->SetIsScrollAnimationEnabled(bEnableScrollAnimation);
+		MyTableViewBase->SetScrollingAnimationInterpolationSpeed(ScrollingAnimationInterpolationSpeed);
 		MyTableViewBase->SetEnableTouchAnimatedScrolling(bInEnableTouchAnimatedScrolling);
 		MyTableViewBase->SetIsRightClickScrollingEnabled(bEnableRightClickScrolling);
 		MyTableViewBase->SetIsTouchScrollingEnabled(bEnableTouchScrolling);
diff --git a/Engine/Source/Runtime/UMG/Private/Components/ScrollBox.cpp b/Engine/Source/Runtime/UMG/Private/Components/ScrollBox.cpp
--- a/Engine/Source/Runtime/UMG/Private/Components/ScrollBox.cpp
+++ b/Engine/Source/Runtime/UMG/Private/Components/ScrollBox.cpp
@@ -102,6 +102,7 @@ PRAGMA_DISABLE_DEPRECATION_WARNINGS
 		.BackPadScrolling(BackPadScrolling)
 		.FrontPadScrolling(FrontPadScrolling)
 		.AnimateWheelScrolling(bAnimateWheelScrolling)
+		.ScrollAnimationInterpSpeed(ScrollAnimationInterpolationSpeed)
 		.WheelScrollMultiplier(WheelScrollMultiplier)
 		.OnUserScrolled(BIND_UOBJECT_DELEGATE(FOnUserScrolled, SlateHandleUserScrolled))
 		.OnScrollBarVisibilityChanged(BIND_UOBJECT_DELEGATE(FOnScrollBarVisibilityChanged, SlateHandleScrollBarVisibilityChanged));
@@ -139,6 +140,7 @@ PRAGMA_DISABLE_DEPRECATION_WARNINGS
 	MyScrollBox->SetScrollBarRightClickDragAllowed(bAllowRightClickDragScrolling);
 	MyScrollBox->SetConsumeMouseWheel(ConsumeMouseWheel);
 	MyScrollBox->SetAnimateWheelScrolling(bAnimateWheelScrolling);
+	MyScrollBox->SetScrollingAnimationInterpolationSpeed(ScrollAnimationInterpolationSpeed);
 	MyScrollBox->SetWheelScrollMultiplier(WheelScrollMultiplier);
 	MyScrollBox->SetStyle(&WidgetStyle);
 	MyScrollBox->InvalidateStyle();
@@ -429,6 +431,20 @@ bool UScrollBox::IsAnimateWheelScrolling() const
 	return bAnimateWheelScrolling;
 }
 
+void UScrollBox::SetScrollAnimationInterpolationSpeed(float NewScrollAnimationInterpolationSpeed)
+{
+	ScrollAnimationInterpolationSpeed = NewScrollAnimationInterpolationSpeed;
+	if (MyScrollBox)
+	{
+		MyScrollBox->SetScrollingAnimationInterpolationSpeed(ScrollAnimationInterpolationSpeed);
+	}
+}
+
+float UScrollBox::GetScrollAnimationInterpolationSpeed() const
+{
+	return ScrollAnimationInterpolationSpeed;
+}
+
 void UScrollBox::SetWheelScrollMultiplier(float NewWheelScrollMultiplier)
 {
 	WheelScrollMultiplier = NewWheelScrollMultiplier;
diff --git a/Engine/Source/Runtime/UMG/Public/Components/ListViewBase.h b/Engine/Source/Runtime/UMG/Public/Components/ListViewBase.h
--- a/Engine/Source/Runtime/UMG/Public/Components/ListViewBase.h
+++ b/Engine/Source/Runtime/UMG/Public/Components/ListViewBase.h
@@ -682,6 +682,10 @@ protected:
 	UPROPERTY(EditAnywhere, BlueprintReadOnly, Category = Scrolling)
 	bool bEnableScrollAnimation = false;
 	
+	/** The speed to apply when lerping in the scroll animation. */
+	UPROPERTY(EditAnywhere, BlueprintReadOnly, Category = Scrolling)
+	float ScrollingAnimationInterpolationSpeed = 12.f;
+
 	/** True to enable lerped animation when scrolling through the list with touch*/
 	UPROPERTY(EditAnywhere, BlueprintReadOnly, Category = Scrolling)
 	bool bInEnableTouchAnimatedScrolling = false;
diff --git a/Engine/Source/Runtime/UMG/Public/Components/ScrollBox.h b/Engine/Source/Runtime/UMG/Public/Components/ScrollBox.h
--- a/Engine/Source/Runtime/UMG/Public/Components/ScrollBox.h
+++ b/Engine/Source/Runtime/UMG/Public/Components/ScrollBox.h
@@ -22,6 +22,11 @@ class UScrollBox : public UPanelWidget
 {
 	GENERATED_UCLASS_BODY()
 
+private:
+
+	UPROPERTY(EditAnywhere, BlueprintReadWrite, Getter, Setter, Category = "Scroll", meta = (AllowPrivateAccess = "true", DisplayAfter = "bAnimateWheelScrolling"))
+	float ScrollAnimationInterpolationSpeed = 15.f;
+
 public:
 
 	UE_DEPRECATED(5.2, "Direct access to WidgetStyle is deprecated. Please use the getter or setter.")
@@ -165,6 +170,11 @@ public:
 
 	UMG_API bool IsAnimateWheelScrolling() const;
 
+	UFUNCTION(BlueprintCallable, Category = "Scroll")
+	UMG_API void SetScrollAnimationInterpolationSpeed(float NewScrollAnimationInterpolationSpeed);
+
+	UMG_API float GetScrollAnimationInterpolationSpeed() const;
+
 	UFUNCTION(BlueprintCallable, Category = "Scroll")
 	UMG_API void SetWheelScrollMultiplier(float NewWheelScrollMultiplier);
 
