Fix to TakeWidget when the Widget is not part of an asset

diff --git a/Engine/Source/Runtime/UMG/Private/Components/Widget.cpp b/Engine/Source/Runtime/UMG/Private/Components/Widget.cpp
--- a/Engine/Source/Runtime/UMG/Private/Components/Widget.cpp
+++ b/Engine/Source/Runtime/UMG/Private/Components/Widget.cpp
@@ -955,20 +955,27 @@ TSharedRef<SWidget> UWidget::TakeWidget()
 {
 	LLM_SCOPE_BYTAG(UI_UMG);
 
-	return TakeWidget_Private( []( UUserWidget* Widget, TSharedRef<SWidget> Content ) -> TSharedPtr<SObjectWidget> {
-		       return SNew( SObjectWidget, Widget )[ Content ];
-		   } );
-}
-
-TSharedRef<SWidget> UWidget::TakeWidget_Private(ConstructMethodType ConstructMethod)
-{
 #if WIDGET_INCLUDE_RELFECTION_METADATA
 	UObject* SourceAsset = GetSourceAssetOrClass();
 	UClass* WidgetClass = GetClass();
-	LLM_SCOPE_DYNAMIC_STAT_OBJECTPATH(SourceAsset->GetPackage(), ELLMTagSet::Assets);
-	LLM_SCOPE_DYNAMIC_STAT_OBJECTPATH(WidgetClass, ELLMTagSet::AssetClasses);
-	UE_TRACE_METADATA_SCOPE_ASSET(SourceAsset, WidgetClass);
+	if(SourceAsset && WidgetClass)
+	{
+		LLM_SCOPE_DYNAMIC_STAT_OBJECTPATH(SourceAsset->GetPackage(), ELLMTagSet::Assets);
+		LLM_SCOPE_DYNAMIC_STAT_OBJECTPATH(WidgetClass, ELLMTagSet::AssetClasses);
+		UE_TRACE_METADATA_SCOPE_ASSET(SourceAsset, WidgetClass);
+
+		return TakeWidget_Private([](UUserWidget* Widget, TSharedRef<SWidget> Content) -> TSharedPtr<SObjectWidget> {
+			return SNew(SObjectWidget, Widget)[Content];
+			});
+	}
 #endif
+	return TakeWidget_Private([](UUserWidget* Widget, TSharedRef<SWidget> Content) -> TSharedPtr<SObjectWidget> {
+		return SNew(SObjectWidget, Widget)[Content];
+		});
+}
+
+TSharedRef<SWidget> UWidget::TakeWidget_Private(ConstructMethodType ConstructMethod)
+{
 	bool bNewlyCreated = false;
 	TSharedPtr<SWidget> PublicWidget;
 
@@ -1036,6 +1043,8 @@ TSharedRef<SWidget> UWidget::TakeWidget_Private(ConstructMethodType ConstructMet
 #endif
 
 #if WIDGET_INCLUDE_RELFECTION_METADATA
+		UObject* SourceAsset = GetSourceAssetOrClass();
+		UClass* WidgetClass = GetClass();
 		// We only need to do this once, when the slate widget is created.
 		PublicWidget->AddMetadata<FReflectionMetaData>(MakeShared<FReflectionMetaData>(GetFName(), WidgetClass, this, SourceAsset));
 #endif
