Add a ensure to detect a null extension and prevent a crash

diff --git a/Engine/Source/Runtime/UMG/Private/UserWidget.cpp b/Engine/Source/Runtime/UMG/Private/UserWidget.cpp
--- a/Engine/Source/Runtime/UMG/Private/UserWidget.cpp
+++ b/Engine/Source/Runtime/UMG/Private/UserWidget.cpp
@@ -2509,9 +2509,12 @@ UUserWidgetExtension* UUserWidget::GetExtension(TSubclassOf<UUserWidgetExtension
 {
 	for (UUserWidgetExtension* Extension : Extensions)
 	{
-		if (Extension->IsA(InExtensionType))
+		if (ensure(Extension))
 		{
-			return Extension;
+			if (Extension->IsA(InExtensionType))
+			{
+				return Extension;
+			}
 		}
 	}
 	return nullptr;
