Adding EndInertialScrolling function to ListViewBase

diff --git a/Engine/Source/Runtime/UMG/Private/Components/ListViewBase.cpp b/Engine/Source/Runtime/UMG/Private/Components/ListViewBase.cpp
--- a/Engine/Source/Runtime/UMG/Private/Components/ListViewBase.cpp
+++ b/Engine/Source/Runtime/UMG/Private/Components/ListViewBase.cpp
@@ -73,6 +73,14 @@ void UListViewBase::SetScrollOffset(const float InScrollOffset)
 	}
 }
 
+void UListViewBase::EndInertialScrolling() 
+{
+	if (MyTableViewBase.IsValid())
+	{
+		MyTableViewBase->EndInertialScrolling();
+	}
+}
+
 void UListViewBase::SetWheelScrollMultiplier(float NewWheelScrollMultiplier)
 {
 	WheelScrollMultiplier = NewWheelScrollMultiplier;
diff --git a/Engine/Source/Runtime/UMG/Public/Components/ListViewBase.h b/Engine/Source/Runtime/UMG/Public/Components/ListViewBase.h
--- a/Engine/Source/Runtime/UMG/Public/Components/ListViewBase.h
+++ b/Engine/Source/Runtime/UMG/Public/Components/ListViewBase.h
@@ -538,6 +538,10 @@ public:
 	UFUNCTION(BlueprintCallable, Category = ListView)
 	UMG_API void SetScrollOffset(const float InScrollOffset);
 
+	/** Stops the scroll inertia */
+	UFUNCTION(BlueprintCallable, Category = ListView)
+	UMG_API void EndInertialScrolling();
+
 	UFUNCTION(BlueprintCallable, Category = ListViewBase)
 	UMG_API void SetWheelScrollMultiplier(float NewWheelScrollMultiplier);
 
