Experimental UI Compoment System feature for UMG

This is the first submit for very basic Component implementation. It's hidden behind a CVar and there is a few features missing. Using this we can Create Components that add functionnality to any UWidget in UMG. The first Components will come in future submits.

diff --git a/Engine/Source/Editor/UMGEditor/Private/Customizations/UIComponentCustomizationExtender.cpp b/Engine/Source/Editor/UMGEditor/Private/Customizations/UIComponentCustomizationExtender.cpp
new file mode 100644
--- /dev/null
+++ b/Engine/Source/Editor/UMGEditor/Private/Customizations/UIComponentCustomizationExtender.cpp
@@ -0,0 +1,50 @@
+// Copyright Epic Games, Inc. All Rights Reserved.
+
+#include "UIComponentCustomizationExtender.h"
+
+#include "DetailCategoryBuilder.h"
+#include "DetailLayoutBuilder.h"
+#include "Extensions/UIComponent.h"
+#include "Extensions/UIComponentContainer.h"
+#include "PropertyCustomizationHelpers.h"
+#include "PropertyHandle.h"
+#include "WidgetBlueprint.h"
+#include "WidgetBlueprintEditor.h"
+#include "UIComponentUtils.h"
+
+#define LOCTEXT_NAMESPACE "UIComponentCustomizationExtender"
+
+
+TSharedPtr<FUIComponentCustomizationExtender> FUIComponentCustomizationExtender::MakeInstance()
+{
+	return MakeShared<FUIComponentCustomizationExtender>();
+}
+
+void FUIComponentCustomizationExtender::CustomizeDetails(IDetailLayoutBuilder& InDetailLayout, const TArrayView<UWidget*> InWidgets, const TSharedRef<FWidgetBlueprintEditor>& InWidgetBlueprintEditor)
+{
+	if (InWidgets.Num() == 1)
+	{
+		Widget = InWidgets[0];
+		WidgetBlueprintEditor = InWidgetBlueprintEditor;
+
+		if (UUIComponentContainer* UIComponentContainer = FUIComponentUtils::GetUIComponentContainerFromWidgetBlueprint(WidgetBlueprintEditor.Pin()->GetWidgetBlueprintObj()))
+		{
+			TArray<UUIComponent*> ComponentsOnWidget = UIComponentContainer->GetExtensionsFor(Widget->GetFName());
+
+			for (int32 Index = ComponentsOnWidget.Num() - 1; Index >= 0; Index--)
+			{
+				if (UUIComponent* Component = ComponentsOnWidget[Index])
+				{
+					IDetailCategoryBuilder& ComponentCategory = InDetailLayout.EditCategory(Component->GetFName(), FText::GetEmpty(), ECategoryPriority::Important);
+
+					IDetailPropertyRow* PropertyRow = ComponentCategory.AddExternalObjects({ Component }, EPropertyLocation::Default,
+						FAddPropertyParams()
+						.CreateCategoryNodes(false)
+						.AllowChildren(true)
+						.HideRootObjectNode(true));
+				}
+			}
+		}
+	}
+}
+#undef LOCTEXT_NAMESPACE
\ No newline at end of file
diff --git a/Engine/Source/Editor/UMGEditor/Private/Customizations/UIComponentCustomizationExtender.h b/Engine/Source/Editor/UMGEditor/Private/Customizations/UIComponentCustomizationExtender.h
new file mode 100644
--- /dev/null
+++ b/Engine/Source/Editor/UMGEditor/Private/Customizations/UIComponentCustomizationExtender.h
@@ -0,0 +1,22 @@
+// Copyright Epic Games, Inc. All Rights Reserved.
+
+#pragma once
+
+#include "Customizations/IBlueprintWidgetCustomizationExtender.h"
+#include "Blueprint/UserWidget.h"
+#include "UObject/WeakObjectPtr.h"
+
+class SWidget;
+
+class FUIComponentCustomizationExtender : public IBlueprintWidgetCustomizationExtender
+{
+public:
+	static TSharedPtr<FUIComponentCustomizationExtender> MakeInstance();
+	virtual void CustomizeDetails(IDetailLayoutBuilder& InDetailLayout, const TArrayView<UWidget*> InWidgets, const TSharedRef<FWidgetBlueprintEditor>& InWidgetBlueprintEditor);
+
+private:
+	/** The selected widget in the details panel. */
+	TWeakObjectPtr<class UWidget> Widget;
+	TWeakPtr<FWidgetBlueprintEditor> WidgetBlueprintEditor;
+
+};
diff --git a/Engine/Source/Editor/UMGEditor/Private/Details/SWidgetDetailsView.cpp b/Engine/Source/Editor/UMGEditor/Private/Details/SWidgetDetailsView.cpp
--- a/Engine/Source/Editor/UMGEditor/Private/Details/SWidgetDetailsView.cpp
+++ b/Engine/Source/Editor/UMGEditor/Private/Details/SWidgetDetailsView.cpp
@@ -34,8 +34,21 @@
 #include "UMGEditorModule.h"
 #include "WidgetEditingProjectSettings.h"
 
+#include "Extensions/UIComponent.h"
+#include "Extensions/UIComponentContainer.h"
+#include "Kismet2/SClassPickerDialog.h"
+#include "SPositiveActionButton.h"
+#include "UIComponentUtils.h"
+
 #define LOCTEXT_NAMESPACE "UMG"
 
+static TAutoConsoleVariable<bool> CVarUMGExposeUIComponent(
+	TEXT("Widget.ExposeUIComponents"),
+	false,
+	TEXT("Allow adding ui components to widgets in UMG."),
+	ECVF_Default
+);
+
 void SWidgetDetailsView::Construct(const FArguments& InArgs, TSharedPtr<FWidgetBlueprintEditor> InBlueprintEditor)
 {
 	BlueprintEditor = InBlueprintEditor;
@@ -175,6 +188,23 @@ void SWidgetDetailsView::Construct(const FArguments& InArgs, TSharedPtr<FWidgetB
 						SAssignNew(ClassLinkArea, SBox)
 					]
 				]
+				// Only if not the user widget
+				+ SVerticalBox::Slot()
+					.AutoHeight()
+					.HAlign(EHorizontalAlignment::HAlign_Left)
+					.Padding(4.0f, 0.0f)
+					[
+						CVarUMGExposeUIComponent->GetBool() ?
+						SNew(SPositiveActionButton)
+						.Icon(FAppStyle::Get().GetBrush("Icons.Plus"))
+						.Text(LOCTEXT("AddNewRetargetOpLabel", "Add New Component"))
+						.ToolTipText(LOCTEXT("AddNewToolTip", "Add a new Component to this widget."))
+						.IsEnabled(this, &SWidgetDetailsView::IsWidgetNameFieldEnabled)
+						.OnClicked(this, &SWidgetDetailsView::OnAddComponentButtonClicked, InBlueprintEditor)
+						.Visibility(this, &SWidgetDetailsView::GetAddComponentAreaVisibility)
+						:
+						SNullWidget::NullWidget
+					]
 			]
 		]
 
@@ -407,6 +437,11 @@ EVisibility SWidgetDetailsView::GetCategoryAreaVisibility() const
 	return IsWidgetCDOSelected() ? EVisibility::Visible : EVisibility::Collapsed;
 }
 
+EVisibility SWidgetDetailsView::GetAddComponentAreaVisibility() const
+{
+	return IsWidgetCDOSelected() ? EVisibility::Collapsed : EVisibility::Visible;
+}
+
 void SWidgetDetailsView::HandleCategoryTextCommitted(const FText& Text, ETextCommit::Type CommitType)
 {
 	if ( SelectedObjects.Num() == 1 && !Text.IsEmptyOrWhitespace() )
@@ -616,6 +651,28 @@ void SWidgetDetailsView::HandleIsVariableChanged(ECheckBoxState CheckState)
 	}
 }
 
+FReply SWidgetDetailsView::OnAddComponentButtonClicked(TSharedPtr<FWidgetBlueprintEditor> InBlueprintEditor)
+{
+	UClass* ChosenClass = nullptr;
+	const FText TitleText = LOCTEXT("AddExtension", "Pick a component to add to the widget.");
+	const bool bPressedOk = SClassPickerDialog::PickClass(TitleText, FUIComponentUtils::CreateClassViewerInitializationOptions(), ChosenClass, UUIComponent::StaticClass());
+	UUserWidget* UserWidget = InBlueprintEditor->GetBlueprintObj()->GeneratedClass->GetDefaultObject<UUserWidget>();
+
+	if (ChosenClass)
+	{
+		if (UUIComponentContainer* ComponentsContainer = FUIComponentUtils::GetOrCreateComponentsContainerForUserWidget(UserWidget))
+		{
+			UUIComponent* Component = FUIComponentUtils::CreateUIComponent(ChosenClass, UserWidget);
+			UWidget* SelectedWidget = Cast<UWidget>(SelectedObjects[0].Get());
+
+			ComponentsContainer->Modify();
+			ComponentsContainer->AddComponent(SelectedWidget->GetFName(), Component);
+			FBlueprintEditorUtils::MarkBlueprintAsStructurallyModified(InBlueprintEditor->GetBlueprintObj());
+		}
+	}
+	return FReply::Handled();
+}
+
 void SWidgetDetailsView::NotifyPreChange(FEditPropertyChain* PropertyAboutToChange)
 {
 	// During auto-key do not migrate values
diff --git a/Engine/Source/Editor/UMGEditor/Private/Details/SWidgetDetailsView.h b/Engine/Source/Editor/UMGEditor/Private/Details/SWidgetDetailsView.h
--- a/Engine/Source/Editor/UMGEditor/Private/Details/SWidgetDetailsView.h
+++ b/Engine/Source/Editor/UMGEditor/Private/Details/SWidgetDetailsView.h
@@ -57,6 +57,8 @@ private:
 	bool IsWidgetCDOSelected() const;
 
 	EVisibility GetNameAreaVisibility() const;
+	
+	EVisibility GetAddComponentAreaVisibility() const;
 
 	const FSlateBrush* GetNameIcon() const;
 
@@ -77,6 +79,8 @@ private:
 	FText GetCategoryText() const;
 	void HandleCategoryTextCommitted(const FText& Text, ETextCommit::Type CommitType);
 
+	FReply OnAddComponentButtonClicked(TSharedPtr<FWidgetBlueprintEditor> InBlueprintEditor);
+
 private:
 	/** The editor that owns this details view */
 	TWeakPtr<class FWidgetBlueprintEditor> BlueprintEditor;
diff --git a/Engine/Source/Editor/UMGEditor/Private/Extensions/UIComponentContainerDesignerExtension.cpp b/Engine/Source/Editor/UMGEditor/Private/Extensions/UIComponentContainerDesignerExtension.cpp
new file mode 100644
--- /dev/null
+++ b/Engine/Source/Editor/UMGEditor/Private/Extensions/UIComponentContainerDesignerExtension.cpp
@@ -0,0 +1,20 @@
+// Copyright Epic Games, Inc. All Rights Reserved.
+
+#include "Extensions/UIComponentContainerDesignerExtension.h"
+#include "Blueprint/UserWidget.h"
+#include "Extensions/UIComponentContainer.h"
+#include "WidgetBlueprint.h"
+#include "UIComponentUtils.h"
+
+TSharedRef<FDesignerExtension> FUIComponentContainerDesignerExtensionFactory::CreateDesignerExtension() const
+{
+	return StaticCastSharedRef<FUIComponentContainerDesignerExtension>(MakeShared<FUIComponentContainerDesignerExtension>());
+}
+
+void FUIComponentContainerDesignerExtension::PreviewContentChanged(TSharedRef<SWidget> NewContent)
+{
+	if (UUIComponentContainer* UIComponentContainer = FUIComponentUtils::GetUIComponentContainerFromWidgetBlueprint(Blueprint.Get()))
+	{
+		UIComponentContainer->OnPreviewContentChanged(NewContent);
+	}
+}
diff --git a/Engine/Source/Editor/UMGEditor/Private/Extensions/UIComponentContainerDesignerExtension.h b/Engine/Source/Editor/UMGEditor/Private/Extensions/UIComponentContainerDesignerExtension.h
new file mode 100644
--- /dev/null
+++ b/Engine/Source/Editor/UMGEditor/Private/Extensions/UIComponentContainerDesignerExtension.h
@@ -0,0 +1,21 @@
+// Copyright Epic Games, Inc. All Rights Reserved.
+
+#pragma once
+
+#include "DesignerExtension.h"
+#include "IHasDesignerExtensibility.h"
+
+class FUIComponentContainerDesignerExtensionFactory : public IDesignerExtensionFactory, public TSharedFromThis<FUIComponentContainerDesignerExtensionFactory>
+{
+public:
+	FUIComponentContainerDesignerExtensionFactory() = default;
+	virtual ~FUIComponentContainerDesignerExtensionFactory() = default;
+
+	virtual TSharedRef<FDesignerExtension> CreateDesignerExtension() const override;
+};
+
+class FUIComponentContainerDesignerExtension : public FDesignerExtension
+{
+public:
+	virtual void PreviewContentChanged(TSharedRef<SWidget> NewContent) override;
+};
diff --git a/Engine/Source/Editor/UMGEditor/Private/UIComponentUtils.cpp b/Engine/Source/Editor/UMGEditor/Private/UIComponentUtils.cpp
new file mode 100644
--- /dev/null
+++ b/Engine/Source/Editor/UMGEditor/Private/UIComponentUtils.cpp
@@ -0,0 +1,71 @@
+// Copyright Epic Games, Inc. All Rights Reserved.
+
+#include "UIComponentUtils.h"
+
+#include "Blueprint/UserWidget.h"
+#include "ClassViewerModule.h"
+#include "CoreGlobals.h"
+#include "Components/Widget.h"
+#include "Extensions/UIComponent.h"
+#include "Extensions/UIComponentContainer.h"
+#include "WidgetBlueprint.h"
+
+UUIComponent* FUIComponentUtils::CreateUIComponent(TSubclassOf<UUIComponent> ComponentClass, UUserWidget* Outer)
+{
+	ensure(Outer);
+	UObject* NewComponent = NewObject<UObject>(Outer, ComponentClass);
+	NewComponent->SetFlags(RF_Transactional);
+	return NewComponent ? CastChecked<UUIComponent>(NewComponent) : nullptr;
+}
+
+UUIComponentContainer* FUIComponentUtils::GetOrCreateComponentsContainerForUserWidget(UUserWidget* UserWidget)
+{
+	check(UserWidget);
+	if (UUIComponentContainer* ExistingComponentsContainer = UserWidget->GetExtension<UUIComponentContainer>())
+	{
+		return ExistingComponentsContainer;
+	}
+	else if (UUIComponentContainer* AddedComponentContainer = UserWidget->AddExtension<UUIComponentContainer>())
+	{
+		AddedComponentContainer->SetFlags(RF_Transactional);
+		return AddedComponentContainer;
+	}
+	return nullptr;
+}
+
+FClassViewerInitializationOptions FUIComponentUtils::CreateClassViewerInitializationOptions()
+{
+	FClassViewerInitializationOptions Options;
+	Options.Mode = EClassViewerMode::ClassPicker;
+	TSharedPtr<FUIComponentUtils::FUIComponentClassFilter> Filter = MakeShared<FUIComponentUtils::FUIComponentClassFilter>();
+	Options.ClassFilters.Add(Filter.ToSharedRef());
+
+	Filter->DisallowedClassFlags = CLASS_Deprecated | CLASS_NewerVersionExists | CLASS_HideDropDown | CLASS_Abstract;
+	Filter->AllowedChildrenOfClasses.Add(UUIComponent::StaticClass());
+
+	return Options;
+}
+
+UUIComponentContainer* FUIComponentUtils::GetUIComponentContainerFromWidgetBlueprint(UWidgetBlueprint* WidgetBlueprint)
+{
+	if (WidgetBlueprint)
+	{
+		if (UUserWidget* CDO = WidgetBlueprint->GeneratedClass->GetDefaultObject<UUserWidget>())
+		{
+			return CDO->GetExtension<UUIComponentContainer>();
+		}
+	}
+	return nullptr;
+}
+
+bool FUIComponentUtils::FUIComponentClassFilter::IsClassAllowed(const FClassViewerInitializationOptions& InInitOptions, const UClass* InClass, TSharedRef< FClassViewerFilterFuncs > InFilterFuncs)
+{
+	return !InClass->HasAnyClassFlags(DisallowedClassFlags)
+		&& InFilterFuncs->IfInChildOfClassesSet(AllowedChildrenOfClasses, InClass) != EFilterReturn::Failed;
+}
+
+bool FUIComponentUtils::FUIComponentClassFilter::IsUnloadedClassAllowed(const FClassViewerInitializationOptions& InInitOptions, const TSharedRef< const IUnloadedBlueprintData > InUnloadedClassData, TSharedRef< FClassViewerFilterFuncs > InFilterFuncs)
+{
+	return !InUnloadedClassData->HasAnyClassFlags(DisallowedClassFlags)
+		&& InFilterFuncs->IfInChildOfClassesSet(AllowedChildrenOfClasses, InUnloadedClassData) != EFilterReturn::Failed;
+}
\ No newline at end of file
diff --git a/Engine/Source/Editor/UMGEditor/Private/UMGEditorModule.cpp b/Engine/Source/Editor/UMGEditor/Private/UMGEditorModule.cpp
--- a/Engine/Source/Editor/UMGEditor/Private/UMGEditorModule.cpp
+++ b/Engine/Source/Editor/UMGEditor/Private/UMGEditorModule.cpp
@@ -41,12 +41,18 @@
 #include "Customizations/DynamicEntryBoxDetails.h"
 #include "Customizations/IBlueprintWidgetCustomizationExtender.h"
 #include "Customizations/ListViewBaseDetails.h"
+#include "Customizations/UIComponentCustomizationExtender.h"
 #include "WidgetBlueprintThumbnailRenderer.h"
 #include "Customizations/WidgetThumbnailCustomization.h"
 #include "Widgets/SBindWidgetView.h"
 #include "MovieSceneDynamicBindingUtils.h"
+#include "Kismet2/BlueprintEditorUtils.h"
 #include "Kismet2/KismetEditorUtilities.h"
 
+#include "Extensions/UIComponentContainer.h"
+#include "Extensions/UIComponentContainerDesignerExtension.h"
+#include "UIComponentUtils.h"
+
 #define LOCTEXT_NAMESPACE "UMG"
 
 DEFINE_LOG_CATEGORY_STATIC(LogUMGEditor, Log, All);
@@ -91,16 +97,22 @@ public:
 		DesignerExtensibilityManager = MakeShared<FDesignerExtensibilityManager>();
 
 		DesignerExtensibilityManager->AddDesignerExtensionFactory(SWidgetDesignerNavigation::MakeDesignerExtension());
+		DesignerExtensibilityManager->AddDesignerExtensionFactory(MakeShared<FUIComponentContainerDesignerExtensionFactory>());
 
 		PropertyBindingExtensibilityManager = MakeShared<FPropertyBindingExtensibilityManager>();
 		ClipboardExtensibilityManager = MakeShared<FClipboardExtensibilityManager>();
 		WidgetDragDropExtensibilityManager = MakeShared<FWidgetDragDropExtensibilityManager>();
 
+		UIComponentCustomizationExtender = FUIComponentCustomizationExtender::MakeInstance();
+		AddWidgetCustomizationExtender(UIComponentCustomizationExtender.ToSharedRef());
+
 		// Register widget blueprint compiler we do this no matter what.
 		IKismetCompilerInterface& KismetCompilerModule = FModuleManager::LoadModuleChecked<IKismetCompilerInterface>("KismetCompiler");
 		KismetCompilerModule.GetCompilers().Add(&WidgetBlueprintCompiler);
 		KismetCompilerModule.OverrideBPTypeForClass(UUserWidget::StaticClass(), UWidgetBlueprint::StaticClass());
 
+		FBlueprintEditorUtils::OnRenameVariableReferencesEvent.AddRaw(this, &FUMGEditorModule::HandleRenameVariableReferences);
+
 		// Add Customization for variable in Graph editor
 		if (FBlueprintEditorModule* BlueprintEditorModule = FModuleManager::GetModulePtr<FBlueprintEditorModule>("Kismet"))
 		{
@@ -159,6 +171,7 @@ public:
 
 		MenuExtensibilityManager.Reset();
 		ToolBarExtensibilityManager.Reset();
+		FBlueprintEditorUtils::OnRenameVariableReferencesEvent.RemoveAll(this);
 
 		if (IKismetCompilerInterface* KismetCompilerModule = FModuleManager::GetModulePtr<IKismetCompilerInterface>("KismetCompiler"))
 		{
@@ -213,6 +226,8 @@ public:
 			FEdGraphUtilities::UnregisterVisualPinFactory(GraphPanelPinFactory);
 		}
 
+		RemoveWidgetCustomizationExtender(UIComponentCustomizationExtender.ToSharedRef());
+
 		//// Unregister the setting
 		//ISettingsModule* SettingsModule = FModuleManager::GetModulePtr<ISettingsModule>("Settings");
 
@@ -495,6 +510,13 @@ private:
 		}
 	}
 
+	void HandleRenameVariableReferences(UBlueprint* Blueprint, UClass* VariableClass, const FName& OldVarName, const FName& NewVarName)
+	{
+		if (UUIComponentContainer* ComponentsContainer = FUIComponentUtils::GetUIComponentContainerFromWidgetBlueprint(Cast<UWidgetBlueprint>(Blueprint)))
+		{
+			ComponentsContainer->RenameWidget(OldVarName, NewVarName);
+		}
+	}
 
 private:
 	TSharedPtr<FExtensibilityManager> MenuExtensibilityManager;
@@ -510,6 +532,8 @@ private:
 	FDelegateHandle TransformTrackEditorCreateTrackEditorHandle;
 	FDelegateHandle WidgetMaterialTrackEditorCreateTrackEditorHandle;
 
+	TSharedPtr<FUIComponentCustomizationExtender> UIComponentCustomizationExtender;
+
 	/** All created asset type actions.  Cached here so that we can unregister it during shutdown. */
 	TArray< TSharedPtr<IAssetTypeActions> > CreatedAssetTypeActions;
 
diff --git a/Engine/Source/Editor/UMGEditor/Public/UIComponentUtils.h b/Engine/Source/Editor/UMGEditor/Public/UIComponentUtils.h
new file mode 100644
--- /dev/null
+++ b/Engine/Source/Editor/UMGEditor/Public/UIComponentUtils.h
@@ -0,0 +1,38 @@
+// Copyright Epic Games, Inc. All Rights Reserved.
+
+#pragma once
+
+#include "CoreMinimal.h"
+#include "ClassViewerFilter.h"
+#include "UObject/ObjectMacros.h"
+
+class UUserWidget;
+class UUIComponent;
+class UUIComponentContainer;
+class UWidgetBlueprint;
+/**  */
+
+class UMGEDITOR_API FUIComponentUtils
+{
+public:
+	class FUIComponentClassFilter : public IClassViewerFilter
+	{
+	public:
+		/** All children of these classes will be included unless filtered out by another setting. */
+		TSet <const UClass*> AllowedChildrenOfClasses;
+		TSet <const UClass*> ExcludedChildrenOfClasses;
+
+		/** Disallowed class flags. */
+		EClassFlags DisallowedClassFlags;
+
+		virtual bool IsClassAllowed(const FClassViewerInitializationOptions& InInitOptions, const UClass* InClass, TSharedRef< FClassViewerFilterFuncs > InFilterFuncs) override;
+
+		virtual bool IsUnloadedClassAllowed(const FClassViewerInitializationOptions& InInitOptions, const TSharedRef< const IUnloadedBlueprintData > InUnloadedClassData, TSharedRef< FClassViewerFilterFuncs > InFilterFuncs) override;
+	};
+
+	static UUIComponent* CreateUIComponent(TSubclassOf<UUIComponent> ComponentClass, UUserWidget* Outer);
+	static UUIComponentContainer* GetOrCreateComponentsContainerForUserWidget(UUserWidget* UserWidget);
+	static FClassViewerInitializationOptions CreateClassViewerInitializationOptions();
+
+	static UUIComponentContainer* GetUIComponentContainerFromWidgetBlueprint(UWidgetBlueprint* WidgetBlueprint);
+};
\ No newline at end of file
diff --git a/Engine/Source/Runtime/UMG/Private/Extensions/UIComponent.cpp b/Engine/Source/Runtime/UMG/Private/Extensions/UIComponent.cpp
new file mode 100644
--- /dev/null
+++ b/Engine/Source/Runtime/UMG/Private/Extensions/UIComponent.cpp
@@ -0,0 +1,50 @@
+// Copyright Epic Games, Inc. All Rights Reserved.
+
+#include "Extensions/UIComponent.h"
+
+#include "CoreGlobals.h"
+#include "Components/Widget.h"
+#include "Blueprint/UserWidget.h"
+
+bool UUIComponent::Initialize(UWidget* Target)
+{
+	Owner = Target;
+	const bool Result = OnInitialize();
+
+	if (!Result)
+	{
+		Owner = nullptr;
+	}
+
+	return Result;
+}
+
+void UUIComponent::Construct()
+{
+	OnConstruct();
+}
+
+void UUIComponent::Destruct()
+{
+	OnDestruct();
+}
+
+TWeakObjectPtr<UWidget> UUIComponent::GetOwner() const
+{
+	return Owner;
+}
+
+bool UUIComponent::OnInitialize()
+{
+	return true;
+}
+
+void UUIComponent::OnConstruct()
+{
+
+}
+
+void UUIComponent::OnDestruct()
+{
+
+}
diff --git a/Engine/Source/Runtime/UMG/Private/Extensions/UIComponentContainer.cpp b/Engine/Source/Runtime/UMG/Private/Extensions/UIComponentContainer.cpp
new file mode 100644
--- /dev/null
+++ b/Engine/Source/Runtime/UMG/Private/Extensions/UIComponentContainer.cpp
@@ -0,0 +1,247 @@
+// Copyright Epic Games, Inc. All Rights Reserved.
+
+#include "Extensions/UIComponentContainer.h"
+
+#include "Blueprint/WidgetTree.h"
+#include "CoreMinimal.h"
+#include "Components/Widget.h"
+#include "Extensions/UIComponent.h"
+#include "Slate/SObjectWidget.h"
+
+FUIComponentTarget::FUIComponentTarget()
+{}
+
+FUIComponentTarget::FUIComponentTarget(UUIComponent* InComponent, FName InChildName)
+{
+	TargetName = InChildName;
+	Component = InComponent;
+}
+
+UWidget* FUIComponentTarget::Resolve(const UWidgetTree* WidgetTree)
+{
+	if (!TargetName.IsNone() && WidgetTree)
+	{
+		return WidgetTree->FindWidget(TargetName);
+	}
+
+	return nullptr;
+}
+
+void FUIComponentTarget::SetTargetName(FName NewName)
+{
+	TargetName = NewName;
+}
+
+void UUIComponentContainer::Initialize()
+{
+	if (UUserWidget* Owner = GetUserWidget())
+	{
+		// Before intializing, we resolve all components
+		for (FUIComponentTarget& ComponentTarget : Components)
+		{
+			if (UWidget* Widget = ComponentTarget.Resolve(Owner->WidgetTree))
+			{
+				if (UUIComponent* Component = ComponentTarget.GetComponent())
+				{
+					Component->Initialize(Widget);
+				}
+			}
+		}
+	}
+}
+
+void UUIComponentContainer::Construct()
+{
+	for (FUIComponentTarget& ComponentTarget : Components)
+	{
+		if(UUIComponent* Component = ComponentTarget.GetComponent())
+		{
+			if (TStrongObjectPtr<UWidget> Owner = Component->GetOwner().Pin())
+			{
+				Component->Construct();
+			}
+		}		
+	}
+}
+
+void UUIComponentContainer::Destruct()
+{
+	for (FUIComponentTarget& ComponentTarget : Components)
+	{
+		if (UUIComponent* Component = ComponentTarget.GetComponent())
+		{
+			if (TStrongObjectPtr<UWidget> Owner = Component->GetOwner().Pin())
+			{
+				Component->Destruct();
+			}
+		}
+	}
+}
+
+
+TArray<UUIComponent*> UUIComponentContainer::GetExtensionsFor(FName TargetName)
+{
+	TArray<UUIComponent*> Result;
+
+	for (FUIComponentTarget& Target : Components)
+	{
+		if (Target.GetTargetName() == TargetName)
+		{
+			Result.Push(Target.GetComponent());
+		}
+	}
+
+	return MoveTemp(Result);
+}
+
+void UUIComponentContainer::AddComponent(FName TargetName, UUIComponent* Component)
+{
+	ensure(!TargetName.IsNone() && Component != nullptr);
+	Components.Emplace(FUIComponentTarget(Component, TargetName));
+}
+
+void UUIComponentContainer::RemoveComponent(FName TargetName, UUIComponent* Component)
+{
+	ensure(!TargetName.IsNone() && Component != nullptr);
+
+	for (int i = 0; i < Components.Num(); i++)
+	{
+		FUIComponentTarget& Target = Components[i];
+		if (Target.GetComponent() == Component)
+		{
+			ensure(Target.GetTargetName() == TargetName);
+			Components.RemoveAtSwap(i);
+			return;
+		}
+	}
+}
+
+void UUIComponentContainer::Resolve()
+{
+	if (UUserWidget* Owner = GetUserWidget())
+	{
+		for (FUIComponentTarget& Component : Components)
+		{
+			Component.Resolve(Owner->WidgetTree);
+		}
+	}
+}
+
+void UUIComponentContainer::RemoveAllComponentsFor(FName TargetName)
+{
+	for (int32 Index = Components.Num() - 1; Index >= 0; --Index)
+	{
+		FUIComponentTarget& Target = Components[Index];
+		if (Target.GetTargetName() == TargetName)
+		{
+			Components.RemoveAtSwap(Index);
+		}
+	}
+}
+
+void UUIComponentContainer::RenameWidget(FName OldName, FName NewName)
+{
+	bool bModified = false;
+	auto UpdateModify = [&bModified, Self = this]()
+		{
+			if (!bModified)
+			{
+				bModified = true;
+				Self->Modify();
+			}
+		};
+
+	for (FUIComponentTarget& Target : Components)
+	{
+		if (Target.GetTargetName() == OldName)
+		{
+			UpdateModify();
+			Target.SetTargetName(NewName);
+
+			// @TODO Do our widget components need to be notified of name change?
+
+			//if (Component.GetComponent())
+			//{
+			//	Component.GetComponent()->WidgetRenamed(OldName, NewName);
+			//}
+		}
+	}
+}
+
+
+bool UUIComponentContainer::IsEmpty() const 
+{
+	return Components.IsEmpty();
+}
+
+void UUIComponentContainer::OnPreviewContentChanged(TSharedRef<SWidget> NewContent)
+{
+	if (TSharedPtr<SObjectWidget> ObjectWidget = StaticCastSharedPtr<SObjectWidget>(NewContent.ToSharedPtr()))
+	{
+		CleanupUIComponents(ObjectWidget->GetWidgetObject());
+	}
+}
+
+void UUIComponentContainer::CleanupUIComponents(UUserWidget* UserWidget)
+{
+	if (UserWidget)
+	{
+		if (const UWidgetTree* WidgetTree = UserWidget->WidgetTree)
+		{
+			bool bModified = false;
+			auto UpdateModify = [&bModified, Self = this]()
+			{
+				if (!bModified)
+				{
+					bModified = true;
+					Self->Modify();
+				}
+			};
+			
+			TArray <FName, TInlineAllocator<4>> TargetNamesToRemove;
+			for (int32 Index = Components.Num() - 1; Index >= 0; --Index)
+			{
+				// Remove components that are invalid
+				if (Components[Index].GetComponent() == nullptr)
+				{
+					UpdateModify();
+					Components.RemoveAtSwap(Index);
+				}
+				else if (!Components[Index].GetTargetName().IsNone())
+				{
+					// Fill list with all Valid Target Names
+					TargetNamesToRemove.AddUnique(Components[Index].GetTargetName());
+				}
+			}
+
+			// Exclude components in use from removal list
+			if (TargetNamesToRemove.Num() > 0)
+			{
+				WidgetTree->ForEachWidget([&TargetNamesToRemove, this](TObjectPtr<UWidget> Widget) {
+					if (Widget)
+					{
+						for (int32 Index = TargetNamesToRemove.Num() - 1; Index >= 0; Index--)
+						{
+							const FName TargetName = TargetNamesToRemove[Index];
+							if (TargetName == Widget->GetFName())
+							{
+								TargetNamesToRemove.RemoveSingleSwap(TargetName);
+							}
+						}
+					}
+					});
+			}
+
+			// Remove all unused component 
+			if (TargetNamesToRemove.Num() > 0)
+			{
+				UpdateModify();
+
+				for (FName TargetName : TargetNamesToRemove)
+				{
+					RemoveAllComponentsFor(TargetName);
+				}
+			}
+		}
+	}
+}
\ No newline at end of file
diff --git a/Engine/Source/Runtime/UMG/Public/Extensions/UIComponent.h b/Engine/Source/Runtime/UMG/Public/Extensions/UIComponent.h
new file mode 100644
--- /dev/null
+++ b/Engine/Source/Runtime/UMG/Public/Extensions/UIComponent.h
@@ -0,0 +1,54 @@
+// Copyright Epic Games, Inc. All Rights Reserved.
+
+#pragma once
+
+#include "CoreMinimal.h"
+#include "UObject/ObjectMacros.h"
+
+#include "UIComponent.generated.h"
+
+class UWidget;
+
+/** 
+ * This is the base class to for UI Components that can be added to any UMG Widgets
+ * in UMG Designer.When initialized, it will pass the widget it's attached to.
+ */
+UCLASS(Abstract, MinimalAPI, Experimental)
+class UUIComponent : public UObject
+{
+	GENERATED_BODY()
+
+public:
+	/**
+	 * Called when the owner widget is inititalized.
+	 */
+	UMG_API bool Initialize(UWidget* target);
+
+	/**
+	 * Called when the owner widget is constructed.
+	 */
+	UMG_API void Construct();
+	
+
+	/**
+	 * Called when the owner widget is destructed.
+	 */
+	UMG_API void Destruct();
+
+	/**
+	 * Returns the Owner Widget this component is attached to.
+	 *  @returns The owner widget
+	*/
+	UMG_API TWeakObjectPtr<UWidget> GetOwner() const;
+
+protected:
+	UMG_API virtual bool OnInitialize();
+	UMG_API virtual void OnConstruct();
+	UMG_API virtual void OnDestruct();
+
+private:
+	UPROPERTY(transient, DuplicateTransient)
+	TWeakObjectPtr<UWidget> Owner;
+
+
+};
diff --git a/Engine/Source/Runtime/UMG/Public/Extensions/UIComponentContainer.h b/Engine/Source/Runtime/UMG/Public/Extensions/UIComponentContainer.h
new file mode 100644
--- /dev/null
+++ b/Engine/Source/Runtime/UMG/Public/Extensions/UIComponentContainer.h
@@ -0,0 +1,76 @@
+// Copyright Epic Games, Inc. All Rights Reserved.
+
+#pragma once
+
+#include "CoreMinimal.h"
+#include "UObject/ObjectMacros.h"
+#include "UObject/Object.h"
+#include "UObject/WeakObjectPtr.h"
+#include "Extensions/UserWidgetExtension.h"
+
+#include "UIComponentContainer.generated.h"
+
+class UUserWidget;
+class UWidget;
+class UUIComponent;
+
+USTRUCT()
+struct FUIComponentTarget
+{
+	GENERATED_BODY();
+public:
+
+	FUIComponentTarget();
+	FUIComponentTarget(UUIComponent* Component, FName InChildName);
+
+	/** Resolves the Widget ptr using it's name. */
+	UWidget* Resolve(const class UWidgetTree* WidgetTree);
+	FName GetTargetName() const { return TargetName; }
+	void SetTargetName(FName NewName);
+	UUIComponent* GetComponent() const { return Component; }
+
+private:
+
+	// We use a TargetName to resolve the Widget only at compile time and on the Runtime Widget. 
+	// It simplify edition in UMG Designer and make sure we do not need to keep Association in sync with the WidgetTree.
+	UPROPERTY()
+	FName TargetName;
+	
+	UPROPERTY(Instanced)
+	TObjectPtr<UUIComponent> Component;
+};
+
+
+/**
+ * Class that holds all the UIComponents for a UUserWidget.
+ */
+UCLASS(NotBlueprintType, Experimental)
+class UMG_API UUIComponentContainer : public UUserWidgetExtension
+{
+	GENERATED_BODY()
+
+public:
+	virtual void Initialize() override;
+	virtual void Construct() override;		
+	virtual void Destruct() override;
+
+	TArray<UUIComponent*> GetExtensionsFor(FName TargetName);
+
+	void AddComponent(FName TargetName, UUIComponent* Component);
+	void RemoveComponent(FName TargetName, UUIComponent* Component);
+	void RemoveAllComponentsFor(FName TargetName);
+
+	bool IsEmpty() const;
+
+	void RenameWidget(FName OldName, FName NewName);
+
+	void OnPreviewContentChanged(TSharedRef<SWidget> NewContent);
+
+private:
+	void Resolve();
+	void CleanupUIComponents(UUserWidget* UserWidget);
+
+	// Use a single TArray for the Entire UUserWidget to reduce memory usage.
+	UPROPERTY()
+	TArray<FUIComponentTarget> Components;
+};
