Workaround for missing widgets in PIE when starting Lyra. We can end up never rendering the widget because the last render time is never set.

diff --git a/Engine/Source/Runtime/UMG/Private/Components/WidgetComponent.cpp b/Engine/Source/Runtime/UMG/Private/Components/WidgetComponent.cpp
--- a/Engine/Source/Runtime/UMG/Private/Components/WidgetComponent.cpp
+++ b/Engine/Source/Runtime/UMG/Private/Components/WidgetComponent.cpp
@@ -1334,10 +1334,10 @@ void UWidgetComponent::AddWidgetToScreen(ULocalPlayer* TargetPlayer)
 bool UWidgetComponent::ShouldDrawWidget() const
 {
 	const float RenderTimeThreshold = .5f;
-	if ( IsVisible() )
+	if (IsVisible())
 	{
 		// If we don't tick when off-screen, don't bother ticking if it hasn't been rendered recently
-		if ( TickWhenOffscreen || GetWorld()->TimeSince(GetLastRenderTime()) <= RenderTimeThreshold )
+		if (TickWhenOffscreen || WasRecentlyRendered(RenderTimeThreshold) || LastWidgetRenderTime == 0.0)
 		{
 			if ( ( GetCurrentTime() - LastWidgetRenderTime) >= RedrawTime )
 			{
