Remove padding from Widget and TextBlock
Save 0.5mb (MB3)

#rnx

diff --git a/Engine/Source/Runtime/UMG/Public/Components/Widget.h b/Engine/Source/Runtime/UMG/Public/Components/Widget.h
--- a/Engine/Source/Runtime/UMG/Public/Components/Widget.h
+++ b/Engine/Source/Runtime/UMG/Public/Components/Widget.h
@@ -266,15 +266,15 @@ public:
 	UPROPERTY()
 	FGetBool bIsEnabledDelegate;
 
+	/** A bindable delegate for ToolTipText */
+	UPROPERTY()
+	FGetText ToolTipTextDelegate;
+
 	UE_DEPRECATED(5.1, "Direct access to ToolTipText is deprecated. Please use the getter or setter.")
 	/** Tooltip text to show when the user hovers over the widget with the mouse */
 	UPROPERTY(EditAnywhere, BlueprintReadWrite, Getter, Setter, BlueprintSetter="SetToolTipText", Category="Behavior", meta=(MultiLine=true))
 	FText ToolTipText;
 
-	/** A bindable delegate for ToolTipText */
-	UPROPERTY()
-	FGetText ToolTipTextDelegate;
-
 	UE_DEPRECATED(5.1, "Direct access to ToolTipWidget is deprecated. Please use the getter or setter.")
 	/** Tooltip widget to show when the user hovers over the widget with the mouse */
 	UPROPERTY(VisibleAnywhere, BlueprintReadWrite, Getter="GetToolTip", Setter="SetToolTip", BlueprintSetter="SetToolTip", Category="Behavior", AdvancedDisplay)
@@ -432,17 +432,19 @@ public:
 	/** The visibility of the widget */
 	UPROPERTY(EditAnywhere, BlueprintReadWrite, FieldNotify, Getter, Setter, BlueprintGetter="GetVisibility", BlueprintSetter="SetVisibility", Category="Behavior")
 	ESlateVisibility Visibility;
+	
+private:
+	/** If the widget will draw snapped to the nearest pixel.  Improves clarity but might cause visibile stepping in animation */
+	UPROPERTY(EditAnywhere, BlueprintReadWrite, Getter, Setter, Category = "Rendering", meta=(AllowPrivateAccess = true))
+	EWidgetPixelSnapping PixelSnapping;
+
+protected:
 
 	UE_DEPRECATED(5.1, "Direct access to RenderOpacity is deprecated. Please use the getter or setter.")
 	/** The opacity of the widget */
 	UPROPERTY(EditAnywhere, BlueprintReadWrite, Getter, Setter, BlueprintGetter="GetRenderOpacity", BlueprintSetter="SetRenderOpacity", Category="Rendering")
 	float RenderOpacity;
 
-private:
-	/** If the widget will draw snapped to the nearest pixel.  Improves clarity but might cause visibile stepping in animation */
-	UPROPERTY(EditAnywhere, BlueprintReadWrite, Getter, Setter, Category = "Rendering", meta=(AllowPrivateAccess = true))
-	EWidgetPixelSnapping PixelSnapping;
-
 private:
 	/** A custom set of accessibility rules for this widget. If null, default rules for the widget are used. */
 	UPROPERTY(Instanced)
