UE-193778 Detect name clashes in derived blueprint widgets when compiling

#tests - Created widget with same name as other widgets, same name as named slots. Validated with preview widgets during compilation.
#rnx

diff --git a/Engine/Source/Editor/UMGEditor/Private/WidgetBlueprint.cpp b/Engine/Source/Editor/UMGEditor/Private/WidgetBlueprint.cpp
--- a/Engine/Source/Editor/UMGEditor/Private/WidgetBlueprint.cpp
+++ b/Engine/Source/Editor/UMGEditor/Private/WidgetBlueprint.cpp
@@ -1242,6 +1242,47 @@ bool HasCircularReferences(const UClass* CurrentClass, TArray<const UClass*, TIn
 
 	return false;
 }
+
+const UWidgetBlueprint* GetGeneratedWidgetBlueprintFromClass(const UClass* CurrentClass)
+{
+	if (const UWidgetBlueprintGeneratedClass* GeneratedClass = Cast<const UWidgetBlueprintGeneratedClass>(CurrentClass))
+	{
+		return Cast<const UWidgetBlueprint>(GeneratedClass->ClassGeneratedBy);
+	}
+	return nullptr;
+}
+
+bool TryBuildSetNameForWidgets(const UClass* CurrentClass, TSet<FName>& InOutWidgetNames, TSet<UWidget*>& InOutConflictingWidgets)
+{
+	if (const UWidgetBlueprint* WidgetBP = GetGeneratedWidgetBlueprintFromClass(CurrentClass))
+	{
+		// only search in parent class if this class isn't overwriting the root
+		if (WidgetBP->WidgetTree->RootWidget == nullptr)
+		{
+			if (!TryBuildSetNameForWidgets(WidgetBP->ParentClass, InOutWidgetNames, InOutConflictingWidgets))
+			{
+				return false;
+			}
+		}
+
+		if (WidgetBP->WidgetTree)
+		{
+			WidgetBP->WidgetTree->ForEachWidget([&InOutWidgetNames,&InOutConflictingWidgets](UWidget* Widget) {
+				FName WidgetName = Widget->GetFName();
+				if (InOutWidgetNames.Contains(WidgetName))
+				{
+					InOutConflictingWidgets.Add(Widget);
+				}
+				else
+				{
+					InOutWidgetNames.Add(Widget->GetFName());
+				}
+				});
+		}
+	}
+	return InOutConflictingWidgets.IsEmpty();
+}
+
 }
 
 TValueOrError<void, UWidget*> UWidgetBlueprint::HasCircularReferences() const
@@ -1258,6 +1299,28 @@ TValueOrError<void, UWidget*> UWidgetBlueprint::HasCircularReferences() const
 	return MakeValue();
 }
 
+TValueOrError<void, TSet<UWidget*>> UWidgetBlueprint::HasConflictingWidgetNamesFromInheritance() const
+{
+	if (GeneratedClass)
+	{
+		// we search for conflicting widget names when the parent is also a generated widgetblueprint
+		// this allows bailing out early on most compilations
+		if (const UWidgetBlueprint* WidgetBP = UE::UMG::Private::GetGeneratedWidgetBlueprintFromClass(GeneratedClass) )
+		{
+			if (Cast<const UWidgetBlueprintGeneratedClass>(WidgetBP->ParentClass) != nullptr)
+			{
+				TSet<FName> WidgetNames;
+				TSet<UWidget*> Result;
+				if (!UE::UMG::Private::TryBuildSetNameForWidgets(GeneratedClass, WidgetNames, Result))
+				{
+					return MakeError(Result);
+				}
+			}
+		}
+	}
+	return MakeValue();
+}
+
 UPackage* UWidgetBlueprint::GetWidgetTemplatePackage() const
 {
 	return GetOutermost();
diff --git a/Engine/Source/Editor/UMGEditor/Private/WidgetBlueprintCompiler.cpp b/Engine/Source/Editor/UMGEditor/Private/WidgetBlueprintCompiler.cpp
--- a/Engine/Source/Editor/UMGEditor/Private/WidgetBlueprintCompiler.cpp
+++ b/Engine/Source/Editor/UMGEditor/Private/WidgetBlueprintCompiler.cpp
@@ -893,6 +893,35 @@ void FWidgetBlueprintCompilerContext::FinishCompilingClass(UClass* Class)
 			}
 		}
 
+		{
+#if WITH_EDITOR
+			BPGClass->NameClashingInHierarchy.Reset();
+#endif
+			TValueOrError<void, TSet<UWidget*>> HasConflictingWidgetNames = WidgetBP->HasConflictingWidgetNamesFromInheritance();
+			if (HasConflictingWidgetNames.HasError())
+			{
+				TSet<UWidget*>& ConflictingNames = HasConflictingWidgetNames.GetError();
+				for (UWidget* ConflictingWidget : ConflictingNames)
+				{
+#if WITH_EDITOR
+					FName ConflictingWidgetName = ConflictingWidget->GetFName();
+					BPGClass->NameClashingInHierarchy.Add(ConflictingWidgetName);
+#endif
+					if (UWidget* FoundConflictingWidget = BPGClass->GetWidgetTreeArchetype()->FindWidget(ConflictingWidget->GetFName()))
+					{
+						BPGClass->GetWidgetTreeArchetype()->RemoveWidget(FoundConflictingWidget);
+					}
+
+					MessageLog.Error(*FText::Format(LOCTEXT("WidgetTreeDuplicateNames", "The WidgetTree '{0}' already contains a widget named '{1}'."),
+						FText::FromString(WidgetBP->WidgetTree->GetPathName()),
+						FText::FromString(ConflictingWidget->GetName())
+					).ToString());
+				}
+			}
+		}
+
+
+
 		int32 AnimIndex = 0;
 		for ( const UWidgetAnimation* Animation : WidgetBP->Animations )
 		{
@@ -1237,7 +1266,6 @@ void FWidgetBlueprintCompilerContext::FinishCompilingClass(UClass* Class)
 		});
 }
 
-
 class FBlueprintCompilerLog : public IWidgetCompilerLog
 {
 public:
diff --git a/Engine/Source/Editor/UMGEditor/Public/WidgetBlueprint.h b/Engine/Source/Editor/UMGEditor/Public/WidgetBlueprint.h
--- a/Engine/Source/Editor/UMGEditor/Public/WidgetBlueprint.h
+++ b/Engine/Source/Editor/UMGEditor/Public/WidgetBlueprint.h
@@ -302,6 +302,9 @@ public:
 	/**  */
 	TValueOrError<void, UWidget*> HasCircularReferences() const;
 
+	TValueOrError<void, TSet<UWidget*>> HasConflictingWidgetNamesFromInheritance() const;
+
+
 	static bool ValidateGeneratedClass(const UClass* InClass);
 	
 	static TSharedPtr<FKismetCompilerContext> GetCompilerForWidgetBP(UBlueprint* BP, FCompilerResultsLog& InMessageLog, const FKismetCompilerOptions& InCompileOptions);
diff --git a/Engine/Source/Runtime/UMG/Private/UserWidget.cpp b/Engine/Source/Runtime/UMG/Private/UserWidget.cpp
--- a/Engine/Source/Runtime/UMG/Private/UserWidget.cpp
+++ b/Engine/Source/Runtime/UMG/Private/UserWidget.cpp
@@ -223,21 +223,38 @@ void UUserWidget::DuplicateAndInitializeFromWidgetTree(UWidgetTree* InWidgetTree
 		TArray<UWidget*> AllNamedSlotContentWidgets;
 		NamedSlotContentToMerge.GenerateValueArray(AllNamedSlotContentWidgets);
 
-		auto SetContentWidgetForNamedSlot = [this](FName NamedSlotName, UWidget* TemplateSlotContent)
+		TSet<FName> const* ConflictingWidgetNames = nullptr;
+#if WITH_EDITOR
+		if (UWidgetBlueprintGeneratedClass* BGClass = Cast<UWidgetBlueprintGeneratedClass>(GetClass()))
+		{
+			ConflictingWidgetNames = &(BGClass->NameClashingInHierarchy);
+		}
+#endif
+
+		auto SetContentWidgetForNamedSlot = [this,&ConflictingWidgetNames](FName NamedSlotName, UWidget* TemplateSlotContent)
 		{
 			FObjectInstancingGraph NamedSlotInstancingGraph;
 			// We need to add a mapping from the template's widget tree to the new widget tree, that way
 			// as we instance the widget hierarchy it's grafted onto the new widget tree.
 			NamedSlotInstancingGraph.AddNewObject(WidgetTree, TemplateSlotContent->GetTypedOuter<UWidgetTree>());
 
-			// Instance the new widget from the foreign tree, but do it in a way that grafts it onto the tree we're instancing.
-			UWidget* Content = NewObject<UWidget>(WidgetTree, TemplateSlotContent->GetClass(), TemplateSlotContent->GetFName(), RF_Transactional, TemplateSlotContent, false, &NamedSlotInstancingGraph);
-			Content->SetFlags(RF_Transient | RF_DuplicateTransient);
-
-			// Insert the newly constructed widget into the named slot that corresponds.  The above creates
-			// it as if it was always part of the widget tree, but this actually puts it into a widget's
-			// slot for the named slot.
-			SetContentForSlot(NamedSlotName, Content);
+			FName TemplateSlotContentName = TemplateSlotContent->GetFName();
+			// ConflictingWidgetNames is an optional parameter. If we find an item with the name we were about to create in the widget tree, we remove the NamedSlot to avoid the corrupted tree we would get otherwise
+			if (ConflictingWidgetNames == nullptr || !ConflictingWidgetNames->Contains(TemplateSlotContentName))
+			{
+				// Instance the new widget from the foreign tree, but do it in a way that grafts it onto the tree we're instancing.
+				UWidget* Content = NewObject<UWidget>(WidgetTree, TemplateSlotContent->GetClass(), TemplateSlotContentName, RF_Transactional, TemplateSlotContent, false, &NamedSlotInstancingGraph);
+				Content->SetFlags(RF_Transient | RF_DuplicateTransient);
+
+				// Insert the newly constructed widget into the named slot that corresponds.  The above creates
+				// it as if it was always part of the widget tree, but this actually puts it into a widget's
+				// slot for the named slot.
+				SetContentForSlot(NamedSlotName, Content);
+			}
+			else
+			{
+				SetContentForSlot(NamedSlotName, nullptr);
+			}
 		};
 
 		// This block controls merging named slot content specified in a child class for the widget we're templated after.
diff --git a/Engine/Source/Runtime/UMG/Public/Blueprint/WidgetBlueprintGeneratedClass.h b/Engine/Source/Runtime/UMG/Public/Blueprint/WidgetBlueprintGeneratedClass.h
--- a/Engine/Source/Runtime/UMG/Public/Blueprint/WidgetBlueprintGeneratedClass.h
+++ b/Engine/Source/Runtime/UMG/Public/Blueprint/WidgetBlueprintGeneratedClass.h
@@ -134,6 +134,9 @@ public:
 
 	UPROPERTY(Transient)
 	TSet<FName> NamedSlotsWithContentInSameTree;
+
+	UPROPERTY(Transient)
+	TSet<FName> NameClashingInHierarchy;
 #endif
 
 	/**
