UE-185233 ensure when overwriting namedslots with new widget at root

#tests Tested the 2 cases with namedslots being defined and then overwritten by a RootElement in the derived class. Tried the ui in some of the projects which seemed to be ok
#rnx

diff --git a/Engine/Source/Runtime/UMG/Private/UserWidget.cpp b/Engine/Source/Runtime/UMG/Private/UserWidget.cpp
--- a/Engine/Source/Runtime/UMG/Private/UserWidget.cpp
+++ b/Engine/Source/Runtime/UMG/Private/UserWidget.cpp
@@ -1200,6 +1200,27 @@ void UUserWidget::SetContentForSlot(FName SlotName, UWidget* Content)
 {
 	bool bFoundExistingSlot = false;
 
+	bool bIsMissingSlot = false;
+	// Dynamically insert the new widget into the hierarchy if it exists.
+	if (WidgetTree)
+	{
+		ensureMsgf(!HasAnyFlags(RF_ClassDefaultObject), TEXT("The Widget CDO is not expected to ever have a valid widget tree."));
+
+		if (UNamedSlot* NamedSlot = Cast<UNamedSlot>(WidgetTree->FindWidget(SlotName)))
+		{
+			NamedSlot->ClearChildren();
+
+			if (Content)
+			{
+				NamedSlot->AddChild(Content);
+			}
+		}
+		else
+		{
+			bIsMissingSlot = true;
+		}
+	}
+
 	// Find the binding in the existing set and replace the content for that binding.
 	for ( int32 BindingIndex = 0; BindingIndex < NamedSlotBindings.Num(); BindingIndex++ )
 	{
@@ -1209,7 +1230,7 @@ void UUserWidget::SetContentForSlot(FName SlotName, UWidget* Content)
 		{
 			bFoundExistingSlot = true;
 
-			if ( Content )
+			if ( Content && !bIsMissingSlot)
 			{
 				Binding.Content = Content;
 			}
@@ -1222,7 +1243,7 @@ void UUserWidget::SetContentForSlot(FName SlotName, UWidget* Content)
 		}
 	}
 
-	if ( !bFoundExistingSlot && Content )
+	if ( !bFoundExistingSlot && Content && !bIsMissingSlot)
 	{
 		// Add the new binding to the list of bindings.
 		FNamedSlotBinding NewBinding;
@@ -1232,21 +1253,6 @@ void UUserWidget::SetContentForSlot(FName SlotName, UWidget* Content)
 		NamedSlotBindings.Add(NewBinding);
 	}
 
-	// Dynamically insert the new widget into the hierarchy if it exists.
-	if ( WidgetTree )
-	{
-		ensureMsgf(!HasAnyFlags(RF_ClassDefaultObject), TEXT("The Widget CDO is not expected to ever have a valid widget tree."));
-		
-		if ( UNamedSlot* NamedSlot = Cast<UNamedSlot>(WidgetTree->FindWidget(SlotName)))
-		{
-			NamedSlot->ClearChildren();
-
-			if ( Content )
-			{
-				NamedSlot->AddChild(Content);
-			}
-		}
-	}
 }
 
 UWidget* UUserWidget::GetRootWidget() const
