Mark WidgetTree as duplicate transient so that when we reinstance widget blueprints we don't create unneeded copies of the widget tree

[RN] Fixed warning about REINST widget tree objects being nullptred on compilation in complicated widget graphs

diff --git a/Engine/Source/Runtime/UMG/Public/Blueprint/WidgetBlueprintGeneratedClass.h b/Engine/Source/Runtime/UMG/Public/Blueprint/WidgetBlueprintGeneratedClass.h
--- a/Engine/Source/Runtime/UMG/Public/Blueprint/WidgetBlueprintGeneratedClass.h
+++ b/Engine/Source/Runtime/UMG/Public/Blueprint/WidgetBlueprintGeneratedClass.h
@@ -88,7 +88,7 @@ public:
 private:
 
 	/** A tree of the widget templates to be created */
-	UPROPERTY()
+	UPROPERTY(DuplicateTransient)
 	TObjectPtr<UWidgetTree> WidgetTree;
 	
 	/** The extension that are considered static to the class */
