UWidgetComponent - Added null checks when creating scene proxy.

diff --git a/Engine/Source/Runtime/UMG/Private/Components/WidgetComponent.cpp b/Engine/Source/Runtime/UMG/Private/Components/WidgetComponent.cpp
--- a/Engine/Source/Runtime/UMG/Private/Components/WidgetComponent.cpp
+++ b/Engine/Source/Runtime/UMG/Private/Components/WidgetComponent.cpp
@@ -327,7 +327,10 @@ public:
 	{
 		bWillEverBeLit = false;
 
-		MaterialRelevance = MaterialInstance->GetRelevance_Concurrent(GetScene().GetFeatureLevel());
+		if (MaterialInstance)
+		{
+			MaterialRelevance = MaterialInstance->GetRelevance_Concurrent(GetScene().GetFeatureLevel());
+		}
 	}
 
 	// FPrimitiveSceneProxy interface.
@@ -791,10 +794,13 @@ FPrimitiveSceneProxy* UWidgetComponent::CreateSceneProxy()
 
 	if (WidgetRenderer && CurrentSlateWidget.IsValid())
 	{
-		RequestRenderUpdate();
-		LastWidgetRenderTime = 0;
+		if (ISlate3DRenderer* SlateRenderer = WidgetRenderer->GetSlateRenderer())
+		{
+			RequestRenderUpdate();
+			LastWidgetRenderTime = 0;
 
-		return new FWidget3DSceneProxy(this, *WidgetRenderer->GetSlateRenderer());
+			return new FWidget3DSceneProxy(this, *SlateRenderer);
+		}
 	}
 
 #if WITH_EDITOR
