Changed UMenuAnchor to include the index of its owning player when calling SMenuAnchor::SetIsOpen. Should fix issues in split-screen where Player 2 opening a context menu would cause Player 1 to gain control of the menu.

diff --git a/Engine/Source/Runtime/UMG/Private/Components/MenuAnchor.cpp b/Engine/Source/Runtime/UMG/Private/Components/MenuAnchor.cpp
--- a/Engine/Source/Runtime/UMG/Private/Components/MenuAnchor.cpp
+++ b/Engine/Source/Runtime/UMG/Private/Components/MenuAnchor.cpp
@@ -125,7 +125,8 @@ void UMenuAnchor::Open(bool bFocusMenu)
 {
 	if ( MyMenuAnchor.IsValid() && !MyMenuAnchor->IsOpen() )
 	{
-		MyMenuAnchor->SetIsOpen(true, bFocusMenu);
+		const int32 Index = GetOwningLocalPlayer() ? GetOwningLocalPlayer()->GetLocalPlayerIndex() : 0;
+		MyMenuAnchor->SetIsOpen(true, bFocusMenu, Index);
 	}
 }
 
@@ -133,7 +134,8 @@ void UMenuAnchor::Close()
 {
 	if ( MyMenuAnchor.IsValid() )
 	{
-		return MyMenuAnchor->SetIsOpen(false, false);
+		const int32 Index = GetOwningLocalPlayer() ? GetOwningLocalPlayer()->GetLocalPlayerIndex() : 0;
+		return MyMenuAnchor->SetIsOpen(false, false, Index);
 	}
 }
 
