[UMG] WidgetComponent crash when dereferencing null Materialinstance

diff --git a/Engine/Source/Runtime/UMG/Private/Components/WidgetComponent.cpp b/Engine/Source/Runtime/UMG/Private/Components/WidgetComponent.cpp
--- a/Engine/Source/Runtime/UMG/Private/Components/WidgetComponent.cpp
+++ b/Engine/Source/Runtime/UMG/Private/Components/WidgetComponent.cpp
@@ -351,12 +351,12 @@ public:
 		{
 			ParentMaterialProxy = WireframeMaterialInstance;
 		}
-		else
+		else if (MaterialInstance != nullptr)
 		{
 			ParentMaterialProxy = MaterialInstance->GetRenderProxy();
 		}
 #else
-		FMaterialRenderProxy* ParentMaterialProxy = MaterialInstance->GetRenderProxy();
+		FMaterialRenderProxy* ParentMaterialProxy = MaterialInstance ? MaterialInstance->GetRenderProxy() : nullptr;
 #endif
 
 		//FSpriteTextureOverrideRenderProxy* TextureOverrideMaterialProxy = new FSpriteTextureOverrideRenderProxy(ParentMaterialProxy,
