Fix Dcc corruption because of a missing transition in the SlatePostRT system

#tests played the game on multiple consoles

diff --git a/Engine/Source/Runtime/UMG/Private/Slate/SPostBufferUpdate.cpp b/Engine/Source/Runtime/UMG/Private/Slate/SPostBufferUpdate.cpp
--- a/Engine/Source/Runtime/UMG/Private/Slate/SPostBufferUpdate.cpp
+++ b/Engine/Source/Runtime/UMG/Private/Slate/SPostBufferUpdate.cpp
@@ -73,6 +73,8 @@ void FPostBufferUpdater::Draw_RenderThread(FRDGBuilder& GraphBuilder, const FDra
 		// Write to our custom post process texture.
 		const FScreenPassTexture OutputTexture(RegisterExternalTexture(GraphBuilder, SlatePostBuffer->TextureReference.TextureReferenceRHI, TEXT("SlatePostProcessTexture")));
 
+		GraphBuilder.UseInternalAccessMode(OutputTexture.Texture);
+
 		if (TSharedPtr<FSlateRHIPostBufferProcessorProxy> PostProcessorProxy = USlateFXSubsystem::GetPostProcessorProxy(SlatePostBufferBit))
 		{
 			if (TSharedPtr<FSlatePostProcessorUpdaterProxy>* ProcessorUpdaterItr = ProcessorUpdaters.Find(SlatePostBufferBit))
@@ -94,6 +96,8 @@ void FPostBufferUpdater::Draw_RenderThread(FRDGBuilder& GraphBuilder, const FDra
 		{
 			AddDrawTexturePass(GraphBuilder, FScreenPassViewInfo(), InputTexture, OutputTexture);
 		}
+
+		GraphBuilder.UseExternalAccessMode(OutputTexture.Texture, ERHIAccess::SRVMask);
 	}
 }
 
