UMG: Add SlotAsStackBoxSlot helper

diff --git a/Engine/Source/Runtime/UMG/Private/WidgetLayoutLibrary.cpp b/Engine/Source/Runtime/UMG/Private/WidgetLayoutLibrary.cpp
--- a/Engine/Source/Runtime/UMG/Private/WidgetLayoutLibrary.cpp
+++ b/Engine/Source/Runtime/UMG/Private/WidgetLayoutLibrary.cpp
@@ -18,6 +18,7 @@
 #include "Components/ScaleBoxSlot.h"
 #include "Components/ScrollBoxSlot.h"
 #include "Components/SizeBoxSlot.h"
+#include "Components/StackBoxSlot.h"
 #include "Components/WrapBoxSlot.h"
 #include "Components/WidgetSwitcherSlot.h"
 #include "Engine/UserInterfaceSettings.h"
@@ -338,6 +339,16 @@ USizeBoxSlot* UWidgetLayoutLibrary::SlotAsSizeBoxSlot(UWidget* Widget)
 	return nullptr;
 }
 
+UStackBoxSlot* UWidgetLayoutLibrary::SlotAsStackBoxSlot(UWidget* Widget)
+{
+	if (Widget)
+	{
+		return Cast<UStackBoxSlot>(Widget->Slot);
+	}
+
+	return nullptr;
+}
+
 UWrapBoxSlot* UWidgetLayoutLibrary::SlotAsWrapBoxSlot(UWidget* Widget)
 {
 	if (Widget)
diff --git a/Engine/Source/Runtime/UMG/Public/Blueprint/WidgetLayoutLibrary.h b/Engine/Source/Runtime/UMG/Public/Blueprint/WidgetLayoutLibrary.h
--- a/Engine/Source/Runtime/UMG/Public/Blueprint/WidgetLayoutLibrary.h
+++ b/Engine/Source/Runtime/UMG/Public/Blueprint/WidgetLayoutLibrary.h
@@ -20,6 +20,7 @@ class UScrollBoxSlot;
 class USafeZoneSlot;
 class UScaleBoxSlot;
 class USizeBoxSlot;
+class UStackBoxSlot;
 class UWrapBoxSlot;
 class UWidgetSwitcherSlot;
 class UWidget;
@@ -185,6 +186,13 @@ public:
 	UFUNCTION(BlueprintPure, Category = "Slot")
 	static UMG_API USizeBoxSlot* SlotAsSizeBoxSlot(UWidget* Widget);
 
+	/**
+	 * Gets the slot object on the child widget as a Stack Box Slot, allowing you to manipulate its information.
+	 * @param Widget The child widget of a Stack Box.
+	 */
+	UFUNCTION(BlueprintPure, Category = "Slot")
+	static UMG_API UStackBoxSlot* SlotAsStackBoxSlot(UWidget* Widget);
+
 	/**
 	 * Gets the slot object on the child widget as a Wrap Box Slot, allowing you to manipulate its information.
 	 * @param Widget The child widget of a Wrap Box.
