UMGEditor: Improve the usablility of the Widget Reflector at design time, widgets can now be picked by the widget reflector while in "Pick Painted Widgets" mode

- Make the SPaintSurface in the UMG designer invisible to the widget reflector
- Make the designer outline invisible to the widget reflector

diff --git a/Engine/Source/Editor/UMGEditor/Private/Designer/SDesignerView.cpp b/Engine/Source/Editor/UMGEditor/Private/Designer/SDesignerView.cpp
--- a/Engine/Source/Editor/UMGEditor/Private/Designer/SDesignerView.cpp
+++ b/Engine/Source/Editor/UMGEditor/Private/Designer/SDesignerView.cpp
@@ -80,6 +80,7 @@
 #include "ToolMenus.h"
 #include "Styling/ToolBarStyle.h"
 #include "UMGEditorProjectSettings.h"
+#include "Types/InvisibleToWidgetReflectorMetaData.h"
 
 #define LOCTEXT_NAMESPACE "UMG"
 
@@ -338,6 +339,7 @@ void SDesignerView::Construct(const FArguments& InArgs, TSharedPtr<FWidgetBluepr
 				[
 					SAssignNew(EffectsLayer, SPaintSurface)
 					.OnPaintHandler(this, &SDesignerView::HandleEffectsPainting)
+					.AddMetaData<FInvisibleToWidgetReflectorMetaData>(FInvisibleToWidgetReflectorMetaData())
 				]
 
 				// 
diff --git a/Engine/Source/Runtime/UMG/Private/Components/Widget.cpp b/Engine/Source/Runtime/UMG/Private/Components/Widget.cpp
--- a/Engine/Source/Runtime/UMG/Private/Components/Widget.cpp
+++ b/Engine/Source/Runtime/UMG/Private/Components/Widget.cpp
@@ -39,6 +39,7 @@
 #include "Components/NamedSlotInterface.h"
 #include "ProfilingDebugging/AssetMetadataTrace.h"
 #include "HAL/LowLevelMemStats.h"
+#include "Types/InvisibleToWidgetReflectorMetaData.h"
 
 #include UE_INLINE_GENERATED_CPP_BY_NAME(Widget)
 
@@ -1103,6 +1104,7 @@ TSharedRef<SWidget> UWidget::CreateDesignerOutline(TSharedRef<SWidget> Content)
 		.VAlign(VAlign_Fill)
 		[
 			SNew(SBorder)
+			.AddMetaData<FInvisibleToWidgetReflectorMetaData>(FInvisibleToWidgetReflectorMetaData())
 			.DesiredSizeScale(FVector2D(0.0f, 0.0f))
 			.Visibility(HasAnyDesignerFlags(EWidgetDesignFlags::ShowOutline) ? EVisibility::HitTestInvisible : EVisibility::Collapsed)
 			.BorderImage(FUMGStyle::Get().GetBrush("MarchingAnts"))
