When trying to instantitate an abstract class inside DynamicEntryBoxBase, the GetOrCreateInstance returns null, and crashes. Add a null check to avoid the crash.

diff --git a/Engine/Source/Runtime/UMG/Private/Components/DynamicEntryBoxBase.cpp b/Engine/Source/Runtime/UMG/Private/Components/DynamicEntryBoxBase.cpp
--- a/Engine/Source/Runtime/UMG/Private/Components/DynamicEntryBoxBase.cpp
+++ b/Engine/Source/Runtime/UMG/Private/Components/DynamicEntryBoxBase.cpp
@@ -383,7 +383,7 @@ UUserWidget* UDynamicEntryBoxBase::CreateEntryInternal(TSubclassOf<UUserWidget>
 	DynamicEntryBoxBaseCreateEntryInternal::RecursiveDetection.Push(InEntryClass);
 
 	UUserWidget* NewEntryWidget = EntryWidgetPool.GetOrCreateInstance(InEntryClass);
-	if (MyPanelWidget.IsValid())
+	if (MyPanelWidget.IsValid() && NewEntryWidget != nullptr)
 	{
 		// If we've already been constructed, immediately add the child to our panel widget
 		AddEntryChild(*NewEntryWidget);
