Fixed post buffers to use external access mode after being generated. This fixes a missing transition causing artifacts. Also batched together calls to Use{Internal, External}AccessMode to reduce the number of RDG passes.

diff --git a/Engine/Source/Runtime/SlateRHIRenderer/Private/SlateRHIRenderer.cpp b/Engine/Source/Runtime/SlateRHIRenderer/Private/SlateRHIRenderer.cpp
--- a/Engine/Source/Runtime/SlateRHIRenderer/Private/SlateRHIRenderer.cpp
+++ b/Engine/Source/Runtime/SlateRHIRenderer/Private/SlateRHIRenderer.cpp
@@ -685,28 +685,33 @@ FSlateDrawWindowPassOutputs FSlateRHIRenderer::DrawWindow_RenderThread(FRDGBuild
 		FRDGTexture* ElementsTexture = OutputTexture;
 		const FIntPoint OutputExtent = OutputTexture->Desc.Extent;
 
+		TArray<FRDGTexture*, FRDGArrayAllocator> PostProcessTextures;
+		PostProcessTextures.Reserve(Inputs.PostProcessUpdateRequests.Num());
+
 		for (const FSlatePostProcessUpdateRequest& Request : Inputs.PostProcessUpdateRequests)
 		{
-			FRHITexture* PostProcessOutputTextureRHI = Request.RenderTargetTextureResource->GetTexture2DRHI();
-			const FScreenPassTexture PostProcessOutputTexture(RegisterExternalTexture(GraphBuilder, PostProcessOutputTextureRHI, TEXT("PostProcessRT")));
+			FRDGTexture* Texture = RegisterExternalTexture(GraphBuilder, Request.RenderTargetTextureResource->GetTexture2DRHI(), TEXT("PostProcessRT"));
+			PostProcessTextures.Emplace(Texture);
+			GraphBuilder.UseInternalAccessMode(Texture);
+		}
 
-			FIntRect PostProcessInputViewRect;
-			FRHITexture* PostProcessInputTextureRHI;
+		// The post process input texture will be the separate viewport texture if it exists, or the swap chain.
+		FScreenPassTexture PostProcessInputTexture(RegisterExternalTexture(GraphBuilder, ViewportTextureRHI ? ViewportTextureRHI : SwapChainTextureRHI, TEXT("ViewportTexture")));
+
+		for (int32 PostProcessIndex = 0; PostProcessIndex < Inputs.PostProcessUpdateRequests.Num(); ++PostProcessIndex)
+		{
+			const FSlatePostProcessUpdateRequest& Request = Inputs.PostProcessUpdateRequests[PostProcessIndex];
+			const FScreenPassTexture PostProcessOutputTexture(PostProcessTextures[PostProcessIndex]);
 
-			// The viewport texture will contain the scene if it's valid.
 			if (ViewportTextureRHI)
 			{
-				PostProcessInputTextureRHI = ViewportTextureRHI;
-				PostProcessInputViewRect.Max = ViewportTextureRHI->GetSizeXY();
+				PostProcessInputTexture.ViewRect = FIntRect(FIntPoint::ZeroValue, ViewportTextureRHI->GetSizeXY());
 			}
 			else
 			{
-				PostProcessInputTextureRHI = SwapChainTextureRHI;
-				PostProcessInputViewRect = Inputs.SceneViewRect;
+				PostProcessInputTexture.ViewRect = Inputs.SceneViewRect;
 			}
 
-			const FScreenPassTexture PostProcessInputTexture(RegisterExternalTexture(GraphBuilder, PostProcessInputTextureRHI, TEXT("ViewportTexture")), PostProcessInputViewRect);
-
 			if (Request.PostProcessorProxy)
 			{
 				Request.PostProcessorProxy->PostProcess_Renderthread(GraphBuilder, PostProcessInputTexture, PostProcessOutputTexture);
@@ -717,6 +722,11 @@ FSlateDrawWindowPassOutputs FSlateRHIRenderer::DrawWindow_RenderThread(FRDGBuild
 			}
 		}
 
+		for (FRDGTexture* Texture : PostProcessTextures)
+		{
+			GraphBuilder.UseExternalAccessMode(Texture, ERHIAccess::SRVMask);
+		}
+
 		const bool bCompositeUIWithSceneHDR = ViewportInfo.bDisplayFormatIsHDR && CompositeUIWithSceneHDR();
 
 		bool bClearElementsTexture = Inputs.bClear || GSlateWireframe;
diff --git a/Engine/Source/Runtime/UMG/Private/Slate/SPostBufferUpdate.cpp b/Engine/Source/Runtime/UMG/Private/Slate/SPostBufferUpdate.cpp
--- a/Engine/Source/Runtime/UMG/Private/Slate/SPostBufferUpdate.cpp
+++ b/Engine/Source/Runtime/UMG/Private/Slate/SPostBufferUpdate.cpp
@@ -66,23 +66,32 @@ void FPostBufferUpdater::Draw_RenderThread(FRDGBuilder& GraphBuilder, const FDra
 		return;
 	}
 
+	struct FActivePostBuffer
+	{
+		FRDGTexture* Texture = nullptr;
+		TSharedPtr<FSlateRHIPostBufferProcessorProxy> Proxy;
+	};
+
+	// Issue internal / external access mode calls in batches before and after to reduce the number of RDG passes.
+	TStaticArray<FActivePostBuffer, (int32)ESlatePostRT::Num> ActivePostBuffers;
+	int32 SlatePostBufferIndex = 0;
+
 	for (ESlatePostRT SlatePostBufferBit : MakeFlagsRange(Inputs.UsedSlatePostBuffers & BuffersToUpdate_Renderthread))
 	{
+		ON_SCOPE_EXIT
+		{
+			SlatePostBufferIndex++;
+		};
+
 		UTextureRenderTarget2D* SlatePostBuffer = Cast<UTextureRenderTarget2D>(RendererSettings->TryGetPostBufferRT(SlatePostBufferBit));
 		if (!SlatePostBuffer)
 		{
 			continue;
 		}
 
-		// Provided output texture is actually the input into our custom post process texture.
-		const FScreenPassTexture InputTexture(Inputs.OutputTexture, Inputs.SceneViewRect);
-
-		// Write to our custom post process texture.
-		const FScreenPassTexture OutputTexture(RegisterExternalTexture(GraphBuilder, SlatePostBuffer->TextureReference.TextureReferenceRHI, TEXT("SlatePostProcessTexture")));
+		TSharedPtr<FSlateRHIPostBufferProcessorProxy> PostProcessorProxy = USlateFXSubsystem::GetPostProcessorProxy(SlatePostBufferBit);
 
-		GraphBuilder.UseInternalAccessMode(OutputTexture.Texture);
-
-		if (TSharedPtr<FSlateRHIPostBufferProcessorProxy> PostProcessorProxy = USlateFXSubsystem::GetPostProcessorProxy(SlatePostBufferBit))
+		if (PostProcessorProxy)
 		{
 			if (TSharedPtr<FSlatePostProcessorUpdaterProxy>* ProcessorUpdaterItr = ProcessorUpdaters.Find(SlatePostBufferBit))
 			{
@@ -92,19 +101,49 @@ void FPostBufferUpdater::Draw_RenderThread(FRDGBuilder& GraphBuilder, const FDra
 
 					if (ProcessorUpdater->bSkipBufferUpdate)
 					{
-						return;
+						continue;
 					}
 				}
 			}
-
-			PostProcessorProxy->PostProcess_Renderthread(GraphBuilder, InputTexture, OutputTexture);
 		}
-		else
+
+		FRDGTexture* Texture = RegisterExternalTexture(GraphBuilder, SlatePostBuffer->TextureReference.TextureReferenceRHI, TEXT("SlatePostProcessTexture"));
+
+		ActivePostBuffers[SlatePostBufferIndex] = FActivePostBuffer
+		{
+			  .Texture = Texture
+			, .Proxy = MoveTemp(PostProcessorProxy)
+		};
+
+		GraphBuilder.UseInternalAccessMode(Texture);
+	}
+
+	// Provided output texture is actually the input into our custom post process texture.
+	const FScreenPassTexture InputTexture(Inputs.OutputTexture, Inputs.SceneViewRect);
+
+	for (const FActivePostBuffer& ActivePostBuffer : ActivePostBuffers)
+	{
+		if (ActivePostBuffer.Texture)
 		{
-			AddDrawTexturePass(GraphBuilder, FScreenPassViewInfo(), InputTexture, OutputTexture);
+			const FScreenPassTexture OutputTexture(ActivePostBuffer.Texture);
+	
+			if (ActivePostBuffer.Proxy)
+			{
+				ActivePostBuffer.Proxy->PostProcess_Renderthread(GraphBuilder, InputTexture, OutputTexture);
+			}
+			else
+			{
+				AddDrawTexturePass(GraphBuilder, FScreenPassViewInfo(), InputTexture, OutputTexture);
+			}
 		}
+	}
 
-		GraphBuilder.UseExternalAccessMode(OutputTexture.Texture, ERHIAccess::SRVMask);
+	for (const FActivePostBuffer& ActivePostBuffer : ActivePostBuffers)
+	{
+		if (ActivePostBuffer.Texture)
+		{
+			GraphBuilder.UseExternalAccessMode(ActivePostBuffer.Texture, ERHIAccess::SRVMask);
+		}
 	}
 }
 
