Limit the render opacity slider in the details panel to move between 0 to 1. Users can still manually enter any value.

diff --git a/Engine/Source/Runtime/UMG/Public/Components/Widget.h b/Engine/Source/Runtime/UMG/Public/Components/Widget.h
--- a/Engine/Source/Runtime/UMG/Public/Components/Widget.h
+++ b/Engine/Source/Runtime/UMG/Public/Components/Widget.h
@@ -442,7 +442,7 @@ protected:
 
 	UE_DEPRECATED(5.1, "Direct access to RenderOpacity is deprecated. Please use the getter or setter.")
 	/** The opacity of the widget */
-	UPROPERTY(EditAnywhere, BlueprintReadWrite, Getter, Setter, BlueprintGetter="GetRenderOpacity", BlueprintSetter="SetRenderOpacity", Category="Rendering")
+	UPROPERTY(EditAnywhere, BlueprintReadWrite, Getter, Setter, BlueprintGetter="GetRenderOpacity", BlueprintSetter="SetRenderOpacity", Category="Rendering", meta = (UIMin = "0", UIMax = "1"))
 	float RenderOpacity;
 
 private:
