UMG: add a cvar umg.MaxPostBufferBlurStrength to allow limiting the max post update blur strength. Default to -1 which means no limit. This is useful to reduce the GPU impact of blurs on certain platforms
#tests Tested on console, reviewed with UI artists

diff --git a/Engine/Source/Runtime/UMG/Private/Components/PostBufferBlurUpdater.cpp b/Engine/Source/Runtime/UMG/Private/Components/PostBufferBlurUpdater.cpp
--- a/Engine/Source/Runtime/UMG/Private/Components/PostBufferBlurUpdater.cpp
+++ b/Engine/Source/Runtime/UMG/Private/Components/PostBufferBlurUpdater.cpp
@@ -12,6 +12,14 @@
 
 #define LOCTEXT_NAMESPACE "UMG"
 
+TAutoConsoleVariable<int32> CVarUmgMaxPostBufferBlurStrength(
+	TEXT("umg.MaxPostBufferBlurStrength"),
+	-1,
+	TEXT("Limits post buffer blur strength. -1 = unlimited\r\n"),
+	ECVF_Default
+);
+
+
 /////////////////////////////////////////////////////
 // UPostBufferBlurUpdater
 
@@ -34,7 +42,8 @@ void FPostBufferBlurUpdaterProxy::UpdateProcessor_RenderThread(TSharedPtr<FSlate
 {
 #if !UE_SERVER
 	TSharedPtr<FSlatePostBufferBlurProxy> BlurRHIProxy = StaticCastSharedPtr<FSlatePostBufferBlurProxy>(InProcessor);
-	BlurRHIProxy->GaussianBlurStrength_RenderThread = GaussianBlurStrength_RenderThread;
+	const int32 MaxBlurStrength = CVarUmgMaxPostBufferBlurStrength.GetValueOnRenderThread();
+	BlurRHIProxy->GaussianBlurStrength_RenderThread = (MaxBlurStrength >= 0) ? FMath::Min(GaussianBlurStrength_RenderThread, MaxBlurStrength) : GaussianBlurStrength_RenderThread;
 #endif // !UE_SERVER
 }
 
