PR #12439: Modified the UserWidget to handle the first touch event

diff --git a/Engine/Source/Runtime/UMG/Private/Slate/SObjectWidget.cpp b/Engine/Source/Runtime/UMG/Private/Slate/SObjectWidget.cpp
--- a/Engine/Source/Runtime/UMG/Private/Slate/SObjectWidget.cpp
+++ b/Engine/Source/Runtime/UMG/Private/Slate/SObjectWidget.cpp
@@ -518,6 +518,15 @@ FReply SObjectWidget::OnTouchForceChanged(const FGeometry& MyGeometry, const FPo
 	return FReply::Unhandled();
 }
 
+FReply SObjectWidget::OnTouchFirstMove(const FGeometry& MyGeometry, const FPointerEvent& InTouchEvent)
+{
+	if (CanRouteEvent())
+	{
+		return WidgetObject->NativeOnTouchFirstMove(MyGeometry, InTouchEvent);
+	}
+	return FReply::Unhandled();
+}
+
 void SObjectWidget::OnFinishedPointerInput()
 {
 	WidgetObject->NativeOnFinishedPointerInput();
diff --git a/Engine/Source/Runtime/UMG/Private/UserWidget.cpp b/Engine/Source/Runtime/UMG/Private/UserWidget.cpp
--- a/Engine/Source/Runtime/UMG/Private/UserWidget.cpp
+++ b/Engine/Source/Runtime/UMG/Private/UserWidget.cpp
@@ -2329,6 +2329,11 @@ FReply UUserWidget::NativeOnTouchForceChanged(const FGeometry& InGeometry, const
 	return OnTouchForceChanged(InGeometry, InTouchEvent).NativeReply;
 }
 
+FReply UUserWidget::NativeOnTouchFirstMove(const FGeometry& InGeometry, const FPointerEvent& InTouchEvent)
+{
+	return OnTouchFirstMove(InGeometry, InTouchEvent).NativeReply;
+}
+
 FCursorReply UUserWidget::NativeOnCursorQuery( const FGeometry& InGeometry, const FPointerEvent& InCursorEvent )
 {
 	return (bOverride_Cursor)
diff --git a/Engine/Source/Runtime/UMG/Public/Blueprint/UserWidget.h b/Engine/Source/Runtime/UMG/Public/Blueprint/UserWidget.h
--- a/Engine/Source/Runtime/UMG/Public/Blueprint/UserWidget.h
+++ b/Engine/Source/Runtime/UMG/Public/Blueprint/UserWidget.h
@@ -864,6 +864,15 @@ public:
 	UFUNCTION(BlueprintImplementableEvent, BlueprintCosmetic, Category = "Touch Input")
 	UMG_API FEventReply OnTouchForceChanged(FGeometry MyGeometry, const FPointerEvent& InTouchEvent);
 
+	/**
+	* Called when a touchpad touch is first moved (initial finger movement)
+	*
+	* @param MyGeometry    The geometry of the widget receiving the event.
+	* @param InTouchEvent	The touch event generated
+	*/
+	UFUNCTION(BlueprintImplementableEvent, BlueprintCosmetic, Category = "Touch Input")
+	UMG_API FEventReply OnTouchFirstMove(FGeometry MyGeometry, const FPointerEvent& InTouchEvent);
+
 public:
 
 	/**
@@ -1558,6 +1567,7 @@ protected:
 	UMG_API virtual FReply NativeOnTouchEnded( const FGeometry& InGeometry, const FPointerEvent& InGestureEvent );
 	UMG_API virtual FReply NativeOnMotionDetected( const FGeometry& InGeometry, const FMotionEvent& InMotionEvent );
 	UMG_API virtual FReply NativeOnTouchForceChanged(const FGeometry& MyGeometry, const FPointerEvent& TouchEvent);
+	UMG_API virtual FReply NativeOnTouchFirstMove(const FGeometry& MyGeometry, const FPointerEvent& TouchEvent);
 	UMG_API virtual FCursorReply NativeOnCursorQuery( const FGeometry& InGeometry, const FPointerEvent& InCursorEvent );
 	UMG_API virtual FNavigationReply NativeOnNavigation(const FGeometry& InGeometry, const FNavigationEvent& InNavigationEvent);
 	UMG_API virtual void NativeOnMouseCaptureLost(const FCaptureLostEvent& CaptureLostEvent);
diff --git a/Engine/Source/Runtime/UMG/Public/Slate/SObjectWidget.h b/Engine/Source/Runtime/UMG/Public/Slate/SObjectWidget.h
--- a/Engine/Source/Runtime/UMG/Public/Slate/SObjectWidget.h
+++ b/Engine/Source/Runtime/UMG/Public/Slate/SObjectWidget.h
@@ -95,6 +95,7 @@ class SObjectWidget : public SCompoundWidget, public FGCObject
 	UMG_API virtual FReply OnTouchEnded(const FGeometry& MyGeometry, const FPointerEvent& InTouchEvent) override;
 	UMG_API virtual FReply OnMotionDetected(const FGeometry& MyGeometry, const FMotionEvent& InMotionEvent) override;
 	UMG_API virtual FReply OnTouchForceChanged(const FGeometry& MyGeometry, const FPointerEvent& InTouchEvent) override;
+	UMG_API virtual FReply OnTouchFirstMove(const FGeometry& MyGeometry, const FPointerEvent& InTouchEvent) override;
 	UMG_API virtual void OnFinishedPointerInput() override;
 
 	UMG_API virtual FNavigationReply OnNavigation(const FGeometry& MyGeometry, const FNavigationEvent& InNavigationEvent) override;
