[ALT/UserWidget]
 - Fix race in user widget code exercise during loading on both the game thread and async loading thread

#rnx

diff --git a/Engine/Source/Editor/UMGEditor/Private/WidgetBlueprintCompiler.cpp b/Engine/Source/Editor/UMGEditor/Private/WidgetBlueprintCompiler.cpp
--- a/Engine/Source/Editor/UMGEditor/Private/WidgetBlueprintCompiler.cpp
+++ b/Engine/Source/Editor/UMGEditor/Private/WidgetBlueprintCompiler.cpp
@@ -853,7 +853,7 @@ void FWidgetBlueprintCompilerContext::FinishCompilingClass(UClass* Class)
 		FixAbandonedWidgetTree(WidgetBP);
 
 		{
-			TGuardValue<uint32> DisableInitializeFromWidgetTree(UUserWidget::bInitializingFromWidgetTree, 0);
+			TGuardValue<uint32> DisableInitializeFromWidgetTree(UUserWidget::GetInitializingFromWidgetTree(), 0);
 
 			// Need to clear archetype flag before duplication as we check during dup to see if we should postload
 			EObjectFlags PreviousFlags = WidgetBP->WidgetTree->GetFlags();
diff --git a/Engine/Source/Runtime/UMG/Private/UserWidget.cpp b/Engine/Source/Runtime/UMG/Private/UserWidget.cpp
--- a/Engine/Source/Runtime/UMG/Private/UserWidget.cpp
+++ b/Engine/Source/Runtime/UMG/Private/UserWidget.cpp
@@ -49,7 +49,15 @@ TAutoConsoleVariable<bool> CVarUserWidgetUseParallelAnimation(
 	ECVF_Default
 );
 
-uint32 UUserWidget::bInitializingFromWidgetTree = 0;
+namespace UE::UMG::Private
+{
+	static thread_local uint32 InitializingFromWidgetTree = 0;
+}
+
+uint32& UUserWidget::GetInitializingFromWidgetTree()
+{
+	return UE::UMG::Private::InitializingFromWidgetTree;
+}
 
 static FGeometry NullGeometry;
 static FSlateRect NullRect;
@@ -195,7 +203,7 @@ void UUserWidget::InitializeNamedSlots()
 
 void UUserWidget::DuplicateAndInitializeFromWidgetTree(UWidgetTree* InWidgetTree, const TMap<FName, UWidget*>& NamedSlotContentToMerge)
 {
-	TScopeCounter<uint32> ScopeInitializingFromWidgetTree(bInitializingFromWidgetTree);
+	TScopeCounter<uint32> ScopeInitializingFromWidgetTree(GetInitializingFromWidgetTree());
 
 	if ( ensure(InWidgetTree) && !HasAnyFlags(RF_NeedPostLoad))
 	{
@@ -338,7 +346,7 @@ void UUserWidget::PostDuplicate(bool bDuplicateForPIE)
 {
 	Super::PostDuplicate(bDuplicateForPIE);
 	
-	if ( bInitializingFromWidgetTree )
+	if ( GetInitializingFromWidgetTree() != 0 )
 	{
 		// If this is a sub-widget of another UserWidget, default designer flags to match those of the owning widget before initialize.
 		if (UUserWidget* OwningUserWidget = GetTypedOuter<UUserWidget>())
diff --git a/Engine/Source/Runtime/UMG/Public/Blueprint/UserWidget.h b/Engine/Source/Runtime/UMG/Public/Blueprint/UserWidget.h
--- a/Engine/Source/Runtime/UMG/Public/Blueprint/UserWidget.h
+++ b/Engine/Source/Runtime/UMG/Public/Blueprint/UserWidget.h
@@ -1678,8 +1678,7 @@ private:
 	/** Get World calls can be expensive for Widgets, we speed them up by caching the last found world until it goes away. */
 	mutable TWeakObjectPtr<UWorld> CachedWorld;
 
-	static UMG_API bool bTemplateInitializing;
-	static UMG_API uint32 bInitializingFromWidgetTree;
+	static UMG_API uint32& GetInitializingFromWidgetTree();
 
 protected:
 
