UMG: Recover from and prevent serializing null user widget extensions

diff --git a/Engine/Source/Runtime/UMG/Private/UserWidget.cpp b/Engine/Source/Runtime/UMG/Private/UserWidget.cpp
--- a/Engine/Source/Runtime/UMG/Private/UserWidget.cpp
+++ b/Engine/Source/Runtime/UMG/Private/UserWidget.cpp
@@ -2378,6 +2378,15 @@ void UUserWidget::PreSave(FObjectPreSaveContext ObjectSaveContext)
 		RemoveObsoleteBindings(BGClass->NamedSlots);
 	}
 
+	// Prevent null extensions from getting serialized
+	for (auto It = Extensions.CreateIterator(); It; ++It)
+	{
+		if (*It == nullptr)
+		{
+			It.RemoveCurrent();
+		}
+	}
+
 	Super::PreSave(ObjectSaveContext);
 }
 
@@ -2385,6 +2394,15 @@ void UUserWidget::PostLoad()
 {
 	Super::PostLoad();
 
+	// Remove null extensions that have been serialized in our widget
+	for (auto It = Extensions.CreateIterator(); It; ++It)
+	{
+		if (*It == nullptr)
+		{
+			It.RemoveCurrent();
+		}
+	}
+
 #if WITH_EDITOR
 	if (!HasAnyFlags(RF_ClassDefaultObject))
 	{
