Fixed case where Slate post-processing could set an incorrect render target

diff --git a/Engine/Source/Runtime/UMG/Private/Slate/SPostBufferUpdate.cpp b/Engine/Source/Runtime/UMG/Private/Slate/SPostBufferUpdate.cpp
--- a/Engine/Source/Runtime/UMG/Private/Slate/SPostBufferUpdate.cpp
+++ b/Engine/Source/Runtime/UMG/Private/Slate/SPostBufferUpdate.cpp
@@ -20,6 +20,7 @@
 #include "RHIResources.h"
 #include "RHIUtilities.h"
 #include "SlateRHIRendererSettings.h"
+#include "TextureResource.h"
 
 /**
  * Custom Slate drawer to update slate post buffer
@@ -107,7 +108,7 @@ void FPostBufferUpdater::Draw_RenderThread(FRDGBuilder& GraphBuilder, const FDra
 			}
 		}
 
-		FRDGTexture* Texture = RegisterExternalTexture(GraphBuilder, SlatePostBuffer->TextureReference.TextureReferenceRHI, TEXT("SlatePostProcessTexture"));
+		FRDGTexture* Texture = RegisterExternalTexture(GraphBuilder, SlatePostBuffer->GetRenderTargetResource()->GetTextureRHI(), TEXT("SlatePostProcessTexture"));
 
 		ActivePostBuffers[SlatePostBufferIndex] = FActivePostBuffer
 		{
