UMG: Make UStackBox and UPostBufferUpdate constructors public so they can be extended

diff --git a/Engine/Source/Runtime/UMG/Public/Components/PostBufferUpdate.h b/Engine/Source/Runtime/UMG/Public/Components/PostBufferUpdate.h
--- a/Engine/Source/Runtime/UMG/Public/Components/PostBufferUpdate.h
+++ b/Engine/Source/Runtime/UMG/Public/Components/PostBufferUpdate.h
@@ -45,6 +45,7 @@ class UPostBufferUpdate : public UWidget
 {
 	GENERATED_BODY()
 
+public:
 	UPostBufferUpdate();
 
 private:
diff --git a/Engine/Source/Runtime/UMG/Public/Components/StackBox.h b/Engine/Source/Runtime/UMG/Public/Components/StackBox.h
--- a/Engine/Source/Runtime/UMG/Public/Components/StackBox.h
+++ b/Engine/Source/Runtime/UMG/Public/Components/StackBox.h
@@ -22,6 +22,7 @@ class UStackBox : public UPanelWidget
 {
 	GENERATED_BODY()
 
+public:
 	UStackBox();
 
 private:
