UMG: Add missing UMG_API to UStackBox and UPostBufferUpdate constructors so that they can be subclassed

diff --git a/Engine/Source/Runtime/UMG/Public/Components/PostBufferUpdate.h b/Engine/Source/Runtime/UMG/Public/Components/PostBufferUpdate.h
--- a/Engine/Source/Runtime/UMG/Public/Components/PostBufferUpdate.h
+++ b/Engine/Source/Runtime/UMG/Public/Components/PostBufferUpdate.h
@@ -46,7 +46,7 @@ class UPostBufferUpdate : public UWidget
 	GENERATED_BODY()
 
 public:
-	UPostBufferUpdate();
+	UMG_API UPostBufferUpdate();
 
 private:
 
diff --git a/Engine/Source/Runtime/UMG/Public/Components/StackBox.h b/Engine/Source/Runtime/UMG/Public/Components/StackBox.h
--- a/Engine/Source/Runtime/UMG/Public/Components/StackBox.h
+++ b/Engine/Source/Runtime/UMG/Public/Components/StackBox.h
@@ -23,7 +23,7 @@ class UStackBox : public UPanelWidget
 	GENERATED_BODY()
 
 public:
-	UStackBox();
+	UMG_API UStackBox();
 
 private:
 	/** The orientation of the stack box. */
