Add support for regional post buffer updates
Add Slate.DumpNumWidgetPostBufferUpdates & Slate.DumpNumDefaultPostBufferUpdates to trace number of updates occurring.

#tests standalone, PIE, renderdoc

diff --git a/Engine/Source/Runtime/SlateRHIRenderer/Private/FX/SlatePostBufferBlur.cpp b/Engine/Source/Runtime/SlateRHIRenderer/Private/FX/SlatePostBufferBlur.cpp
--- a/Engine/Source/Runtime/SlateRHIRenderer/Private/FX/SlatePostBufferBlur.cpp
+++ b/Engine/Source/Runtime/SlateRHIRenderer/Private/FX/SlatePostBufferBlur.cpp
@@ -11,7 +11,7 @@
 /////////////////////////////////////////////////////
 // FSlatePostBufferBlurProxy
 
-void FSlatePostBufferBlurProxy::PostProcess_Renderthread(FRDGBuilder& GraphBuilder, const FScreenPassTexture& InputTexture, const FScreenPassTexture& OutputTexture)
+void FSlatePostBufferBlurProxy::PostProcess_Renderthread(FRDGBuilder& GraphBuilder, const FScreenPassTexture& InputTexture, const FScreenPassTexture& OutputTexture, const FProcessorProxyInputs& Inputs)
 {
 	RDG_EVENT_SCOPE(GraphBuilder, "SlatePostBufferBlur");
 
@@ -25,6 +25,8 @@ void FSlatePostBufferBlurProxy::PostProcess_Renderthread(FRDGBuilder& GraphBuild
 		BlurInputs.InputTexture = InputTexture;
 		BlurInputs.OutputTexture = OutputTexture;
 		BlurInputs.Strength = GaussianBlurStrength_RenderThread;
+		BlurInputs.PaintGeometry = Inputs.PaintGeometry;
+		BlurInputs.bUsePaintGeometry = Inputs.bUsePaintGeometry;
 
 		AddSlatePostProcessBlurPass(GraphBuilder, BlurInputs);
 	}
diff --git a/Engine/Source/Runtime/SlateRHIRenderer/Private/SlatePostProcessor.cpp b/Engine/Source/Runtime/SlateRHIRenderer/Private/SlatePostProcessor.cpp
--- a/Engine/Source/Runtime/SlateRHIRenderer/Private/SlatePostProcessor.cpp
+++ b/Engine/Source/Runtime/SlateRHIRenderer/Private/SlatePostProcessor.cpp
@@ -486,6 +486,48 @@ void AddSlatePostProcessBlurPass(FRDGBuilder& GraphBuilder, const FSlatePostProc
 	Inputs.Strength         = Strength;
 	Inputs.DownsampleAmount = DownsampleAmount;
 
+	if (SimpleInputs.bUsePaintGeometry && !SimpleInputs.PaintGeometry.GetLocalSize().IsNearlyZero())
+	{
+		const FSlateRenderTransform& RenderTransform = SimpleInputs.PaintGeometry.GetAccumulatedRenderTransform();
+		const FVector2f LocalSize = FVector2f(SimpleInputs.PaintGeometry.GetLocalSize());
+
+		// Determine the four corners of the quad
+		FVector2f TopLeft = FVector2f::ZeroVector;
+		FVector2f TopRight = FVector2f(LocalSize.X, 0.f);
+		FVector2f BotLeft = FVector2f(0.f, LocalSize.Y);
+		FVector2f BotRight = FVector2f(LocalSize.X, LocalSize.Y);
+
+		FVector2f WorldTopLeft = TransformPoint(RenderTransform, TopLeft).RoundToVector();
+		FVector2f WorldBotRight = TransformPoint(RenderTransform, BotRight).RoundToVector();
+
+		// Note: For now we just use transform for scale, this is here for reference if needed
+		//const FVector2f Position = SimpleInputs.PaintGeometry.DrawPosition;
+		//const float Scale = SimpleInputs.PaintGeometry.DrawScale;
+		//FVector2f SizeUV = (WorldBotRight - WorldTopLeft) / WindowSize;
+
+		FIntRect InputProcessingRect = FIntRect(WorldTopLeft.IntPoint(), WorldBotRight.IntPoint());
+		FIntRect OutputProcessingRect = FIntRect(WorldTopLeft.IntPoint(), WorldBotRight.IntPoint());
+
+#if WITH_EDITOR
+		// Our output rect can be offset due to being in a PIE viewport subregion, offset to account for this
+		if (Inputs.InputRect.Min.X > 0 || Inputs.InputRect.Min.Y > 0)
+		{
+			// Note: Subtract offset since our transform already has the PIE viewport offset
+			FVector2f ViewportOffset = FVector2f(Inputs.InputRect.Min.X, Inputs.InputRect.Min.Y);
+			WorldTopLeft -= ViewportOffset;
+			WorldBotRight -= ViewportOffset;
+			OutputProcessingRect = FIntRect(WorldTopLeft.IntPoint(), WorldBotRight.IntPoint());
+		}
+#endif 
+
+		if (InputProcessingRect.Width() <= Inputs.InputRect.Width() && InputProcessingRect.Height() <= Inputs.InputRect.Height()
+			&& OutputProcessingRect.Width() <= Inputs.OutputRect.Width() && OutputProcessingRect.Height() <= Inputs.OutputRect.Height())
+		{
+			Inputs.InputRect = InputProcessingRect;
+			Inputs.OutputRect = OutputProcessingRect;
+		}
+	}
+
 	AddSlatePostProcessBlurPass(GraphBuilder, Inputs);
 }
 
diff --git a/Engine/Source/Runtime/SlateRHIRenderer/Private/SlateRHIRenderer.cpp b/Engine/Source/Runtime/SlateRHIRenderer/Private/SlateRHIRenderer.cpp
--- a/Engine/Source/Runtime/SlateRHIRenderer/Private/SlateRHIRenderer.cpp
+++ b/Engine/Source/Runtime/SlateRHIRenderer/Private/SlateRHIRenderer.cpp
@@ -125,6 +125,14 @@ static TAutoConsoleVariable<int32> CVarRenderThreadTimeIncludesDependentWaits(
 	ECVF_Default
 );
 
+#if WITH_SLATE_DEBUGGING
+static TAutoConsoleVariable<bool> CVarSlateDumpNumDefaultPostBufferUpdates(
+	TEXT("Slate.DumpNumDefaultPostBufferUpdates"),
+	false,
+	TEXT("Dump number of slate default post buffer updates in a frame. Updates every 60f. See also: Slate.DumpNumWidgetPostBufferUpdates.")
+);
+#endif // WITH_SLATE_DEBUGGING
+
 static bool IsVSyncRequired(const FSlateElementBatcher& ElementBatcher)
 {
 	bool bLockToVsync = ElementBatcher.RequiresVsync();
@@ -699,8 +707,35 @@ FSlateDrawWindowPassOutputs FSlateRHIRenderer::DrawWindow_RenderThread(FRDGBuild
 		// The post process input texture will be the separate viewport texture if it exists, or the swap chain.
 		FScreenPassTexture PostProcessInputTexture(RegisterExternalTexture(GraphBuilder, ViewportTextureRHI ? ViewportTextureRHI : SwapChainTextureRHI, TEXT("ViewportTexture")));
 
+#if WITH_SLATE_DEBUGGING
+		static uint64 LastFrameDumpNumUpdates = GFrameCounterRenderThread;
+		static uint32 NumDefaultPostBufferUpdatesThisFrame = 0;
+		if (CVarSlateDumpNumDefaultPostBufferUpdates.GetValueOnRenderThread() && LastFrameDumpNumUpdates != GFrameCounterRenderThread)
+		{
+			LastFrameDumpNumUpdates = GFrameCounterRenderThread;
+
+			uint64 OnScreenTraceKey = 192312;
+			FString NumUpdates = FString::Printf(TEXT("Num Default Post Buffer Updates: %u"), NumDefaultPostBufferUpdatesThisFrame);
+
+			if (LastFrameDumpNumUpdates % 60 == 0)
+			{
+				UE_LOG(LogSlate, Log, TEXT("%s"), *NumUpdates);
+				GEngine->AddOnScreenDebugMessage(OnScreenTraceKey, 15.f, FColor::Cyan, NumUpdates);
+			}
+
+			NumDefaultPostBufferUpdatesThisFrame = 0;
+		}
+#endif // WITH_SLATE_DEBUGGING
+
 		for (int32 PostProcessIndex = 0; PostProcessIndex < Inputs.PostProcessUpdateRequests.Num(); ++PostProcessIndex)
 		{
+#if WITH_SLATE_DEBUGGING
+			if (CVarSlateDumpNumDefaultPostBufferUpdates.GetValueOnRenderThread())
+			{
+				NumDefaultPostBufferUpdatesThisFrame++;
+			}
+#endif // WITH_SLATE_DEBUGGING
+
 			const FSlatePostProcessUpdateRequest& Request = Inputs.PostProcessUpdateRequests[PostProcessIndex];
 			const FScreenPassTexture PostProcessOutputTexture(PostProcessTextures[PostProcessIndex]);
 
diff --git a/Engine/Source/Runtime/SlateRHIRenderer/Public/FX/SlatePostBufferBlur.h b/Engine/Source/Runtime/SlateRHIRenderer/Public/FX/SlatePostBufferBlur.h
--- a/Engine/Source/Runtime/SlateRHIRenderer/Public/FX/SlatePostBufferBlur.h
+++ b/Engine/Source/Runtime/SlateRHIRenderer/Public/FX/SlatePostBufferBlur.h
@@ -16,7 +16,7 @@ class SLATERHIRENDERER_API FSlatePostBufferBlurProxy : public FSlateRHIPostBuffe
 public:
 
 	//~ Begin FSlateRHIPostBufferProcessorProxy Interface
-	virtual void PostProcess_Renderthread(FRDGBuilder& GraphBuilder, const FScreenPassTexture& InputTexture, const FScreenPassTexture& OutputTexture) override;
+	virtual void PostProcess_Renderthread(FRDGBuilder& GraphBuilder, const FScreenPassTexture& InputTexture, const FScreenPassTexture& OutputTexture, const FProcessorProxyInputs& Inputs) override;
 	virtual void OnUpdateValuesRenderThread() override;
 	//~ End FSlateRHIPostBufferProcessorProxy Interface
 
diff --git a/Engine/Source/Runtime/SlateRHIRenderer/Public/FX/SlateRHIPostBufferProcessor.h b/Engine/Source/Runtime/SlateRHIRenderer/Public/FX/SlateRHIPostBufferProcessor.h
--- a/Engine/Source/Runtime/SlateRHIRenderer/Public/FX/SlateRHIPostBufferProcessor.h
+++ b/Engine/Source/Runtime/SlateRHIRenderer/Public/FX/SlateRHIPostBufferProcessor.h
@@ -3,6 +3,7 @@
 #pragma once
 
 #include "Engine/TextureRenderTarget2D.h"
+#include "Layout/PaintGeometry.h"
 #include "Rendering/RenderingCommon.h"
 #include "Rendering/SlateRenderer.h"
 #include "UObject/Object.h"
@@ -14,6 +15,15 @@
 
 class USlateRHIPostBufferProcessor;
 
+struct FProcessorProxyInputs
+{
+	/** Region to clip, if needed */
+	FPaintGeometry PaintGeometry = FPaintGeometry();
+
+	/** True implies we should clip our processing region */
+	bool bUsePaintGeometry = false;
+};
+
 /**
  * Proxy for post buffer processor that the renderthread uses to perform processing
  * This proxy exists because generally speaking usage on UObjects on the renderthread
@@ -23,12 +33,20 @@ class SLATERHIRENDERER_API FSlateRHIPostBufferProcessorProxy : public TSharedFro
 {
 
 public:
+
 	virtual ~FSlateRHIPostBufferProcessorProxy()
 	{
 	}
 
 	/** Called on the render thread to run a post processing operation on the input texture and produce the output texture. */
 	virtual void PostProcess_Renderthread(FRDGBuilder& GraphBuilder, const FScreenPassTexture& InputTexture, const FScreenPassTexture& OutputTexture)
+	{
+		FProcessorProxyInputs ProcessorProxyInputs = FProcessorProxyInputs();
+		PostProcess_Renderthread(GraphBuilder, InputTexture, OutputTexture, ProcessorProxyInputs);
+	}
+
+	/** Same as PostProcess_Renderthread above, but with a inputs struct used for widget specific processing */
+	virtual void PostProcess_Renderthread(FRDGBuilder& GraphBuilder, const FScreenPassTexture& InputTexture, const FScreenPassTexture& OutputTexture, const FProcessorProxyInputs& Inputs)
 	{
 	}
 
diff --git a/Engine/Source/Runtime/SlateRHIRenderer/Public/Interfaces/SlateRHIRenderingPolicyInterface.h b/Engine/Source/Runtime/SlateRHIRenderer/Public/Interfaces/SlateRHIRenderingPolicyInterface.h
--- a/Engine/Source/Runtime/SlateRHIRenderer/Public/Interfaces/SlateRHIRenderingPolicyInterface.h
+++ b/Engine/Source/Runtime/SlateRHIRenderer/Public/Interfaces/SlateRHIRenderingPolicyInterface.h
@@ -2,13 +2,16 @@
 
 #pragma once
 
+#include "Layout/PaintGeometry.h"
 #include "ScreenPass.h"
 
 struct FSlatePostProcessSimpleBlurPassInputs
 {
 	FScreenPassTexture InputTexture;
 	FScreenPassTexture OutputTexture;
+	FPaintGeometry PaintGeometry;
 	float Strength = 0.0f;
+	bool bUsePaintGeometry = false;
 };
 
 SLATERHIRENDERER_API void AddSlatePostProcessBlurPass(FRDGBuilder& GraphBuilder, const FSlatePostProcessSimpleBlurPassInputs& Inputs);
diff --git a/Engine/Source/Runtime/UMG/Private/Components/PostBufferUpdate.cpp b/Engine/Source/Runtime/UMG/Private/Components/PostBufferUpdate.cpp
--- a/Engine/Source/Runtime/UMG/Private/Components/PostBufferUpdate.cpp
+++ b/Engine/Source/Runtime/UMG/Private/Components/PostBufferUpdate.cpp
@@ -14,6 +14,7 @@
 UPostBufferUpdate::UPostBufferUpdate()
 {
 	bPerformDefaultPostBufferUpdate = true;
+	bUpdateOnlyPaintArea = false;
 PRAGMA_DISABLE_DEPRECATION_WARNINGS
 	BuffersToUpdate = {};
 PRAGMA_ENABLE_DEPRECATION_WARNINGS
@@ -35,6 +36,7 @@ void UPostBufferUpdate::SetPerformDefaultPostBufferUpdate(bool bInPerformDefault
 TSharedRef<SWidget> UPostBufferUpdate::RebuildWidget()
 {
 	MyPostBufferUpdate = SNew(SPostBufferUpdate)
+		.bUsePaintGeometry(bUpdateOnlyPaintArea)
 		.bPerformDefaultPostBufferUpdate(bPerformDefaultPostBufferUpdate);
 
 	bool bSetBuffersToUpdate = true;
@@ -88,6 +90,7 @@ void UPostBufferUpdate::SynchronizeProperties()
 	}
 
 	MyPostBufferUpdate->SetPerformDefaultPostBufferUpdate(bPerformDefaultPostBufferUpdate);
+	MyPostBufferUpdate->SetUsePaintGeometry(bUpdateOnlyPaintArea);
 
 	bool bSetBuffersToUpdate = true;
 
diff --git a/Engine/Source/Runtime/UMG/Private/Slate/SPostBufferUpdate.cpp b/Engine/Source/Runtime/UMG/Private/Slate/SPostBufferUpdate.cpp
--- a/Engine/Source/Runtime/UMG/Private/Slate/SPostBufferUpdate.cpp
+++ b/Engine/Source/Runtime/UMG/Private/Slate/SPostBufferUpdate.cpp
@@ -22,6 +22,14 @@
 #include "SlateRHIRendererSettings.h"
 #include "TextureResource.h"
 
+#if WITH_SLATE_DEBUGGING
+static TAutoConsoleVariable<bool> CVarSlateDumpNumWidgetPostBufferUpdates(
+	TEXT("Slate.DumpNumWidgetPostBufferUpdates"),
+	false,
+	TEXT("Dump number of slate widget-driven post buffer updates in a frame. Updates every 60f. See also: Slate.DumpNumDefaultPostBufferUpdates.")
+);
+#endif // WITH_SLATE_DEBUGGING
+
 /**
  * Custom Slate drawer to update slate post buffer
  */
@@ -42,6 +50,9 @@ public:
 	/** True if buffers to update has been initialized. Used to ensure 'BuffersToUpdate_Renderthread' is only set once at initialization. */
 	bool bBuffersToUpdateInitialized = false;
 
+	/** True if we should use the subregion rect */
+	bool bUsePaintGeometry_Renderthread = false;
+
 	/** 
 	 * Buffers that we should update, all of these buffers will be affected by 'bPerformDefaultPostBufferUpdate' if disabled 
 	 * This value is read by the Renderthread, so all non-initialization updates must be done via a render command.
@@ -53,6 +64,9 @@ public:
 
 	/** Proxies used to update a post processor within a frame */
 	TMap<ESlatePostRT, TSharedPtr<FSlatePostProcessorUpdaterProxy>> ProcessorUpdaters;
+
+	/** Paint geometry to use for subregion processing */
+	FPaintGeometry PaintGeometry_Renderthread = FPaintGeometry();
 };
 
 /////////////////////////////////////////////////////
@@ -61,7 +75,6 @@ public:
 void FPostBufferUpdater::Draw_RenderThread(FRDGBuilder& GraphBuilder, const FDrawPassInputs& Inputs)
 {
 	const USlateRHIRendererSettings* RendererSettings = USlateRHIRendererSettings::Get();
-
 	if (!RendererSettings)
 	{
 		return;
@@ -122,15 +135,48 @@ void FPostBufferUpdater::Draw_RenderThread(FRDGBuilder& GraphBuilder, const FDra
 	// Provided output texture is actually the input into our custom post process texture.
 	const FScreenPassTexture InputTexture(Inputs.OutputTexture, Inputs.SceneViewRect);
 
+#if WITH_SLATE_DEBUGGING
+	static uint64 LastFrameDumpNumUpdates = GFrameCounterRenderThread;
+	static uint32 NumWidgetPostBufferUpdatesThisFrame = 0;
+	if (CVarSlateDumpNumWidgetPostBufferUpdates.GetValueOnRenderThread() && LastFrameDumpNumUpdates != GFrameCounterRenderThread)
+	{
+		LastFrameDumpNumUpdates = GFrameCounterRenderThread;
+
+		uint64 OnScreenTraceKey = 202323;
+		FString NumUpdates = FString::Printf(TEXT("Num Widget Post Buffer Updates: %u"), NumWidgetPostBufferUpdatesThisFrame);
+
+		if (LastFrameDumpNumUpdates % 60 == 0)
+		{
+			UE_LOG(LogSlate, Log, TEXT("%s"), *NumUpdates);
+			GEngine->AddOnScreenDebugMessage(OnScreenTraceKey, 15.f, FColor::Cyan, NumUpdates);
+		}
+
+		NumWidgetPostBufferUpdatesThisFrame = 0;
+	}
+#endif // WITH_SLATE_DEBUGGING
+
 	for (const FActivePostBuffer& ActivePostBuffer : ActivePostBuffers)
 	{
 		if (ActivePostBuffer.Texture)
 		{
+#if WITH_SLATE_DEBUGGING
+			if (CVarSlateDumpNumWidgetPostBufferUpdates.GetValueOnRenderThread())
+			{
+				NumWidgetPostBufferUpdatesThisFrame++;
+			}
+#endif // WITH_SLATE_DEBUGGING
+
 			const FScreenPassTexture OutputTexture(ActivePostBuffer.Texture);
 	
 			if (ActivePostBuffer.Proxy)
 			{
-				ActivePostBuffer.Proxy->PostProcess_Renderthread(GraphBuilder, InputTexture, OutputTexture);
+				FProcessorProxyInputs ProcessorProxyInputs = FProcessorProxyInputs();
+				if (bUsePaintGeometry_Renderthread)
+				{
+					ProcessorProxyInputs.PaintGeometry = PaintGeometry_Renderthread;
+					ProcessorProxyInputs.bUsePaintGeometry = true;
+				}
+				ActivePostBuffer.Proxy->PostProcess_Renderthread(GraphBuilder, InputTexture, OutputTexture, ProcessorProxyInputs);
 			}
 			else
 			{
@@ -188,7 +234,8 @@ void SPostBufferUpdate::PrivateRegisterAttributes(FSlateAttributeInitializer& At
 }
 
 SPostBufferUpdate::SPostBufferUpdate()
-	: bPerformDefaultPostBufferUpdate(true)
+	: bUsePaintGeometry(false)
+	, bPerformDefaultPostBufferUpdate(true)
 	, BuffersToUpdate({})
 	, PostBufferUpdater(nullptr)
 {
@@ -197,6 +244,7 @@ SPostBufferUpdate::SPostBufferUpdate()
 void SPostBufferUpdate::Construct(const FArguments& InArgs)
 {
 #if !UE_SERVER
+	bUsePaintGeometry = InArgs._bUsePaintGeometry;
 	bPerformDefaultPostBufferUpdate = InArgs._bPerformDefaultPostBufferUpdate;
 
 	BuffersToUpdate = {};
@@ -205,6 +253,30 @@ void SPostBufferUpdate::Construct(const FArguments& InArgs)
 	if (PostBufferUpdater)
 	{
 		PostBufferUpdater->bPerformDefaultPostBufferUpdate = bPerformDefaultPostBufferUpdate;
+
+		// Safe to set RT values here on construct
+		PostBufferUpdater->bUsePaintGeometry_Renderthread = bUsePaintGeometry;
+	}
+#endif // !UE_SERVER
+}
+
+void SPostBufferUpdate::SetUsePaintGeometry(bool bInUsePaintGeometry)
+{
+#if !UE_SERVER
+	bUsePaintGeometry = bInUsePaintGeometry;
+
+	if (PostBufferUpdater)
+	{
+		TWeakPtr<FPostBufferUpdater> TempWeakPostBufferUpdater = PostBufferUpdater;
+		bool bTempUsePaintGeometry = bUsePaintGeometry;
+
+		ENQUEUE_RENDER_COMMAND(FUpdateValuesRenderThreadFX_UsePaintGeometry)([TempWeakPostBufferUpdater, bTempUsePaintGeometry](FRHICommandListImmediate& RHICmdList)
+		{
+			if (TSharedPtr<FPostBufferUpdater> PostBufferUpdater = TempWeakPostBufferUpdater.Pin())
+			{
+				PostBufferUpdater->bUsePaintGeometry_Renderthread = bTempUsePaintGeometry;
+			}
+		});
 	}
 #endif // !UE_SERVER
 }
@@ -284,6 +356,23 @@ int32 SPostBufferUpdate::OnPaint( const FPaintArgs& Args, const FGeometry& Allot
 #if !UE_SERVER
 	FSlateRect RenderBoundingRect = AllottedGeometry.GetRenderBoundingRect();
 	FPaintGeometry PaintGeometry(RenderBoundingRect.GetTopLeft(), RenderBoundingRect.GetSize(), AllottedGeometry.GetAccumulatedLayoutTransform().GetScale());
+
+	if (PostBufferUpdater && bUsePaintGeometry)
+	{
+		PaintGeometry.CommitTransformsIfUsingLegacyConstructor();
+
+		TWeakPtr<FPostBufferUpdater> TempWeakPostBufferUpdater = PostBufferUpdater;
+
+		ENQUEUE_RENDER_COMMAND(FUpdateValuesRenderThreadFX_PaintGeometry)([TempWeakPostBufferUpdater, PaintGeometry](FRHICommandListImmediate& RHICmdList)
+		{
+			if (TSharedPtr<FPostBufferUpdater> PostBufferUpdater = TempWeakPostBufferUpdater.Pin())
+			{
+				PostBufferUpdater->PaintGeometry_Renderthread = PaintGeometry;
+				PostBufferUpdater->bUsePaintGeometry_Renderthread = true;
+			}
+		});
+	}
+
 	FSlateDrawElement::MakeCustom(OutDrawElements, LayerId, PostBufferUpdater);
 #endif // !UE_SERVER
 
@@ -291,7 +380,7 @@ int32 SPostBufferUpdate::OnPaint( const FPaintArgs& Args, const FGeometry& Allot
 	return ++LayerId;
 }
 
-FVector2D SPostBufferUpdate::ComputeDesiredSize( float ) const
+FVector2D SPostBufferUpdate::ComputeDesiredSize(float LayoutScaleMultiplier) const
 {
 	return FVector2D::ZeroVector;
 }
diff --git a/Engine/Source/Runtime/UMG/Public/Components/PostBufferUpdate.h b/Engine/Source/Runtime/UMG/Public/Components/PostBufferUpdate.h
--- a/Engine/Source/Runtime/UMG/Public/Components/PostBufferUpdate.h
+++ b/Engine/Source/Runtime/UMG/Public/Components/PostBufferUpdate.h
@@ -50,6 +50,12 @@ public:
 
 private:
 
+	/**
+	 * True if we should only update the buffer within the bounds of this widget
+	 */
+	UPROPERTY(EditAnywhere, BlueprintReadOnly, Category = "Behavior", meta = (AllowPrivateAccess = "true"))
+	bool bUpdateOnlyPaintArea;
+
 	/** 
 	 * True if we should do the default post buffer update of the scene before any UI.
 	 * If any PostBufferUpdate widget has this set as false, be default scene copy / processing will not occur.
diff --git a/Engine/Source/Runtime/UMG/Public/Slate/SPostBufferUpdate.h b/Engine/Source/Runtime/UMG/Public/Slate/SPostBufferUpdate.h
--- a/Engine/Source/Runtime/UMG/Public/Slate/SPostBufferUpdate.h
+++ b/Engine/Source/Runtime/UMG/Public/Slate/SPostBufferUpdate.h
@@ -54,9 +54,13 @@ class SPostBufferUpdate
 
 public:
 	SLATE_BEGIN_ARGS( SPostBufferUpdate )
-		: _bPerformDefaultPostBufferUpdate( true )
+		: _bUsePaintGeometry(false)
+		, _bPerformDefaultPostBufferUpdate(true)
 		{ }
 
+		/** True if we should perform the default post buffer update */
+		SLATE_ARGUMENT(bool, bUsePaintGeometry)
+
 		/** True if we should perform the default post buffer update */
 		SLATE_ARGUMENT(bool, bPerformDefaultPostBufferUpdate)
 	SLATE_END_ARGS()
@@ -72,6 +76,8 @@ public:
 	UMG_API void Construct( const FArguments& InArgs );
 
 public:
+	/** Set if we should use paint geometry */
+	UMG_API void SetUsePaintGeometry(bool bInUsePaintGeometry);
 
 	/** Set if we should perform the default post buffer update */
 	UMG_API void SetPerformDefaultPostBufferUpdate(bool bInPerformDefaultPostBufferUpdate);
@@ -107,11 +113,14 @@ public:
 protected:
 
 	//~Begin Widget Interface
-	UMG_API virtual FVector2D ComputeDesiredSize(float) const override;
+	UMG_API virtual FVector2D ComputeDesiredSize(float LayoutScaleMultiplier) const override;
 	//~End Widget Interface
 
 private:
 
+	/** True if we should limit processing to our paint geometry */
+	bool bUsePaintGeometry;
+
 	/** True if we should perform the default post buffer update */
 	bool bPerformDefaultPostBufferUpdate;
 
