Fix false shadowing error on Mac builds PostBuffer

#tests nonunity

diff --git a/Engine/Source/Runtime/UMG/Private/Slate/SPostBufferUpdate.cpp b/Engine/Source/Runtime/UMG/Private/Slate/SPostBufferUpdate.cpp
--- a/Engine/Source/Runtime/UMG/Private/Slate/SPostBufferUpdate.cpp
+++ b/Engine/Source/Runtime/UMG/Private/Slate/SPostBufferUpdate.cpp
@@ -272,9 +272,9 @@ void SPostBufferUpdate::SetUsePaintGeometry(bool bInUsePaintGeometry)
 
 		ENQUEUE_RENDER_COMMAND(FUpdateValuesRenderThreadFX_UsePaintGeometry)([TempWeakPostBufferUpdater, bTempUsePaintGeometry](FRHICommandListImmediate& RHICmdList)
 		{
-			if (TSharedPtr<FPostBufferUpdater> PostBufferUpdater = TempWeakPostBufferUpdater.Pin())
+			if (TSharedPtr<FPostBufferUpdater> PinnedPostBufferUpdater = TempWeakPostBufferUpdater.Pin())
 			{
-				PostBufferUpdater->bUsePaintGeometry_Renderthread = bTempUsePaintGeometry;
+				PinnedPostBufferUpdater->bUsePaintGeometry_Renderthread = bTempUsePaintGeometry;
 			}
 		});
 	}
@@ -365,10 +365,10 @@ int32 SPostBufferUpdate::OnPaint( const FPaintArgs& Args, const FGeometry& Allot
 
 		ENQUEUE_RENDER_COMMAND(FUpdateValuesRenderThreadFX_PaintGeometry)([TempWeakPostBufferUpdater, PaintGeometry](FRHICommandListImmediate& RHICmdList)
 		{
-			if (TSharedPtr<FPostBufferUpdater> PostBufferUpdater = TempWeakPostBufferUpdater.Pin())
+			if (TSharedPtr<FPostBufferUpdater> PinnedPostBufferUpdater = TempWeakPostBufferUpdater.Pin())
 			{
-				PostBufferUpdater->PaintGeometry_Renderthread = PaintGeometry;
-				PostBufferUpdater->bUsePaintGeometry_Renderthread = true;
+				PinnedPostBufferUpdater->PaintGeometry_Renderthread = PaintGeometry;
+				PinnedPostBufferUpdater->bUsePaintGeometry_Renderthread = true;
 			}
 		});
 	}
