Simplified post buffer fix to use geometry bounds of the widget. Moved the logic up the stack so that the inputs to the simple blur pass can remain the same (rather than having conflicting inputs). Also removed hacked in debug logging code.

diff --git a/Engine/Source/Runtime/SlateRHIRenderer/Private/FX/SlatePostBufferBlur.cpp b/Engine/Source/Runtime/SlateRHIRenderer/Private/FX/SlatePostBufferBlur.cpp
--- a/Engine/Source/Runtime/SlateRHIRenderer/Private/FX/SlatePostBufferBlur.cpp
+++ b/Engine/Source/Runtime/SlateRHIRenderer/Private/FX/SlatePostBufferBlur.cpp
@@ -11,7 +11,7 @@
 /////////////////////////////////////////////////////
 // FSlatePostBufferBlurProxy
 
-void FSlatePostBufferBlurProxy::PostProcess_Renderthread(FRDGBuilder& GraphBuilder, const FScreenPassTexture& InputTexture, const FScreenPassTexture& OutputTexture, const FProcessorProxyInputs& Inputs)
+void FSlatePostBufferBlurProxy::PostProcess_Renderthread(FRDGBuilder& GraphBuilder, const FScreenPassTexture& InputTexture, const FScreenPassTexture& OutputTexture)
 {
 	RDG_EVENT_SCOPE(GraphBuilder, "SlatePostBufferBlur");
 
@@ -25,8 +25,6 @@ void FSlatePostBufferBlurProxy::PostProcess_Renderthread(FRDGBuilder& GraphBuild
 		BlurInputs.InputTexture = InputTexture;
 		BlurInputs.OutputTexture = OutputTexture;
 		BlurInputs.Strength = GaussianBlurStrength_RenderThread;
-		BlurInputs.PaintGeometry = Inputs.PaintGeometry;
-		BlurInputs.bUsePaintGeometry = Inputs.bUsePaintGeometry;
 
 		AddSlatePostProcessBlurPass(GraphBuilder, BlurInputs);
 	}
diff --git a/Engine/Source/Runtime/SlateRHIRenderer/Private/SlatePostProcessor.cpp b/Engine/Source/Runtime/SlateRHIRenderer/Private/SlatePostProcessor.cpp
--- a/Engine/Source/Runtime/SlateRHIRenderer/Private/SlatePostProcessor.cpp
+++ b/Engine/Source/Runtime/SlateRHIRenderer/Private/SlatePostProcessor.cpp
@@ -486,48 +486,6 @@ void AddSlatePostProcessBlurPass(FRDGBuilder& GraphBuilder, const FSlatePostProc
 	Inputs.Strength         = Strength;
 	Inputs.DownsampleAmount = DownsampleAmount;
 
-	if (SimpleInputs.bUsePaintGeometry && !SimpleInputs.PaintGeometry.GetLocalSize().IsNearlyZero())
-	{
-		const FSlateRenderTransform& RenderTransform = SimpleInputs.PaintGeometry.GetAccumulatedRenderTransform();
-		const FVector2f LocalSize = FVector2f(SimpleInputs.PaintGeometry.GetLocalSize());
-
-		// Determine the four corners of the quad
-		FVector2f TopLeft = FVector2f::ZeroVector;
-		FVector2f TopRight = FVector2f(LocalSize.X, 0.f);
-		FVector2f BotLeft = FVector2f(0.f, LocalSize.Y);
-		FVector2f BotRight = FVector2f(LocalSize.X, LocalSize.Y);
-
-		FVector2f WorldTopLeft = TransformPoint(RenderTransform, TopLeft).RoundToVector();
-		FVector2f WorldBotRight = TransformPoint(RenderTransform, BotRight).RoundToVector();
-
-		// Note: For now we just use transform for scale, this is here for reference if needed
-		//const FVector2f Position = SimpleInputs.PaintGeometry.DrawPosition;
-		//const float Scale = SimpleInputs.PaintGeometry.DrawScale;
-		//FVector2f SizeUV = (WorldBotRight - WorldTopLeft) / WindowSize;
-
-		FIntRect InputProcessingRect = FIntRect(WorldTopLeft.IntPoint(), WorldBotRight.IntPoint());
-		FIntRect OutputProcessingRect = FIntRect(WorldTopLeft.IntPoint(), WorldBotRight.IntPoint());
-
-#if WITH_EDITOR
-		// Our output rect can be offset due to being in a PIE viewport subregion, offset to account for this
-		if (Inputs.InputRect.Min.X > 0 || Inputs.InputRect.Min.Y > 0)
-		{
-			// Note: Subtract offset since our transform already has the PIE viewport offset
-			FVector2f ViewportOffset = FVector2f(Inputs.InputRect.Min.X, Inputs.InputRect.Min.Y);
-			WorldTopLeft -= ViewportOffset;
-			WorldBotRight -= ViewportOffset;
-			OutputProcessingRect = FIntRect(WorldTopLeft.IntPoint(), WorldBotRight.IntPoint());
-		}
-#endif 
-
-		if (InputProcessingRect.Width() <= Inputs.InputRect.Width() && InputProcessingRect.Height() <= Inputs.InputRect.Height()
-			&& OutputProcessingRect.Width() <= Inputs.OutputRect.Width() && OutputProcessingRect.Height() <= Inputs.OutputRect.Height())
-		{
-			Inputs.InputRect = InputProcessingRect;
-			Inputs.OutputRect = OutputProcessingRect;
-		}
-	}
-
 	AddSlatePostProcessBlurPass(GraphBuilder, Inputs);
 }
 
diff --git a/Engine/Source/Runtime/SlateRHIRenderer/Private/SlateRHIRenderer.cpp b/Engine/Source/Runtime/SlateRHIRenderer/Private/SlateRHIRenderer.cpp
--- a/Engine/Source/Runtime/SlateRHIRenderer/Private/SlateRHIRenderer.cpp
+++ b/Engine/Source/Runtime/SlateRHIRenderer/Private/SlateRHIRenderer.cpp
@@ -707,35 +707,8 @@ FSlateDrawWindowPassOutputs FSlateRHIRenderer::DrawWindow_RenderThread(FRDGBuild
 		// The post process input texture will be the separate viewport texture if it exists, or the swap chain.
 		FScreenPassTexture PostProcessInputTexture(RegisterExternalTexture(GraphBuilder, ViewportTextureRHI ? ViewportTextureRHI : SwapChainTextureRHI, TEXT("ViewportTexture")));
 
-#if WITH_SLATE_DEBUGGING
-		static uint64 LastFrameDumpNumUpdates = GFrameCounterRenderThread;
-		static uint32 NumDefaultPostBufferUpdatesThisFrame = 0;
-		if (CVarSlateDumpNumDefaultPostBufferUpdates.GetValueOnRenderThread() && LastFrameDumpNumUpdates != GFrameCounterRenderThread)
-		{
-			LastFrameDumpNumUpdates = GFrameCounterRenderThread;
-
-			uint64 OnScreenTraceKey = 192312;
-			FString NumUpdates = FString::Printf(TEXT("Num Default Post Buffer Updates: %u"), NumDefaultPostBufferUpdatesThisFrame);
-
-			if (LastFrameDumpNumUpdates % 60 == 0)
-			{
-				UE_LOG(LogSlate, Log, TEXT("%s"), *NumUpdates);
-				GEngine->AddOnScreenDebugMessage(OnScreenTraceKey, 15.f, FColor::Cyan, NumUpdates);
-			}
-
-			NumDefaultPostBufferUpdatesThisFrame = 0;
-		}
-#endif // WITH_SLATE_DEBUGGING
-
 		for (int32 PostProcessIndex = 0; PostProcessIndex < Inputs.PostProcessUpdateRequests.Num(); ++PostProcessIndex)
 		{
-#if WITH_SLATE_DEBUGGING
-			if (CVarSlateDumpNumDefaultPostBufferUpdates.GetValueOnRenderThread())
-			{
-				NumDefaultPostBufferUpdatesThisFrame++;
-			}
-#endif // WITH_SLATE_DEBUGGING
-
 			const FSlatePostProcessUpdateRequest& Request = Inputs.PostProcessUpdateRequests[PostProcessIndex];
 			const FScreenPassTexture PostProcessOutputTexture(PostProcessTextures[PostProcessIndex]);
 
diff --git a/Engine/Source/Runtime/SlateRHIRenderer/Public/FX/SlatePostBufferBlur.h b/Engine/Source/Runtime/SlateRHIRenderer/Public/FX/SlatePostBufferBlur.h
--- a/Engine/Source/Runtime/SlateRHIRenderer/Public/FX/SlatePostBufferBlur.h
+++ b/Engine/Source/Runtime/SlateRHIRenderer/Public/FX/SlatePostBufferBlur.h
@@ -16,7 +16,7 @@ class SLATERHIRENDERER_API FSlatePostBufferBlurProxy : public FSlateRHIPostBuffe
 public:
 
 	//~ Begin FSlateRHIPostBufferProcessorProxy Interface
-	virtual void PostProcess_Renderthread(FRDGBuilder& GraphBuilder, const FScreenPassTexture& InputTexture, const FScreenPassTexture& OutputTexture, const FProcessorProxyInputs& Inputs) override;
+	virtual void PostProcess_Renderthread(FRDGBuilder& GraphBuilder, const FScreenPassTexture& InputTexture, const FScreenPassTexture& OutputTexture) override;
 	virtual void OnUpdateValuesRenderThread() override;
 	//~ End FSlateRHIPostBufferProcessorProxy Interface
 
diff --git a/Engine/Source/Runtime/SlateRHIRenderer/Public/FX/SlateRHIPostBufferProcessor.h b/Engine/Source/Runtime/SlateRHIRenderer/Public/FX/SlateRHIPostBufferProcessor.h
--- a/Engine/Source/Runtime/SlateRHIRenderer/Public/FX/SlateRHIPostBufferProcessor.h
+++ b/Engine/Source/Runtime/SlateRHIRenderer/Public/FX/SlateRHIPostBufferProcessor.h
@@ -15,15 +15,6 @@
 
 class USlateRHIPostBufferProcessor;
 
-struct FProcessorProxyInputs
-{
-	/** Region to clip, if needed */
-	FPaintGeometry PaintGeometry = FPaintGeometry();
-
-	/** True implies we should clip our processing region */
-	bool bUsePaintGeometry = false;
-};
-
 /**
  * Proxy for post buffer processor that the renderthread uses to perform processing
  * This proxy exists because generally speaking usage on UObjects on the renderthread
@@ -40,13 +31,6 @@ public:
 
 	/** Called on the render thread to run a post processing operation on the input texture and produce the output texture. */
 	virtual void PostProcess_Renderthread(FRDGBuilder& GraphBuilder, const FScreenPassTexture& InputTexture, const FScreenPassTexture& OutputTexture)
-	{
-		FProcessorProxyInputs ProcessorProxyInputs = FProcessorProxyInputs();
-		PostProcess_Renderthread(GraphBuilder, InputTexture, OutputTexture, ProcessorProxyInputs);
-	}
-
-	/** Same as PostProcess_Renderthread above, but with a inputs struct used for widget specific processing */
-	virtual void PostProcess_Renderthread(FRDGBuilder& GraphBuilder, const FScreenPassTexture& InputTexture, const FScreenPassTexture& OutputTexture, const FProcessorProxyInputs& Inputs)
 	{
 	}
 
diff --git a/Engine/Source/Runtime/SlateRHIRenderer/Public/Interfaces/SlateRHIRenderingPolicyInterface.h b/Engine/Source/Runtime/SlateRHIRenderer/Public/Interfaces/SlateRHIRenderingPolicyInterface.h
--- a/Engine/Source/Runtime/SlateRHIRenderer/Public/Interfaces/SlateRHIRenderingPolicyInterface.h
+++ b/Engine/Source/Runtime/SlateRHIRenderer/Public/Interfaces/SlateRHIRenderingPolicyInterface.h
@@ -2,16 +2,13 @@
 
 #pragma once
 
-#include "Layout/PaintGeometry.h"
 #include "ScreenPass.h"
 
 struct FSlatePostProcessSimpleBlurPassInputs
 {
 	FScreenPassTexture InputTexture;
 	FScreenPassTexture OutputTexture;
-	FPaintGeometry PaintGeometry;
 	float Strength = 0.0f;
-	bool bUsePaintGeometry = false;
 };
 
 SLATERHIRENDERER_API void AddSlatePostProcessBlurPass(FRDGBuilder& GraphBuilder, const FSlatePostProcessSimpleBlurPassInputs& Inputs);
diff --git a/Engine/Source/Runtime/UMG/Private/Slate/SPostBufferUpdate.cpp b/Engine/Source/Runtime/UMG/Private/Slate/SPostBufferUpdate.cpp
--- a/Engine/Source/Runtime/UMG/Private/Slate/SPostBufferUpdate.cpp
+++ b/Engine/Source/Runtime/UMG/Private/Slate/SPostBufferUpdate.cpp
@@ -22,14 +22,6 @@
 #include "SlateRHIRendererSettings.h"
 #include "TextureResource.h"
 
-#if WITH_SLATE_DEBUGGING
-static TAutoConsoleVariable<bool> CVarSlateDumpNumWidgetPostBufferUpdates(
-	TEXT("Slate.DumpNumWidgetPostBufferUpdates"),
-	false,
-	TEXT("Dump number of slate widget-driven post buffer updates in a frame. Updates every 60f. See also: Slate.DumpNumDefaultPostBufferUpdates.")
-);
-#endif // WITH_SLATE_DEBUGGING
-
 /**
  * Custom Slate drawer to update slate post buffer
  */
@@ -132,51 +124,33 @@ void FPostBufferUpdater::Draw_RenderThread(FRDGBuilder& GraphBuilder, const FDra
 		GraphBuilder.UseInternalAccessMode(Texture);
 	}
 
-	// Provided output texture is actually the input into our custom post process texture.
-	const FScreenPassTexture InputTexture(Inputs.OutputTexture, Inputs.SceneViewRect);
-
-#if WITH_SLATE_DEBUGGING
-	static uint64 LastFrameDumpNumUpdates = GFrameCounterRenderThread;
-	static uint32 NumWidgetPostBufferUpdatesThisFrame = 0;
-	if (CVarSlateDumpNumWidgetPostBufferUpdates.GetValueOnRenderThread() && LastFrameDumpNumUpdates != GFrameCounterRenderThread)
-	{
-		LastFrameDumpNumUpdates = GFrameCounterRenderThread;
-
-		uint64 OnScreenTraceKey = 202323;
-		FString NumUpdates = FString::Printf(TEXT("Num Widget Post Buffer Updates: %u"), NumWidgetPostBufferUpdatesThisFrame);
-
-		if (LastFrameDumpNumUpdates % 60 == 0)
-		{
-			UE_LOG(LogSlate, Log, TEXT("%s"), *NumUpdates);
-			GEngine->AddOnScreenDebugMessage(OnScreenTraceKey, 15.f, FColor::Cyan, NumUpdates);
-		}
-
-		NumWidgetPostBufferUpdatesThisFrame = 0;
-	}
-#endif // WITH_SLATE_DEBUGGING
-
 	for (const FActivePostBuffer& ActivePostBuffer : ActivePostBuffers)
 	{
 		if (ActivePostBuffer.Texture)
 		{
-#if WITH_SLATE_DEBUGGING
-			if (CVarSlateDumpNumWidgetPostBufferUpdates.GetValueOnRenderThread())
-			{
-				NumWidgetPostBufferUpdatesThisFrame++;
-			}
-#endif // WITH_SLATE_DEBUGGING
-
-			const FScreenPassTexture OutputTexture(ActivePostBuffer.Texture);
+			// Provided output texture is actually the input into our custom post process texture.
+			FScreenPassTexture InputTexture(Inputs.OutputTexture, Inputs.SceneViewRect);
+			FScreenPassTexture OutputTexture(ActivePostBuffer.Texture);
 	
 			if (ActivePostBuffer.Proxy)
 			{
-				FProcessorProxyInputs ProcessorProxyInputs = FProcessorProxyInputs();
-				if (bUsePaintGeometry_Renderthread)
+				// If we are using the paint geometry instead, then override the input / output view rects.
+				if (bUsePaintGeometry_Renderthread && !PaintGeometry_Renderthread.GetLocalSize().IsNearlyZero())
 				{
-					ProcessorProxyInputs.PaintGeometry = PaintGeometry_Renderthread;
-					ProcessorProxyInputs.bUsePaintGeometry = true;
+					const FSlateRenderTransform& RenderTransform = PaintGeometry_Renderthread.GetAccumulatedRenderTransform();
+
+					const FVector2f WorldTopLeft = TransformPoint(RenderTransform, FVector2f::ZeroVector).RoundToVector();
+					const FVector2f WorldBottomRight = TransformPoint(RenderTransform, FVector2f(PaintGeometry_Renderthread.GetLocalSize())).RoundToVector();
+
+					const FVector2f InputOffset = FVector2f(InputTexture.ViewRect.Min.X, InputTexture.ViewRect.Min.Y);
+					const FVector2f OutputOffset = FVector2f(OutputTexture.ViewRect.Min.X, OutputTexture.ViewRect.Min.Y);
+
+					InputTexture.ViewRect = FIntRect(WorldTopLeft.IntPoint(), WorldBottomRight.IntPoint());
+
+					// Subtract the input offset because the geometry world transform will have it implicitly added.
+					OutputTexture.ViewRect = FIntRect((WorldTopLeft - InputOffset + OutputOffset).IntPoint(), (WorldBottomRight - InputOffset + OutputOffset).IntPoint());
 				}
-				ActivePostBuffer.Proxy->PostProcess_Renderthread(GraphBuilder, InputTexture, OutputTexture, ProcessorProxyInputs);
+				ActivePostBuffer.Proxy->PostProcess_Renderthread(GraphBuilder, InputTexture, OutputTexture);
 			}
 			else
 			{
