UMG: When rebuilding widgets, don't overwrite the player context on child widgets if they already have a valid player context

#rnx

diff --git a/Engine/Source/Runtime/UMG/Private/UserWidget.cpp b/Engine/Source/Runtime/UMG/Private/UserWidget.cpp
--- a/Engine/Source/Runtime/UMG/Private/UserWidget.cpp
+++ b/Engine/Source/Runtime/UMG/Private/UserWidget.cpp
@@ -1117,7 +1117,7 @@ TSharedRef<SWidget> UUserWidget::RebuildWidget()
 		WidgetTree->ForEachWidget([&] (UWidget* Widget) {
 			if ( UUserWidget* UserWidget = Cast<UUserWidget>(Widget) )
 			{
-				UserWidget->SetPlayerContext(PlayerContext);
+				UserWidget->UpdatePlayerContextIfInvalid(PlayerContext);
 			}
 		});
 	}
@@ -2550,6 +2550,28 @@ void UUserWidget::ClearStoppedSequencePlayers()
 	StoppedSequencePlayers.Empty();
 }
 
+void UUserWidget::UpdatePlayerContextIfInvalid(const FLocalPlayerContext& ParentPlayerContext)
+{
+	if (PlayerContext.IsValid())
+	{
+		if (WidgetTree)
+		{
+			WidgetTree->ForEachWidget(
+				[this] (UWidget* Widget)
+				{
+					if (UUserWidget* UserWidget = Cast<UUserWidget>(Widget))
+					{
+						UserWidget->UpdatePlayerContextIfInvalid(PlayerContext);
+					}
+				});
+		}
+	}
+	else
+	{
+		SetPlayerContext(ParentPlayerContext);
+	}
+}
+
 void UUserWidget::OnLatentActionsChanged(UObject* ObjectWhichChanged, ELatentActionChangeType ChangeType)
 {
 	if (UUserWidget* WidgetThatChanged = Cast<UUserWidget>(ObjectWhichChanged))
diff --git a/Engine/Source/Runtime/UMG/Public/Blueprint/UserWidget.h b/Engine/Source/Runtime/UMG/Public/Blueprint/UserWidget.h
--- a/Engine/Source/Runtime/UMG/Public/Blueprint/UserWidget.h
+++ b/Engine/Source/Runtime/UMG/Public/Blueprint/UserWidget.h
@@ -1432,6 +1432,9 @@ private:
 
 	UMG_API void ClearStoppedSequencePlayers();
 
+	// If this widget's player context is invalid, this function will set it to the provided player context, otherwise it will keep its current player context and call CheckUpdatePlayerContext on its children
+	void UpdatePlayerContextIfInvalid(const FLocalPlayerContext& ParentPlayerContext);
+
 public:
 
 	/** Animation transitions to trigger on next tick */
