Fixed NaN when trying to draw blurred regions with Slate

diff --git a/Engine/Source/Runtime/UMG/Private/Slate/SPostBufferUpdate.cpp b/Engine/Source/Runtime/UMG/Private/Slate/SPostBufferUpdate.cpp
--- a/Engine/Source/Runtime/UMG/Private/Slate/SPostBufferUpdate.cpp
+++ b/Engine/Source/Runtime/UMG/Private/Slate/SPostBufferUpdate.cpp
@@ -135,7 +135,7 @@ void FPostBufferUpdater::Draw_RenderThread(FRDGBuilder& GraphBuilder, const FDra
 			if (ActivePostBuffer.Proxy)
 			{
 				// If we are using the paint geometry instead, then override the input / output view rects.
-				if (bUsePaintGeometry_Renderthread && !PaintGeometry_Renderthread.GetLocalSize().IsNearlyZero())
+				if (bUsePaintGeometry_Renderthread && !FMath::IsNearlyZero(PaintGeometry_Renderthread.GetLocalSize().X) && !FMath::IsNearlyZero(PaintGeometry_Renderthread.GetLocalSize().Y))
 				{
 					const FSlateRenderTransform& RenderTransform = PaintGeometry_Renderthread.GetAccumulatedRenderTransform();
 
