Workaround for RHI validation error in PostBuffer due to it being referenced by an earlier slate render batch in the list before the buffer is updated for the first time.

diff --git a/Engine/Source/Runtime/UMG/Private/Slate/SPostBufferUpdate.cpp b/Engine/Source/Runtime/UMG/Private/Slate/SPostBufferUpdate.cpp
--- a/Engine/Source/Runtime/UMG/Private/Slate/SPostBufferUpdate.cpp
+++ b/Engine/Source/Runtime/UMG/Private/Slate/SPostBufferUpdate.cpp
@@ -113,7 +113,8 @@ void FPostBufferUpdater::Draw_RenderThread(FRDGBuilder& GraphBuilder, const FDra
 			}
 		}
 
-		FRDGTexture* Texture = RegisterExternalTexture(GraphBuilder, SlatePostBuffer->GetRenderTargetResource()->GetTextureRHI(), TEXT("SlatePostProcessTexture"));
+		// Force the first barrier to be immediate to handle edge case where a prior slate render batch can be referencing an older version of this resource.
+		FRDGTexture* Texture = RegisterExternalTexture(GraphBuilder, SlatePostBuffer->GetRenderTargetResource()->GetTextureRHI(), TEXT("SlatePostProcessTexture"), ERDGTextureFlags::ForceImmediateFirstBarrier);
 
 		ActivePostBuffers[SlatePostBufferIndex] = FActivePostBuffer
 		{
