Added a OnEntriesGenerated event to UListViewBase

In Blueprint, it is sometimes useful to get notified once all widgets have been generated for a list. We were already firing individual events for each generated entry (OnEntryGenerated) in UListViewBase::HandleAnnounceGeneratedEntries() so I added a single OnEntriesGenerated (plural) which is invoked once after the loop in HandleAnnounceGeneratedEntries.

#rnx

#tests Tested in PIE.

diff --git a/Engine/Source/Runtime/UMG/Private/Components/ListViewBase.cpp b/Engine/Source/Runtime/UMG/Private/Components/ListViewBase.cpp
--- a/Engine/Source/Runtime/UMG/Private/Components/ListViewBase.cpp
+++ b/Engine/Source/Runtime/UMG/Private/Components/ListViewBase.cpp
@@ -274,6 +274,9 @@ void UListViewBase::HandleAnnounceGeneratedEntries()
 		}
 	}
 	GeneratedEntriesToAnnounce.Empty();
+
+	NativeOnEntriesGenerated();
+	BP_OnEntriesGenerated.Broadcast(GetDisplayedEntryWidgets().Num());
 }
 
 #undef LOCTEXT_NAMESPACE
diff --git a/Engine/Source/Runtime/UMG/Public/Components/ListViewBase.h b/Engine/Source/Runtime/UMG/Public/Components/ListViewBase.h
--- a/Engine/Source/Runtime/UMG/Public/Components/ListViewBase.h
+++ b/Engine/Source/Runtime/UMG/Public/Components/ListViewBase.h
@@ -15,6 +15,7 @@
 
 DECLARE_DYNAMIC_MULTICAST_DELEGATE_OneParam(FOnListEntryGeneratedDynamic, UUserWidget*, Widget);
 DECLARE_DYNAMIC_MULTICAST_DELEGATE_OneParam(FOnListEntryReleasedDynamic, UUserWidget*, Widget);
+DECLARE_DYNAMIC_MULTICAST_DELEGATE_OneParam(FOnListEntriesGeneratedDynamic, int32, NumEntries);
 
 //////////////////////////////////////////////////////////////////////////
 // ITypedUMGListView<T>
@@ -594,6 +595,11 @@ protected:
     FOnListEntryGeneratedDynamic BP_OnEntryGenerated;
 	virtual void NativeOnEntryGenerated(UUserWidget* EntryWidget) {};
 
+	/** Called when all row widgets are generated for all list items */
+	UPROPERTY(BlueprintAssignable, Category = Events, meta = (DisplayName = "On Entries Generated"))
+	FOnListEntriesGeneratedDynamic BP_OnEntriesGenerated;
+	virtual void NativeOnEntriesGenerated() {};
+
 	/**
 	* Normally these are processed by UListViewBase::FinishGeneratingEntry which uses World->GetTimerManager() to generate entries next frame
 	* However when for example using a listview in editor utility widgets a world there is not reliable and an alternative is to use
