ensure that bAutomaticallyRegisterInputOnConstruction is set correctly from its value on the CDO.

When widgets are initalized, they may have been initalized based on their archetype acquired by the native CDO, so we need to make sure that we check the actual BP Generated CDO upon initaliziation of the widget to register any input delegates it may have in its graph.

This fixes a bug where bAutomaticallyRegisterInputOnConstruction would be false, when it is supposed to be true, on nested instances of blueprint widgets. The value would be correct on the CDO for the instance of the widget, but because of the way the UWidgetTree is generated the CDO is not being used to initalize the instance. This is a stop-gap until the widget compilation process can be improved, which may require some changes to *when* bAutomaticallyRegisterInputOnConstruction is actually set on the CDO during BP compilation.

#rnx

diff --git a/Engine/Source/Runtime/UMG/Private/UserWidget.cpp b/Engine/Source/Runtime/UMG/Private/UserWidget.cpp
--- a/Engine/Source/Runtime/UMG/Private/UserWidget.cpp
+++ b/Engine/Source/Runtime/UMG/Private/UserWidget.cpp
@@ -1720,6 +1720,14 @@ void UUserWidget::BindToAnimationEvent(UWidgetAnimation* InAnimation, FWidgetAni
 
 void UUserWidget::NativeOnInitialized()
 {
+	// The widget tree may be constructed from instanced subobjects so we need to set this value
+	// based on this widgets actual CDO.
+	if (!HasAnyFlags(RF_ClassDefaultObject))
+	{
+		const UUserWidget* DefaultWidget = CastChecked<UUserWidget>(GetClass()->GetDefaultObject());
+		bAutomaticallyRegisterInputOnConstruction = DefaultWidget->bAutomaticallyRegisterInputOnConstruction;
+	}
+
 	// Bind any input delegates that may be on this widget to its owning player controller
 	if (bAutomaticallyRegisterInputOnConstruction)
 	{
@@ -1732,7 +1740,7 @@ void UUserWidget::NativeOnInitialized()
 		}
 		else if (!IsEditorUtility())
 		{
-			UE_LOG(LogUMG, Error, TEXT("[%hs] Widget '%s' has bAutomaticallyRegisterInputOnConstruction set to true, but no valid player controller. Input delegates will not work!"), __func__, *GetNameSafe(this));
+			UE_LOG(LogUMG, Error, TEXT("[%hs] Widget '%s' has bAutomaticallyRegisterInputOnConstruction set to true, but has no valid player controller. Input delegates will not work!"), __func__, *GetNameSafe(this));
 		}
 	}
 	
