Exposes ViewProjectionMatrix of a UMG viewport to BP

#rnx

diff --git a/Engine/Source/Runtime/UMG/Private/Components/Viewport.cpp b/Engine/Source/Runtime/UMG/Private/Components/Viewport.cpp
--- a/Engine/Source/Runtime/UMG/Private/Components/Viewport.cpp
+++ b/Engine/Source/Runtime/UMG/Private/Components/Viewport.cpp
@@ -172,6 +172,7 @@ void FUMGViewportClient::Draw(FViewport* InViewport, FCanvas* Canvas)
 	ViewFamily.SetScreenPercentageInterface(new FLegacyScreenPercentageDriver(
 		ViewFamily, GlobalResolutionFraction));
 
+	CachedViewProjectionMatrix = View->ViewMatrices.GetViewProjectionMatrix();
 	//SetupViewForRendering(ViewFamily, *View);
 
 	FSlateRect SafeFrame;
@@ -257,6 +258,10 @@ float FUMGViewportClient::GetOrthoUnitsPerPixel(const FViewport* InViewport) con
 	return ( GetOrthoZoom() / ( SizeX * 15.f ) )/* * ComputeOrthoZoomFactor(SizeX)*/;
 }
 
+FMatrix FUMGViewportClient::GetViewProjectionMatrix() const
+{
+	return CachedViewProjectionMatrix;
+}
 
 FSceneView* FUMGViewportClient::CalcSceneView(FSceneViewFamily* ViewFamily)
 {
@@ -506,6 +511,17 @@ void UViewport::SetShowFlag(FString InShowFlagName, bool InValue)
 	}
 }
 
+FMatrix UViewport::GetViewProjectionMatrix() const
+{
+	FMatrix ProjectionMatrix = FMatrix::Identity;
+	if (ViewportWidget.IsValid())
+	{
+		check(ViewportWidget->ViewportClient.IsValid());
+		ProjectionMatrix = ViewportWidget->ViewportClient->GetViewProjectionMatrix();
+	}
+	return ProjectionMatrix;
+}
+
 void UViewport::SetEnableAdvancedFeatures(bool InEnableAdvancedFeatures)
 {
 	ShowFlags.DisableAdvancedFeatures();
diff --git a/Engine/Source/Runtime/UMG/Public/Components/Viewport.h b/Engine/Source/Runtime/UMG/Public/Components/Viewport.h
--- a/Engine/Source/Runtime/UMG/Public/Components/Viewport.h
+++ b/Engine/Source/Runtime/UMG/Public/Components/Viewport.h
@@ -202,11 +202,13 @@ public:
 		EngineShowFlags = InEngineShowFlags;
 	}
 
+	FMatrix GetViewProjectionMatrix() const;
+
 protected:
 
 	/** The scene used for the viewport. Owned externally */
 	FPreviewScene* PreviewScene;
-
+	
 	FMinimalViewInfo ViewInfo;
 
 	FLinearColor BackgroundColor;
@@ -221,6 +223,8 @@ protected:
 
 	/** A set of flags that determines visibility for various scene elements. */
 	FEngineShowFlags EngineShowFlags;
+
+	FMatrix CachedViewProjectionMatrix;
 };
 
 /**
@@ -269,6 +273,10 @@ class UViewport : public UContentWidget
 	UFUNCTION(BlueprintCallable, Category = "Viewport")
 	UMG_API void SetShowFlag(FString InShowFlagName, bool InValue);
 
+	UFUNCTION(BlueprintCallable, Category = "Viewport")
+	UMG_API FMatrix GetViewProjectionMatrix() const;
+	
+
 	// UWidget interface
 	UMG_API virtual void SynchronizeProperties() override;
 	// End of UWidget interface
