Menu Anchor Widget: Added an option to hide the menu background when using the application stack.

diff --git a/Engine/Source/Runtime/Slate/Private/Widgets/Input/SMenuAnchor.cpp b/Engine/Source/Runtime/Slate/Private/Widgets/Input/SMenuAnchor.cpp
--- a/Engine/Source/Runtime/Slate/Private/Widgets/Input/SMenuAnchor.cpp
+++ b/Engine/Source/Runtime/Slate/Private/Widgets/Input/SMenuAnchor.cpp
@@ -99,6 +99,7 @@ void SMenuAnchor::Construct( const FArguments& InArgs )
 	Method                                 = InArgs._Method;
 	bShouldDeferPaintingAfterWindowContent = InArgs._ShouldDeferPaintingAfterWindowContent;
 	bUseApplicationMenuStack               = InArgs._UseApplicationMenuStack;
+	bShowMenuBackground					   = InArgs._ShowMenuBackground;
 	bIsCollapsedByParent                   = InArgs._IsCollapsedByParent;
 	bApplyWidgetStyleToMenu                = InArgs._ApplyWidgetStyleToMenu;
 }
diff --git a/Engine/Source/Runtime/Slate/Public/Widgets/Input/SMenuAnchor.h b/Engine/Source/Runtime/Slate/Public/Widgets/Input/SMenuAnchor.h
--- a/Engine/Source/Runtime/Slate/Public/Widgets/Input/SMenuAnchor.h
+++ b/Engine/Source/Runtime/Slate/Public/Widgets/Input/SMenuAnchor.h
@@ -41,6 +41,7 @@ public:
 		, _Method()
 		, _ShouldDeferPaintingAfterWindowContent(true)
 		, _UseApplicationMenuStack(true)
+		, _ShowMenuBackground(true)
 		, _IsCollapsedByParent(false)
 		, _ApplyWidgetStyleToMenu(true)
 		{}
@@ -64,6 +65,8 @@ public:
 		SLATE_ARGUMENT(bool, ShouldDeferPaintingAfterWindowContent)
 
 		SLATE_ARGUMENT(bool, UseApplicationMenuStack)
+
+		SLATE_ARGUMENT(bool, ShowMenuBackground)
 		
 		/** True if this menu anchor should be collapsed when its parent receives focus, false (default) otherwise */
 		SLATE_ARGUMENT(bool, IsCollapsedByParent)
diff --git a/Engine/Source/Runtime/UMG/Private/Components/MenuAnchor.cpp b/Engine/Source/Runtime/UMG/Private/Components/MenuAnchor.cpp
--- a/Engine/Source/Runtime/UMG/Private/Components/MenuAnchor.cpp
+++ b/Engine/Source/Runtime/UMG/Private/Components/MenuAnchor.cpp
@@ -17,6 +17,7 @@ UMenuAnchor::UMenuAnchor(const FObjectInitializer& ObjectInitializer)
 	: Super(ObjectInitializer)
 	, ShouldDeferPaintingAfterWindowContent(true)
 	, UseApplicationMenuStack(true)
+	, ShowMenuBackground(true)
 {
 PRAGMA_DISABLE_DEPRECATION_WARNINGS
 	Placement = MenuPlacement_ComboBox;
@@ -40,7 +41,8 @@ PRAGMA_DISABLE_DEPRECATION_WARNINGS
 		.OnGetMenuContent(BIND_UOBJECT_DELEGATE(FOnGetContent, HandleGetMenuContent))
 		.OnMenuOpenChanged(BIND_UOBJECT_DELEGATE(FOnIsOpenChanged, HandleMenuOpenChanged))
 		.ShouldDeferPaintingAfterWindowContent(ShouldDeferPaintingAfterWindowContent)
-		.UseApplicationMenuStack(UseApplicationMenuStack);
+		.UseApplicationMenuStack(UseApplicationMenuStack)
+		.ShowMenuBackground(ShowMenuBackground);
 PRAGMA_ENABLE_DEPRECATION_WARNINGS
 	if ( GetChildrenCount() > 0 )
 	{
@@ -189,6 +191,11 @@ bool UMenuAnchor::IsUseApplicationMenuStack() const
 	return UseApplicationMenuStack;
 }
 
+bool UMenuAnchor::IsShowMenuBackground() const
+{
+	return ShowMenuBackground;
+}
+
 void UMenuAnchor::InitShouldDeferPaintingAfterWindowContent(bool InShouldDeferPaintingAfterWindowContent)
 {
 	ensureMsgf(!MyMenuAnchor.IsValid(), TEXT("The widget is already created."));
@@ -201,6 +208,12 @@ void UMenuAnchor::InitUseApplicationMenuStack(bool InUseApplicationMenuStack)
 	UseApplicationMenuStack = InUseApplicationMenuStack;
 }
 
+void UMenuAnchor::InitShowMenuBackground(bool InShowMenuBackground)
+{
+	ensureMsgf(!MyMenuAnchor.IsValid(), TEXT("The widget is already created."));
+	ShowMenuBackground = InShowMenuBackground;
+}
+
 PRAGMA_ENABLE_DEPRECATION_WARNINGS
 
 bool UMenuAnchor::ShouldOpenDueToClick() const
diff --git a/Engine/Source/Runtime/UMG/Public/Components/MenuAnchor.h b/Engine/Source/Runtime/UMG/Public/Components/MenuAnchor.h
--- a/Engine/Source/Runtime/UMG/Public/Components/MenuAnchor.h
+++ b/Engine/Source/Runtime/UMG/Public/Components/MenuAnchor.h
@@ -64,6 +64,10 @@ public:
 	UPROPERTY(EditAnywhere, BlueprintReadOnly, Getter = "IsUseApplicationMenuStack", AdvancedDisplay, Category = "Menu Anchor")
 	bool UseApplicationMenuStack;
 
+	/** For menus using the application menu stack, should the window background be visible? */
+	UPROPERTY(EditAnywhere, BlueprintReadOnly, Getter = "IsShowMenuBackground", AdvancedDisplay, Category = "Menu Anchor",  meta = (EditCondition = "UseApplicationMenuStack", EditConditionHides, AllowPrivateAccess = "true", DisplayAfter = "UseApplicationMenuStack"))
+	bool ShowMenuBackground;
+
 public:
 	/** Called when the opened state of the menu changes */
 	UPROPERTY(BlueprintAssignable, Category="Menu Anchor|Event")
@@ -86,6 +90,8 @@ public:
 
 	UMG_API bool IsUseApplicationMenuStack() const;
 
+	UMG_API bool IsShowMenuBackground() const;
+
 public:
 
 	/**
@@ -147,6 +153,8 @@ protected:
 	// Initialize UseApplicationMenuStack in the constructor before the SWidget is constructed.
 	UMG_API void InitUseApplicationMenuStack(bool InUseApplicationMenuStack);
 
+	// Initialize ShowMenuBackground in the constructor before the SWidget is constructed.
+	UMG_API void InitShowMenuBackground(bool InShowMenuBackground);
 
 protected:
 	UMG_API TSharedRef<SWidget> HandleGetMenuContent();
