Fix to UI Components experimental version.
- Fix instanciation to of WidgetComponent on the preview widget
- Adding PreConstruct function to Components.
- Correction to MigrateFromChain to support component types

This is disabled by default, to enable it, set Widget.ExposeUIComponents CVar to true.
Known issues: Remove Components not exposed, Detail panel do not refresh when adding/removing Component, no support of copy/cut/paste.

diff --git a/Engine/Source/Editor/UMGEditor/Private/Customizations/UIComponentCustomizationExtender.cpp b/Engine/Source/Editor/UMGEditor/Private/Customizations/UIComponentCustomizationExtender.cpp
--- a/Engine/Source/Editor/UMGEditor/Private/Customizations/UIComponentCustomizationExtender.cpp
+++ b/Engine/Source/Editor/UMGEditor/Private/Customizations/UIComponentCustomizationExtender.cpp
@@ -5,12 +5,10 @@
 #include "DetailCategoryBuilder.h"
 #include "DetailLayoutBuilder.h"
 #include "Extensions/UIComponent.h"
-#include "Extensions/UIComponentContainer.h"
-#include "PropertyCustomizationHelpers.h"
+#include "Extensions/UIComponentUserWidgetExtension.h"
 #include "PropertyHandle.h"
 #include "WidgetBlueprint.h"
 #include "WidgetBlueprintEditor.h"
-#include "UIComponentUtils.h"
 
 #define LOCTEXT_NAMESPACE "UIComponentCustomizationExtender"
 
@@ -22,27 +20,32 @@ TSharedPtr<FUIComponentCustomizationExtender> FUIComponentCustomizationExtender:
 
 void FUIComponentCustomizationExtender::CustomizeDetails(IDetailLayoutBuilder& InDetailLayout, const TArrayView<UWidget*> InWidgets, const TSharedRef<FWidgetBlueprintEditor>& InWidgetBlueprintEditor)
 {
-	if (InWidgets.Num() == 1)
+	const UUserWidget* PreviewUserWidget = InWidgetBlueprintEditor->GetPreview();
+	if (InWidgets.Num() != 1 || !PreviewUserWidget)
 	{
-		Widget = InWidgets[0];
-		WidgetBlueprintEditor = InWidgetBlueprintEditor;
+		return;
+	}
 
-		if (UUIComponentContainer* UIComponentContainer = FUIComponentUtils::GetUIComponentContainerFromWidgetBlueprint(WidgetBlueprintEditor.Pin()->GetWidgetBlueprintObj()))
-		{
-			TArray<UUIComponent*> ComponentsOnWidget = UIComponentContainer->GetExtensionsFor(Widget->GetFName());
+	Widget = InWidgets[0];
+	WidgetBlueprintEditor = InWidgetBlueprintEditor;
 
-			for (int32 Index = ComponentsOnWidget.Num() - 1; Index >= 0; Index--)
+	UWidgetBlueprint* WidgetBlueprint = InWidgetBlueprintEditor->GetWidgetBlueprintObj();
+	
+	// We use the UserWidgetExtension on the preview and it will be migrated to the WBP extension in MigrateFromChain
+	if (const UUIComponentUserWidgetExtension* Extension = PreviewUserWidget->GetExtension<UUIComponentUserWidgetExtension>())
+	{
+		TArray<UUIComponent*> ComponentsOnWidget = Extension->GetComponentsFor(Widget.Get());
+		for (int32 Index = ComponentsOnWidget.Num() - 1; Index >= 0; Index--)
+		{
+			if (UUIComponent* Component = ComponentsOnWidget[Index])
 			{
-				if (UUIComponent* Component = ComponentsOnWidget[Index])
-				{
-					IDetailCategoryBuilder& ComponentCategory = InDetailLayout.EditCategory(Component->GetFName(), FText::GetEmpty(), ECategoryPriority::Important);
-
-					IDetailPropertyRow* PropertyRow = ComponentCategory.AddExternalObjects({ Component }, EPropertyLocation::Default,
-						FAddPropertyParams()
-						.CreateCategoryNodes(false)
-						.AllowChildren(true)
-						.HideRootObjectNode(true));
-				}
+				IDetailCategoryBuilder& ComponentCategory = InDetailLayout.EditCategory(Component->GetFName(), FText::GetEmpty(), ECategoryPriority::Important);
+
+				IDetailPropertyRow* PropertyRow = ComponentCategory.AddExternalObjects({ Component }, EPropertyLocation::Default,
+					FAddPropertyParams()
+					.CreateCategoryNodes(false)
+					.AllowChildren(true)
+					.HideRootObjectNode(true));
 			}
 		}
 	}
diff --git a/Engine/Source/Editor/UMGEditor/Private/Customizations/UIComponentCustomizationExtender.h b/Engine/Source/Editor/UMGEditor/Private/Customizations/UIComponentCustomizationExtender.h
--- a/Engine/Source/Editor/UMGEditor/Private/Customizations/UIComponentCustomizationExtender.h
+++ b/Engine/Source/Editor/UMGEditor/Private/Customizations/UIComponentCustomizationExtender.h
@@ -12,7 +12,7 @@ class FUIComponentCustomizationExtender : public IBlueprintWidgetCustomizationEx
 {
 public:
 	static TSharedPtr<FUIComponentCustomizationExtender> MakeInstance();
-	virtual void CustomizeDetails(IDetailLayoutBuilder& InDetailLayout, const TArrayView<UWidget*> InWidgets, const TSharedRef<FWidgetBlueprintEditor>& InWidgetBlueprintEditor);
+	virtual void CustomizeDetails(IDetailLayoutBuilder& InDetailLayout, const TArrayView<UWidget*> InWidgets, const TSharedRef<FWidgetBlueprintEditor>& InWidgetBlueprintEditor) override;
 
 private:
 	/** The selected widget in the details panel. */
diff --git a/Engine/Source/Editor/UMGEditor/Private/Designer/SDesignerView.cpp b/Engine/Source/Editor/UMGEditor/Private/Designer/SDesignerView.cpp
--- a/Engine/Source/Editor/UMGEditor/Private/Designer/SDesignerView.cpp
+++ b/Engine/Source/Editor/UMGEditor/Private/Designer/SDesignerView.cpp
@@ -2297,6 +2297,12 @@ void SDesignerView::UpdatePreviewWidget(bool bForceUpdate)
 		PreviewWidget = LatestPreviewWidget;
 		if ( PreviewWidget )
 		{
+			// Notify all designer extensions that a new preview was created
+			for (const TSharedRef<FDesignerExtension>& Ext : DesignerExtensions)
+			{
+				Ext->PreviewContentCreated(PreviewWidget);
+			}
+		
 			TSharedRef<SWidget> NewPreviewSlateWidget = PreviewWidget->TakeWidget();
 			NewPreviewSlateWidget->SlatePrepass(PreviewSizeConstraint->GetCachedGeometry().Scale);
 
diff --git a/Engine/Source/Editor/UMGEditor/Private/Details/SWidgetDetailsView.cpp b/Engine/Source/Editor/UMGEditor/Private/Details/SWidgetDetailsView.cpp
--- a/Engine/Source/Editor/UMGEditor/Private/Details/SWidgetDetailsView.cpp
+++ b/Engine/Source/Editor/UMGEditor/Private/Details/SWidgetDetailsView.cpp
@@ -39,6 +39,7 @@
 #include "Kismet2/SClassPickerDialog.h"
 #include "SPositiveActionButton.h"
 #include "UIComponentUtils.h"
+#include "UIComponentWidgetBlueprintExtension.h"
 
 #define LOCTEXT_NAMESPACE "UMG"
 
@@ -653,21 +654,23 @@ void SWidgetDetailsView::HandleIsVariableChanged(ECheckBoxState CheckState)
 
 FReply SWidgetDetailsView::OnAddComponentButtonClicked(TSharedPtr<FWidgetBlueprintEditor> InBlueprintEditor)
 {
-	UClass* ChosenClass = nullptr;
-	const FText TitleText = LOCTEXT("AddExtension", "Pick a component to add to the widget.");
-	const bool bPressedOk = SClassPickerDialog::PickClass(TitleText, FUIComponentUtils::CreateClassViewerInitializationOptions(), ChosenClass, UUIComponent::StaticClass());
-	UUserWidget* UserWidget = InBlueprintEditor->GetBlueprintObj()->GeneratedClass->GetDefaultObject<UUserWidget>();
-
-	if (ChosenClass)
+	if (SelectedObjects.Num() > 0)
 	{
-		if (UUIComponentContainer* ComponentsContainer = FUIComponentUtils::GetOrCreateComponentsContainerForUserWidget(UserWidget))
+		UClass* ChosenClass = nullptr;
+		const FText TitleText = LOCTEXT("AddExtension", "Pick a component to add to the widget.");
+		const bool bPressedOk = SClassPickerDialog::PickClass(TitleText, FUIComponentUtils::CreateClassViewerInitializationOptions(), ChosenClass, UUIComponent::StaticClass());
+		if (ChosenClass)
 		{
-			UUIComponent* Component = FUIComponentUtils::CreateUIComponent(ChosenClass, UserWidget);
-			UWidget* SelectedWidget = Cast<UWidget>(SelectedObjects[0].Get());
-
-			ComponentsContainer->Modify();
-			ComponentsContainer->AddComponent(SelectedWidget->GetFName(), Component);
-			FBlueprintEditorUtils::MarkBlueprintAsStructurallyModified(InBlueprintEditor->GetBlueprintObj());
+			UWidgetBlueprint* WidgetBlueprint = InBlueprintEditor->GetWidgetBlueprintObj();
+			const UWidget* SelectedWidget = Cast<UWidget>(SelectedObjects[0].Get());
+			if (WidgetBlueprint && SelectedWidget)
+			{
+				if (UUIComponentWidgetBlueprintExtension* Extension = UUIComponentWidgetBlueprintExtension::RequestExtension<UUIComponentWidgetBlueprintExtension>(WidgetBlueprint))
+				{
+					Extension->AddComponent(InBlueprintEditor.ToSharedRef(), ChosenClass, SelectedWidget->GetFName());
+					FBlueprintEditorUtils::MarkBlueprintAsStructurallyModified(WidgetBlueprint);
+				}
+			}
 		}
 	}
 	return FReply::Handled();
diff --git a/Engine/Source/Editor/UMGEditor/Private/Extensions/UIComponentContainerDesignerExtension.cpp b/Engine/Source/Editor/UMGEditor/Private/Extensions/UIComponentContainerDesignerExtension.cpp
--- a/Engine/Source/Editor/UMGEditor/Private/Extensions/UIComponentContainerDesignerExtension.cpp
+++ b/Engine/Source/Editor/UMGEditor/Private/Extensions/UIComponentContainerDesignerExtension.cpp
@@ -2,19 +2,48 @@
 
 #include "Extensions/UIComponentContainerDesignerExtension.h"
 #include "Blueprint/UserWidget.h"
+#include "Slate/SObjectWidget.h"
 #include "Extensions/UIComponentContainer.h"
+#include "Extensions/UIComponentWidgetBlueprintGeneratedClassExtension.h"
+
 #include "WidgetBlueprint.h"
 #include "UIComponentUtils.h"
+#include "UIComponentWidgetBlueprintExtension.h"
 
 TSharedRef<FDesignerExtension> FUIComponentContainerDesignerExtensionFactory::CreateDesignerExtension() const
 {
 	return StaticCastSharedRef<FUIComponentContainerDesignerExtension>(MakeShared<FUIComponentContainerDesignerExtension>());
 }
 
+/** Called every time the content of the designer changed. */
+void FUIComponentContainerDesignerExtension::PreviewContentCreated(UUserWidget* PreviewWidget)
+{
+	const UWidgetBlueprint* WidgetBlueprint = Blueprint.Get();
+	if (WidgetBlueprint == nullptr || PreviewWidget == nullptr)
+	{
+		return;
+	}
+	
+	if (UUIComponentWidgetBlueprintExtension* Extension = UUIComponentWidgetBlueprintExtension::GetExtension<UUIComponentWidgetBlueprintExtension>(WidgetBlueprint))
+	{
+		UUIComponentUserWidgetExtension* UserWidgetExtension = Extension->GetOrCreateExtension(PreviewWidget);
+	}
+}
+
 void FUIComponentContainerDesignerExtension::PreviewContentChanged(TSharedRef<SWidget> NewContent)
 {
-	if (UUIComponentContainer* UIComponentContainer = FUIComponentUtils::GetUIComponentContainerFromWidgetBlueprint(Blueprint.Get()))
+	const UWidgetBlueprint* WidgetBlueprint = Blueprint.Get();
+	if (NewContent == SNullWidget::NullWidget || WidgetBlueprint == nullptr)
+	{
+		return;
+	}	
+	
+	if (UUIComponentWidgetBlueprintExtension* Extension = UUIComponentWidgetBlueprintExtension::GetExtension<UUIComponentWidgetBlueprintExtension>(WidgetBlueprint))
 	{
-		UIComponentContainer->OnPreviewContentChanged(NewContent);
+		const SObjectWidget* ObjectWidget = StaticCastSharedPtr<SObjectWidget>(NewContent.ToSharedPtr()).Get();
+		if (UUserWidget* PreviewWidget = ObjectWidget ? ObjectWidget->GetWidgetObject() : nullptr)
+		{
+			Extension->VerifyContainer(PreviewWidget);
+		}
 	}
 }
diff --git a/Engine/Source/Editor/UMGEditor/Private/Extensions/UIComponentContainerDesignerExtension.h b/Engine/Source/Editor/UMGEditor/Private/Extensions/UIComponentContainerDesignerExtension.h
--- a/Engine/Source/Editor/UMGEditor/Private/Extensions/UIComponentContainerDesignerExtension.h
+++ b/Engine/Source/Editor/UMGEditor/Private/Extensions/UIComponentContainerDesignerExtension.h
@@ -17,5 +17,6 @@ public:
 class FUIComponentContainerDesignerExtension : public FDesignerExtension
 {
 public:
+	virtual void PreviewContentCreated(UUserWidget* PreviewWidget) override;
 	virtual void PreviewContentChanged(TSharedRef<SWidget> NewContent) override;
 };
diff --git a/Engine/Source/Editor/UMGEditor/Private/UIComponentUtils.cpp b/Engine/Source/Editor/UMGEditor/Private/UIComponentUtils.cpp
--- a/Engine/Source/Editor/UMGEditor/Private/UIComponentUtils.cpp
+++ b/Engine/Source/Editor/UMGEditor/Private/UIComponentUtils.cpp
@@ -5,36 +5,10 @@
 #include "Blueprint/UserWidget.h"
 #include "ClassViewerModule.h"
 #include "CoreGlobals.h"
-#include "Components/Widget.h"
+#include "UIComponentWidgetBlueprintExtension.h"
 #include "Extensions/UIComponent.h"
-#include "Extensions/UIComponentContainer.h"
-#include "WidgetBlueprint.h"
-
-UUIComponent* FUIComponentUtils::CreateUIComponent(TSubclassOf<UUIComponent> ComponentClass, UUserWidget* Outer)
-{
-	ensure(Outer);
-	if (UUIComponent* NewComponent = NewObject<UUIComponent>(Outer, ComponentClass))
-	{
-		NewComponent->SetFlags(RF_Transactional);
-		return NewComponent;
-	}
-	return nullptr;
-}
-
-UUIComponentContainer* FUIComponentUtils::GetOrCreateComponentsContainerForUserWidget(UUserWidget* UserWidget)
-{
-	check(UserWidget);
-	if (UUIComponentContainer* ExistingComponentsContainer = UserWidget->GetExtension<UUIComponentContainer>())
-	{
-		return ExistingComponentsContainer;
-	}
-	else if (UUIComponentContainer* AddedComponentContainer = UserWidget->AddExtension<UUIComponentContainer>())
-	{
-		AddedComponentContainer->SetFlags(RF_Transactional);
-		return AddedComponentContainer;
-	}
-	return nullptr;
-}
+#include "WidgetBlueprintEditor.h"
+#include "Extensions/UIComponentUserWidgetExtension.h"
 
 FClassViewerInitializationOptions FUIComponentUtils::CreateClassViewerInitializationOptions()
 {
@@ -49,16 +23,21 @@ FClassViewerInitializationOptions FUIComponentUtils::CreateClassViewerInitializa
 	return Options;
 }
 
-UUIComponentContainer* FUIComponentUtils::GetUIComponentContainerFromWidgetBlueprint(UWidgetBlueprint* WidgetBlueprint)
+
+void FUIComponentUtils::RenameWidgetName(TSharedRef<FWidgetBlueprintEditor> BlueprintEditor, UWidgetBlueprint* WidgetBlueprint, const FName& OldVarName, const FName& NewVarName)
 {
-	if (WidgetBlueprint)
+	if (UUIComponentWidgetBlueprintExtension* Extension = UWidgetBlueprintExtension::GetExtension<UUIComponentWidgetBlueprintExtension>(WidgetBlueprint))
 	{
-		if (UUserWidget* CDO = WidgetBlueprint->GeneratedClass->GetDefaultObject<UUserWidget>())
+		Extension->RenameWidget(OldVarName, NewVarName);
+	}
+
+	if (const UUserWidget* PreviewWidget = BlueprintEditor->GetPreview())
+	{
+		if(UUIComponentUserWidgetExtension* UserWidgetExtension = PreviewWidget->GetExtension<UUIComponentUserWidgetExtension>())
 		{
-			return CDO->GetExtension<UUIComponentContainer>();
+			UserWidgetExtension->RenameWidget(OldVarName, NewVarName);
 		}
-	}
-	return nullptr;
+	}	
 }
 
 bool FUIComponentUtils::FUIComponentClassFilter::IsClassAllowed(const FClassViewerInitializationOptions& InInitOptions, const UClass* InClass, TSharedRef< FClassViewerFilterFuncs > InFilterFuncs)
diff --git a/Engine/Source/Editor/UMGEditor/Private/UIComponentWidgetBlueprintExtension.cpp b/Engine/Source/Editor/UMGEditor/Private/UIComponentWidgetBlueprintExtension.cpp
new file mode 100644
--- /dev/null
+++ b/Engine/Source/Editor/UMGEditor/Private/UIComponentWidgetBlueprintExtension.cpp
@@ -0,0 +1,214 @@
+﻿// Copyright Epic Games, Inc. All Rights Reserved.
+
+#include "UIComponentWidgetBlueprintExtension.h"
+
+#include "Extensions/UIComponent.h"
+#include "Extensions/UIComponentContainer.h"
+#include "Extensions/UIComponentWidgetBlueprintGeneratedClassExtension.h"
+#include "Extensions/UIComponentUserWidgetExtension.h"
+
+#include "WidgetBlueprintEditor.h"
+#include "Blueprint/WidgetTree.h"
+
+UUIComponentWidgetBlueprintExtension::UUIComponentWidgetBlueprintExtension(const FObjectInitializer& ObjectInitializer)
+	: Super(ObjectInitializer)
+{
+	ComponentContainer = CreateDefaultSubobject<UUIComponentContainer>(TEXT("ComponentContainer"));
+	ComponentContainer->SetFlags(RF_Transactional | RF_ArchetypeObject);
+}
+
+UUIComponentContainer* UUIComponentWidgetBlueprintExtension::DuplicateContainer(UObject* Outer) const
+{
+	ensure(ComponentContainer->HasAnyFlags(RF_ArchetypeObject));
+			
+	FObjectInstancingGraph ObjectInstancingGraph;
+	UUIComponentContainer* NewContainer = NewObject<UUIComponentContainer>(Outer, ComponentContainer->GetClass(), NAME_None, RF_Transactional, ComponentContainer, false);
+	return NewContainer;
+}
+
+void UUIComponentWidgetBlueprintExtension::RemoveComponent(TSharedRef<FWidgetBlueprintEditor> Editor, UClass* ComponentClass, FName OwnerName)
+{
+	// Remove the Component
+	if(ensure(ComponentContainer))
+	{
+		ComponentContainer->RemoveAllComponentsOfType(ComponentClass, OwnerName);
+	}
+
+	//Also remove it from the PreviewWidget
+	if (const UUserWidget* PreviewWidget = Editor->GetPreview())
+	{
+		if(UUIComponentUserWidgetExtension* UserWidgetExtension = PreviewWidget->GetExtension<UUIComponentUserWidgetExtension>())
+		{
+			UserWidgetExtension->RemoveComponent(ComponentClass, OwnerName);
+		}
+	}	
+}
+
+TArray<UUIComponent*> UUIComponentWidgetBlueprintExtension::GetComponentsFor(UWidget* Target) const
+{
+	TArray<UUIComponent*> Components;
+	if(ensure(Target))
+	{
+		ComponentContainer->ForEachComponentTarget([Target, &Components](const FUIComponentTarget& ComponentTarget){
+			if(ComponentTarget.GetTargetName() == Target->GetFName())			
+			{
+				Components.Push(ComponentTarget.GetComponent());
+			}
+		});
+	}
+	return Components;	
+}
+
+UUIComponent* UUIComponentWidgetBlueprintExtension::GetComponent(const UClass* ComponentClass, FName OwnerName) const
+{
+	if (ComponentContainer)
+	{
+		return ComponentContainer->GetComponent(ComponentClass, OwnerName);
+	}
+	return nullptr;
+}
+
+void UUIComponentWidgetBlueprintExtension::RenameWidget(const FName& OldVarName, const FName& NewVarName)
+{
+	if (ensure(ComponentContainer))
+	{
+		ComponentContainer->RenameWidget(OldVarName, NewVarName);
+	}
+}
+
+bool UUIComponentWidgetBlueprintExtension::VerifyContainer(UUserWidget* UserWidget) const
+{
+	const UUIComponentUserWidgetExtension* UserWidgetExtension = UserWidget->GetExtension<UUIComponentUserWidgetExtension>();
+	if(UserWidgetExtension)
+	{
+		TArray<FUIComponentTarget> Components;
+		ComponentContainer->ForEachComponentTarget([&Components](const FUIComponentTarget& ComponentTarget){
+			Components.Push(ComponentTarget);
+		});
+
+		while(Components.Num() > 0)
+		{
+			FUIComponentTarget& Target = Components[Components.Num()-1];
+			if(!UserWidgetExtension->GetComponent(Target.GetComponent()->GetClass(), Target.GetTargetName()))
+			{
+				return false;
+			}
+			else
+			{
+				Components.RemoveAt(Components.Num()-1);
+			}
+		}
+		return true;
+	}
+	return false;	
+}
+
+UUIComponentUserWidgetExtension* UUIComponentWidgetBlueprintExtension::GetOrCreateExtension(UUserWidget* PreviewWidget)
+{
+	UUIComponentUserWidgetExtension* UserWidgetExtension = PreviewWidget->GetExtension<UUIComponentUserWidgetExtension>();
+	// If the extension do not exist, we create it which will create a copy of the component we just added.
+	if (!UserWidgetExtension)
+	{
+		UserWidgetExtension = PreviewWidget->AddExtension<UUIComponentUserWidgetExtension>();
+		UserWidgetExtension->InitializeContainer(DuplicateContainer(UserWidgetExtension));
+	}
+	else
+	{
+		VerifyContainer(PreviewWidget);
+	}
+	
+	return UserWidgetExtension;
+}
+
+void UUIComponentWidgetBlueprintExtension::AddComponentToPreview(UUserWidget* PreviewWidget, FName OwnerName, UUIComponent* ArchetypeComponent)
+{	
+	UUIComponentUserWidgetExtension* UserWidgetExtension = GetOrCreateExtension(PreviewWidget);
+	// If the extension do not exist, we create it which will create a copy of the component we just added.
+	if (UserWidgetExtension)
+	{
+		UserWidgetExtension->CreateAndAddComponent(ArchetypeComponent, OwnerName);
+	}	
+}
+
+void UUIComponentWidgetBlueprintExtension::AddComponent(TSharedRef<FWidgetBlueprintEditor> Editor, UClass* ComponentClass, FName OwnerName)
+{
+	// Add the Component to the WidgetBlueprint
+	check(ComponentContainer);
+
+	// If we already have a component of that tyype on the Widget early out
+	if (ComponentContainer->GetComponent(ComponentClass, OwnerName))
+	{
+		return;
+	}
+
+	if (UUIComponent* NewComponent = NewObject<UUIComponent>(ComponentContainer, ComponentClass, NAME_None, RF_ArchetypeObject))
+	{
+		Modify();
+		ComponentContainer->AddComponent(OwnerName, NewComponent);
+		AddComponentToPreview(Editor->GetPreview(), OwnerName, NewComponent);
+	}
+}
+
+void UUIComponentWidgetBlueprintExtension::HandleBeginCompilation(FWidgetBlueprintCompilerContext& InCreationContext)
+{
+	Super::HandleBeginCompilation(InCreationContext);
+
+	CompilerContext = &InCreationContext;
+}
+
+void UUIComponentWidgetBlueprintExtension::HandleCleanAndSanitizeClass(UWidgetBlueprintGeneratedClass* ClassToClean, UObject* OldCDO)
+{
+	// Here we handle Widget Deletion. Go through all Widgets referenced by Components and remove them if they no longer exist.
+	const UWidgetTree* WidgetTree = GetWidgetBlueprint()->WidgetTree; 	
+	
+	if (ComponentContainer && WidgetTree)
+	{	
+		ComponentContainer->CleanupUIComponents(WidgetTree);
+	}
+}
+
+void UUIComponentWidgetBlueprintExtension::HandleFinishCompilingClass(UWidgetBlueprintGeneratedClass* Class)
+{
+	Super::HandleFinishCompilingClass(Class);
+
+	// If we do not have Components, do not add the extension to the Generated Class
+	if(ensure(CompilerContext) && !ComponentContainer->IsEmpty())
+	{
+		if(UWidgetBlueprintGeneratedClass* NewWidgetGeneratedClass = Cast<UWidgetBlueprintGeneratedClass>(CompilerContext->NewClass))
+		{
+			UUIComponentWidgetBlueprintGeneratedClassExtension* NewExtension = NewObject<UUIComponentWidgetBlueprintGeneratedClassExtension>(NewWidgetGeneratedClass);
+			CompilerContext->AddExtension(NewWidgetGeneratedClass, NewExtension);
+			NewExtension->InitializeContainer(DuplicateContainer(NewExtension));
+		}		
+	}
+}
+
+bool UUIComponentWidgetBlueprintExtension::HandleValidateGeneratedClass(UWidgetBlueprintGeneratedClass* Class)
+{
+	const UWidgetBlueprintGeneratedClass* GeneratedClass = Class;
+	if ( !ensure(GeneratedClass) )
+	{
+		return false;
+	}
+	const UWidgetBlueprint* Blueprint = GetWidgetBlueprint();
+	if ( !ensure(Blueprint) )
+	{
+		return false;
+	}
+
+	// Validate with WidgetTree if all widget exists.
+	if(UWidgetBlueprintGeneratedClass* NewWidgetGeneratedClass = Cast<UWidgetBlueprintGeneratedClass>(CompilerContext->NewClass))
+	{
+		if (const UUIComponentWidgetBlueprintGeneratedClassExtension* Extension = NewWidgetGeneratedClass->GetExtension<UUIComponentWidgetBlueprintGeneratedClassExtension>())
+		{
+			return Extension->VerifyAllWidgetsExists(NewWidgetGeneratedClass->GetWidgetTreeArchetype());
+		}
+	}
+	
+	return true;
+}
+
+void UUIComponentWidgetBlueprintExtension::HandleEndCompilation()
+{
+	CompilerContext = nullptr;
+}
\ No newline at end of file
diff --git a/Engine/Source/Editor/UMGEditor/Private/UIComponentWidgetBlueprintExtension.h b/Engine/Source/Editor/UMGEditor/Private/UIComponentWidgetBlueprintExtension.h
new file mode 100644
--- /dev/null
+++ b/Engine/Source/Editor/UMGEditor/Private/UIComponentWidgetBlueprintExtension.h
@@ -0,0 +1,59 @@
+﻿// Copyright Epic Games, Inc. All Rights Reserved.
+
+#pragma once
+
+#include "CoreMinimal.h"
+#include "UObject/ObjectMacros.h"
+#include "UObject/Object.h"
+#include "WidgetBlueprintExtension.h"
+
+#include "Extensions/UIComponent.h"
+
+#include "UIComponentWidgetBlueprintExtension.generated.h"
+
+class UUIComponentUserWidgetExtension;
+class UUIComponentContainer;
+class FWidgetBlueprintEditor;
+
+UCLASS()
+/***
+ * Extension to the Widget Blueprint that will save the Components and will be used in the editor for compilation.
+ */
+class UMGEDITOR_API UUIComponentWidgetBlueprintExtension : public UWidgetBlueprintExtension
+{
+	GENERATED_BODY()
+	
+	explicit UUIComponentWidgetBlueprintExtension(const FObjectInitializer& ObjectInitializer = FObjectInitializer());
+
+public:
+	void AddComponent(TSharedRef<FWidgetBlueprintEditor> Editor, UClass* ComponentClass, FName OwnerName);
+	void RemoveComponent(TSharedRef<FWidgetBlueprintEditor> Editor, UClass* ComponentClass, FName OwnerName);
+	TArray<UUIComponent*> GetComponentsFor(UWidget* Target) const;
+	UUIComponent* GetComponent(const UClass* ComponentClass, FName OwnerName) const;
+
+	UUIComponentUserWidgetExtension* GetOrCreateExtension(UUserWidget* PreviewWidget);	
+	void RenameWidget(const FName& OldVarName, const FName& NewVarName);
+
+	bool VerifyContainer(UUserWidget* UserWidget) const;
+	
+private:
+	void AddComponentToPreview(UUserWidget* PreviewWidget, FName OwnerName, UUIComponent* ArchetypeComponent);
+
+protected:
+	virtual void HandleBeginCompilation(FWidgetBlueprintCompilerContext& InCreationContext) override;	
+	virtual void HandleCleanAndSanitizeClass(UWidgetBlueprintGeneratedClass* ClassToClean, UObject* OldCDO) override;
+	virtual void HandleFinishCompilingClass(UWidgetBlueprintGeneratedClass* Class) override;
+	virtual bool HandleValidateGeneratedClass(UWidgetBlueprintGeneratedClass* Class) override; 
+
+	virtual void HandleEndCompilation() override;
+
+private:	
+	UUIComponentContainer* DuplicateContainer(UObject* Outer) const;
+
+	FWidgetBlueprintCompilerContext* CompilerContext = nullptr;
+
+#if WITH_EDITORONLY_DATA	
+	UPROPERTY()
+	TObjectPtr<UUIComponentContainer> ComponentContainer;
+#endif //WITH_EDITORONLY_DATA
+};
\ No newline at end of file
diff --git a/Engine/Source/Editor/UMGEditor/Private/UMGEditorModule.cpp b/Engine/Source/Editor/UMGEditor/Private/UMGEditorModule.cpp
--- a/Engine/Source/Editor/UMGEditor/Private/UMGEditorModule.cpp
+++ b/Engine/Source/Editor/UMGEditor/Private/UMGEditorModule.cpp
@@ -111,8 +111,6 @@ public:
 		KismetCompilerModule.GetCompilers().Add(&WidgetBlueprintCompiler);
 		KismetCompilerModule.OverrideBPTypeForClass(UUserWidget::StaticClass(), UWidgetBlueprint::StaticClass());
 
-		FBlueprintEditorUtils::OnRenameVariableReferencesEvent.AddRaw(this, &FUMGEditorModule::HandleRenameVariableReferences);
-
 		// Add Customization for variable in Graph editor
 		if (FBlueprintEditorModule* BlueprintEditorModule = FModuleManager::GetModulePtr<FBlueprintEditorModule>("Kismet"))
 		{
@@ -510,14 +508,6 @@ private:
 		}
 	}
 
-	void HandleRenameVariableReferences(UBlueprint* Blueprint, UClass* VariableClass, const FName& OldVarName, const FName& NewVarName)
-	{
-		if (UUIComponentContainer* ComponentsContainer = FUIComponentUtils::GetUIComponentContainerFromWidgetBlueprint(Cast<UWidgetBlueprint>(Blueprint)))
-		{
-			ComponentsContainer->RenameWidget(OldVarName, NewVarName);
-		}
-	}
-
 private:
 	TSharedPtr<FExtensibilityManager> MenuExtensibilityManager;
 	TSharedPtr<FExtensibilityManager> ToolBarExtensibilityManager;
diff --git a/Engine/Source/Editor/UMGEditor/Private/WidgetBlueprintEditor.cpp b/Engine/Source/Editor/UMGEditor/Private/WidgetBlueprintEditor.cpp
--- a/Engine/Source/Editor/UMGEditor/Private/WidgetBlueprintEditor.cpp
+++ b/Engine/Source/Editor/UMGEditor/Private/WidgetBlueprintEditor.cpp
@@ -90,6 +90,10 @@
 #include "UnrealEdGlobals.h"
 #include "GraphEditorActions.h"
 #include "MovieSceneDynamicBindingCustomization.h"
+#include "UIComponentWidgetBlueprintExtension.h"
+
+#include "Extensions/UIComponent.h"
+#include "Extensions/UIComponentUserWidgetExtension.h"
 
 #define LOCTEXT_NAMESPACE "UMG"
 
@@ -1114,52 +1118,77 @@ void FWidgetBlueprintEditor::MigrateFromChain(FEditPropertyChain* PropertyThatCh
 {
 	UWidgetBlueprint* Blueprint = GetWidgetBlueprintObj();
 
-	UUserWidget* PreviewUserWidget = GetPreview();
-	if ( PreviewUserWidget != nullptr )
+	if (const UUserWidget* PreviewUserWidget = GetPreview())
 	{
-		for ( TWeakObjectPtr<UObject> ObjectRef : SelectedObjects )
+		FEditPropertyChain::TDoubleLinkedListNode* PropertyChainNode = PropertyThatChanged->GetHead();
+		ensure(PropertyChainNode);
+		
+		const UClass* OwnerClass = PropertyChainNode->GetValue()->GetOwnerClass();
+		if(!OwnerClass)
+		{
+			return;
+		}
+		
+		// Here we handle migrating values on the Root Widget. So we need to migrate the property to the
+		// Widget CDO.
+		for (TWeakObjectPtr<UObject> ObjectRef : SelectedObjects)
 		{
 			UObject* ObjectPtr = ObjectRef.Get();
 			check(ObjectPtr);
-
-			if (ObjectPtr)
+			if (!ObjectPtr)
+			{
+				continue;
+			}
+			
+			// dealing with root widget here
+			if (UObject* WidgetCDO = ObjectPtr->GetClass()->GetDefaultObject(true))
 			{
-				// dealing with root widget here
-				FEditPropertyChain::TDoubleLinkedListNode* PropertyChainNode = PropertyThatChanged->GetHead();
-				UObject* WidgetCDO = ObjectPtr->GetClass()->GetDefaultObject(true);
-
 				// Only do the migration if the property on the head of the linked list lives in the CDO
 				// We want to skip the migration if the property isn't leading to this CDO, such as when we call 
 				// FDetailCategoryImpl::AddExternalObjects to inject external objects in the details panel of the CDO.
-				if (!PropertyChainNode->GetValue() || WidgetCDO == nullptr ||
-					(ObjectPtr->IsA(PropertyChainNode->GetValue()->GetOwnerClass()) && WidgetCDO->IsA(PropertyChainNode->GetValue()->GetOwnerClass())))
+				if (ObjectPtr->IsA(OwnerClass)	&& WidgetCDO->IsA(OwnerClass))
 				{
 					MigratePropertyValue(ObjectPtr, WidgetCDO, PropertyChainNode, PropertyChainNode->GetValue(), bIsModify);
 				}
 			}
 		}
 
+		// Here we handle migrating values on Widgets in the WidgetTree. So we need to migrate the property to the
+		// WidgetBlueprint WidgetTree and instead of the Preview Widget
 		for ( FWidgetReference& WidgetRef : SelectedWidgets )
 		{
-			UWidget* PreviewWidget = WidgetRef.GetPreview();
-
-			if ( PreviewWidget )
+			if ( UWidget* PreviewWidget = WidgetRef.GetPreview())
 			{
 				FName PreviewWidgetName = PreviewWidget->GetFName();
-				UWidget* TemplateWidget = Blueprint->WidgetTree->FindWidget(PreviewWidgetName);
-
-				if ( TemplateWidget )
+				if (UWidget* TemplateWidget = Blueprint->WidgetTree->FindWidget(PreviewWidgetName))
 				{
-					FEditPropertyChain::TDoubleLinkedListNode* PropertyChainNode = PropertyThatChanged->GetHead();
-
+					ensure(TemplateWidget == WidgetRef.GetTemplate());
+					
 					// Only do the migration if the property on the head of the linked list lives in this UWidget
 					// We want to skip the migration if the property isn't leading to this UWidget, such as when we call 
 					// FDetailCategoryImpl::AddExternalObjects to inject external objects in the details panel of this UWidget.
-					if (!PropertyChainNode->GetValue() ||
-						(PreviewWidget->IsA(PropertyChainNode->GetValue()->GetOwnerClass()) && TemplateWidget->IsA(PropertyChainNode->GetValue()->GetOwnerClass())))
+					if (PreviewWidget->IsA(OwnerClass) && TemplateWidget->IsA(OwnerClass))
 					{
 						MigratePropertyValue(PreviewWidget, TemplateWidget, PropertyChainNode, PropertyChainNode->GetValue(), bIsModify);
 					}
+					// If the Object being modified is a UUIComponent, we want to migrate to the Extension
+					else if(OwnerClass->IsChildOf<UUIComponent>())
+					{
+						// We migrate the value to the Widget Blueprint Extension component
+						UUIComponentUserWidgetExtension* ComponentExtension = PreviewUserWidget->GetExtension<UUIComponentUserWidgetExtension>();
+						UUIComponentWidgetBlueprintExtension* ComponentWidgetBlueprintExtension = UWidgetBlueprintExtension::GetExtension<UUIComponentWidgetBlueprintExtension>(Blueprint);
+						if (ComponentExtension && ComponentWidgetBlueprintExtension)
+						{						
+							UUIComponent* Component = ComponentExtension->GetComponent(OwnerClass, TemplateWidget->GetFName());
+							UUIComponent* WPComponent = ComponentWidgetBlueprintExtension->GetComponent(OwnerClass, TemplateWidget->GetFName());							
+
+							if(WPComponent && Component)
+							{
+								MigratePropertyValue(Component, WPComponent, PropertyChainNode, PropertyChainNode->GetValue(), bIsModify);
+							}
+						}
+						
+					}
 				}
 			}
 		}
diff --git a/Engine/Source/Editor/UMGEditor/Private/WidgetBlueprintEditorUtils.cpp b/Engine/Source/Editor/UMGEditor/Private/WidgetBlueprintEditorUtils.cpp
--- a/Engine/Source/Editor/UMGEditor/Private/WidgetBlueprintEditorUtils.cpp
+++ b/Engine/Source/Editor/UMGEditor/Private/WidgetBlueprintEditorUtils.cpp
@@ -68,6 +68,7 @@
 #include "Widgets/SVirtualWindow.h"
 #include "GraphEditorActions.h"
 #include "WidgetEditingProjectSettings.h"
+#include "UIComponentUtils.h"
 
 #define LOCTEXT_NAMESPACE "UMG"
 
@@ -446,6 +447,9 @@ bool FWidgetBlueprintEditorUtils::RenameWidget(TSharedRef<FWidgetBlueprintEditor
 			}
 		});
 
+		// If we use Components, make sure to remane the target.
+		FUIComponentUtils::RenameWidgetName(BlueprintEditor, Blueprint, OldObjectName, NewFName);		
+
 		// Validate child blueprints and adjust variable names to avoid a potential name collision
 		FBlueprintEditorUtils::ValidateBlueprintChildVariables(Blueprint, NewFName);
 
diff --git a/Engine/Source/Editor/UMGEditor/Public/DesignerExtension.h b/Engine/Source/Editor/UMGEditor/Public/DesignerExtension.h
--- a/Engine/Source/Editor/UMGEditor/Public/DesignerExtension.h
+++ b/Engine/Source/Editor/UMGEditor/Public/DesignerExtension.h
@@ -135,6 +135,11 @@ public:
 	{
 	}
 
+	/** Called every time the designer creates new preview. */
+	virtual void PreviewContentCreated(UUserWidget* PreviewWidget)
+	{
+	}
+
 	/** Called every time the content of the designer changed. */
 	virtual void PreviewContentChanged(TSharedRef<SWidget> NewContent)
 	{
diff --git a/Engine/Source/Editor/UMGEditor/Public/UIComponentUtils.h b/Engine/Source/Editor/UMGEditor/Public/UIComponentUtils.h
--- a/Engine/Source/Editor/UMGEditor/Public/UIComponentUtils.h
+++ b/Engine/Source/Editor/UMGEditor/Public/UIComponentUtils.h
@@ -10,6 +10,7 @@ class UUserWidget;
 class UUIComponent;
 class UUIComponentContainer;
 class UWidgetBlueprint;
+class FWidgetBlueprintEditor;
 /**  */
 
 class UMGEDITOR_API FUIComponentUtils
@@ -30,9 +31,7 @@ public:
 		virtual bool IsUnloadedClassAllowed(const FClassViewerInitializationOptions& InInitOptions, const TSharedRef< const IUnloadedBlueprintData > InUnloadedClassData, TSharedRef< FClassViewerFilterFuncs > InFilterFuncs) override;
 	};
 
-	static UUIComponent* CreateUIComponent(TSubclassOf<UUIComponent> ComponentClass, UUserWidget* Outer);
-	static UUIComponentContainer* GetOrCreateComponentsContainerForUserWidget(UUserWidget* UserWidget);
 	static FClassViewerInitializationOptions CreateClassViewerInitializationOptions();
 
-	static UUIComponentContainer* GetUIComponentContainerFromWidgetBlueprint(UWidgetBlueprint* WidgetBlueprint);
+	static void RenameWidgetName(TSharedRef<FWidgetBlueprintEditor> BlueprintEditor, UWidgetBlueprint* WidgetBlueprint, const FName& OldVarName, const FName& NewVarName);
 };
\ No newline at end of file
diff --git a/Engine/Source/Runtime/UMG/Private/Extensions/UIComponent.cpp b/Engine/Source/Runtime/UMG/Private/Extensions/UIComponent.cpp
--- a/Engine/Source/Runtime/UMG/Private/Extensions/UIComponent.cpp
+++ b/Engine/Source/Runtime/UMG/Private/Extensions/UIComponent.cpp
@@ -5,9 +5,11 @@
 #include "CoreGlobals.h"
 #include "Components/Widget.h"
 #include "Blueprint/UserWidget.h"
+//#include "UMGPrivate.h"
 
 bool UUIComponent::Initialize(UWidget* Target)
 {
+	ensure(Target);
 	Owner = Target;
 	const bool Result = OnInitialize();
 
@@ -19,6 +21,12 @@ bool UUIComponent::Initialize(UWidget* Target)
 	return Result;
 }
 
+void UUIComponent::PreConstruct(bool bIsDesignTime)
+{
+	OnPreConstruct(bIsDesignTime);
+}
+
+
 void UUIComponent::Construct()
 {
 	OnConstruct();
@@ -39,6 +47,11 @@ bool UUIComponent::OnInitialize()
 	return true;
 }
 
+void UUIComponent::OnPreConstruct(bool bIsDesignTime)
+{
+
+}
+
 void UUIComponent::OnConstruct()
 {
 
@@ -47,4 +60,4 @@ void UUIComponent::OnConstruct()
 void UUIComponent::OnDestruct()
 {
 
-}
+}
\ No newline at end of file
diff --git a/Engine/Source/Runtime/UMG/Private/Extensions/UIComponentContainer.cpp b/Engine/Source/Runtime/UMG/Private/Extensions/UIComponentContainer.cpp
--- a/Engine/Source/Runtime/UMG/Private/Extensions/UIComponentContainer.cpp
+++ b/Engine/Source/Runtime/UMG/Private/Extensions/UIComponentContainer.cpp
@@ -32,114 +32,115 @@ void FUIComponentTarget::SetTargetName(FName NewName)
 	TargetName = NewName;
 }
 
-void UUIComponentContainer::Initialize()
-{
-	if (UUserWidget* Owner = GetUserWidget())
+void UUIComponentContainer::AddComponent(FName TargetName, UUIComponent* Component)
+{	
+	if (ensure(!TargetName.IsNone() && Component != nullptr))
 	{
-		// Before intializing, we resolve all components
-		for (FUIComponentTarget& ComponentTarget : Components)
+		if(!GetComponent(Component->GetClass(), TargetName))
 		{
-			if (UWidget* Widget = ComponentTarget.Resolve(Owner->WidgetTree))
-			{
-				if (UUIComponent* Component = ComponentTarget.GetComponent())
-				{
-					Component->Initialize(Widget);
-				}
-			}
+			Modify();
+			Components.Emplace(FUIComponentTarget(Component, TargetName));
 		}
 	}
 }
 
-void UUIComponentContainer::Construct()
+void UUIComponentContainer::RemoveComponent(FName TargetName, UUIComponent* Component)
 {
-	for (FUIComponentTarget& ComponentTarget : Components)
+	if(ensure(!TargetName.IsNone() && Component != nullptr))
 	{
-		if(UUIComponent* Component = ComponentTarget.GetComponent())
+		for (int i = 0; i < Components.Num(); i++)
 		{
-			if (TStrongObjectPtr<UWidget> Owner = Component->GetOwner().Pin())
+			FUIComponentTarget& Target = Components[i];
+			if (Target.GetComponent() == Component)
 			{
-				Component->Construct();
+				ensure(Target.GetTargetName() == TargetName);
+				Components.RemoveAtSwap(i);
+				return;
 			}
-		}		
+		}
 	}
 }
 
-void UUIComponentContainer::Destruct()
+void UUIComponentContainer::RemoveAllComponentsOfType(const UClass* ComponentClass, FName TargetName)
 {
-	for (FUIComponentTarget& ComponentTarget : Components)
+	ensure(!TargetName.IsNone());
+
+	for (int32 Index = Components.Num() - 1; Index >= 0; --Index)
 	{
-		if (UUIComponent* Component = ComponentTarget.GetComponent())
+		FUIComponentTarget& Target = Components[Index];
+		if (Target.GetComponent()->GetClass() == ComponentClass && Target.GetTargetName() == TargetName)
 		{
-			if (TStrongObjectPtr<UWidget> Owner = Component->GetOwner().Pin())
-			{
-				Component->Destruct();
-			}
+			Components.RemoveAt(Index);
 		}
 	}
 }
 
-
-TArray<UUIComponent*> UUIComponentContainer::GetExtensionsFor(FName TargetName)
+void UUIComponentContainer::RemoveAllComponentsFor(FName TargetName)
 {
-	TArray<UUIComponent*> Result;
-
-	for (FUIComponentTarget& Target : Components)
+	for (int32 Index = Components.Num() - 1; Index >= 0; --Index)
 	{
+		FUIComponentTarget& Target = Components[Index];
 		if (Target.GetTargetName() == TargetName)
 		{
-			Result.Push(Target.GetComponent());
+			Components.RemoveAtSwap(Index);
 		}
 	}
-
-	return MoveTemp(Result);
 }
 
-void UUIComponentContainer::AddComponent(FName TargetName, UUIComponent* Component)
+UUIComponent* UUIComponentContainer::GetComponent(const UClass* ComponentClass, FName TargetName) const
 {
-	ensure(!TargetName.IsNone() && Component != nullptr);
-	Components.Emplace(FUIComponentTarget(Component, TargetName));
+	for (int32 Index = Components.Num() - 1; Index >= 0; --Index)
+	{
+		const FUIComponentTarget& Target = Components[Index];
+		if (Target.GetTargetName() == TargetName && Target.GetComponent()->GetClass() == ComponentClass)
+		{
+			return Target.GetComponent();
+		}
+	}
+	return nullptr;
 }
 
-void UUIComponentContainer::RemoveComponent(FName TargetName, UUIComponent* Component)
+void UUIComponentContainer::ForEachComponent(TFunctionRef<void(UUIComponent*)> Predicate) const
 {
-	ensure(!TargetName.IsNone() && Component != nullptr);
-
-	for (int i = 0; i < Components.Num(); i++)
+	for (const FUIComponentTarget& ComponentTarget : Components)
 	{
-		FUIComponentTarget& Target = Components[i];
-		if (Target.GetComponent() == Component)
+		if (UUIComponent* Component = ComponentTarget.GetComponent())
 		{
-			ensure(Target.GetTargetName() == TargetName);
-			Components.RemoveAtSwap(i);
-			return;
+			Predicate(Component);
 		}
 	}
 }
 
-void UUIComponentContainer::Resolve()
+void UUIComponentContainer::ForEachComponentTarget(TFunctionRef<void(const FUIComponentTarget&)> Predicate) const
 {
-	if (UUserWidget* Owner = GetUserWidget())
+	for (const FUIComponentTarget& ComponentTarget : Components)
 	{
-		for (FUIComponentTarget& Component : Components)
-		{
-			Component.Resolve(Owner->WidgetTree);
-		}
+		Predicate(ComponentTarget);
 	}
 }
 
-void UUIComponentContainer::RemoveAllComponentsFor(FName TargetName)
+void UUIComponentContainer::InitializeComponents(const UUserWidget* UserWidget)
 {
-	for (int32 Index = Components.Num() - 1; Index >= 0; --Index)
+	for (FUIComponentTarget& ComponentTarget : Components)
 	{
-		FUIComponentTarget& Target = Components[Index];
-		if (Target.GetTargetName() == TargetName)
+		if (UUIComponent* Component = ComponentTarget.GetComponent())
 		{
-			Components.RemoveAtSwap(Index);
+			UWidget* Widget = ComponentTarget.Resolve(UserWidget->WidgetTree);
+			if (ensure(Widget))
+			{
+				Component->Initialize(Widget);
+			}
 		}
 	}
 }
 
-void UUIComponentContainer::RenameWidget(FName OldName, FName NewName)
+bool UUIComponentContainer::IsEmpty() const
+{
+	return Components.IsEmpty();
+}
+
+#ifdef WITH_EDITOR
+void UUIComponentContainer::RenameWidget(const FName& OldName, const FName& NewName)
 {
 	bool bModified = false;
 	auto UpdateModify = [&bModified, Self = this]()
@@ -157,91 +158,68 @@ void UUIComponentContainer::RenameWidget(FName OldName, FName NewName)
 		{
 			UpdateModify();
 			Target.SetTargetName(NewName);
-
-			// @TODO Do our widget components need to be notified of name change?
-
-			//if (Component.GetComponent())
-			//{
-			//	Component.GetComponent()->WidgetRenamed(OldName, NewName);
-			//}
 		}
 	}
 }
-
-
-bool UUIComponentContainer::IsEmpty() const 
-{
-	return Components.IsEmpty();
-}
-
-void UUIComponentContainer::OnPreviewContentChanged(TSharedRef<SWidget> NewContent)
+void UUIComponentContainer::CleanupUIComponents(const UWidgetTree* WidgetTree)
 {
-	if (TSharedPtr<SObjectWidget> ObjectWidget = StaticCastSharedPtr<SObjectWidget>(NewContent.ToSharedPtr()))
+	if (WidgetTree)
 	{
-		CleanupUIComponents(ObjectWidget->GetWidgetObject());
-	}
-}
-
-void UUIComponentContainer::CleanupUIComponents(UUserWidget* UserWidget)
-{
-	if (UserWidget)
-	{
-		if (const UWidgetTree* WidgetTree = UserWidget->WidgetTree)
+		bool bModified = false;
+		auto UpdateModify = [&bModified, Self = this]()
 		{
-			bool bModified = false;
-			auto UpdateModify = [&bModified, Self = this]()
+			if (!bModified)
 			{
-				if (!bModified)
-				{
-					bModified = true;
-					Self->Modify();
-				}
-			};
-			
-			TArray <FName, TInlineAllocator<4>> TargetNamesToRemove;
-			for (int32 Index = Components.Num() - 1; Index >= 0; --Index)
+				bModified = true;
+				Self->Modify();
+			}
+		};
+		
+		TArray <FName, TInlineAllocator<4>> TargetNamesToRemove;
+		for (int32 Index = Components.Num() - 1; Index >= 0; --Index)
+		{
+			// Remove components that are invalid
+			if (Components[Index].GetComponent() == nullptr)
 			{
-				// Remove components that are invalid
-				if (Components[Index].GetComponent() == nullptr)
-				{
-					UpdateModify();
-					Components.RemoveAtSwap(Index);
-				}
-				else if (!Components[Index].GetTargetName().IsNone())
-				{
-					// Fill list with all Valid Target Names
-					TargetNamesToRemove.AddUnique(Components[Index].GetTargetName());
-				}
+				UpdateModify();
+				Components.RemoveAtSwap(Index);
 			}
-
-			// Exclude components in use from removal list
-			if (TargetNamesToRemove.Num() > 0)
+			else if (!Components[Index].GetTargetName().IsNone())
 			{
-				WidgetTree->ForEachWidget([&TargetNamesToRemove, this](TObjectPtr<UWidget> Widget) {
-					if (Widget)
+				// Fill list with all Valid Target Names
+				TargetNamesToRemove.AddUnique(Components[Index].GetTargetName());
+			}
+		}
+
+		// Exclude components in use from removal list
+		if (TargetNamesToRemove.Num() > 0)
+		{
+			WidgetTree->ForEachWidget([&TargetNamesToRemove, this](TObjectPtr<UWidget> Widget) {
+				if (Widget)
+				{
+					for (int32 Index = TargetNamesToRemove.Num() - 1; Index >= 0; Index--)
 					{
-						for (int32 Index = TargetNamesToRemove.Num() - 1; Index >= 0; Index--)
+						const FName TargetName = TargetNamesToRemove[Index];
+						if (TargetName == Widget->GetFName())
 						{
-							const FName TargetName = TargetNamesToRemove[Index];
-							if (TargetName == Widget->GetFName())
-							{
-								TargetNamesToRemove.RemoveSingleSwap(TargetName);
-							}
+							TargetNamesToRemove.RemoveSingleSwap(TargetName);
 						}
 					}
-					});
-			}
+				}
+				});
+		}
 
-			// Remove all unused component 
-			if (TargetNamesToRemove.Num() > 0)
-			{
-				UpdateModify();
+		// Remove all unused component 
+		if (TargetNamesToRemove.Num() > 0)
+		{
+			UpdateModify();
 
-				for (FName TargetName : TargetNamesToRemove)
-				{
-					RemoveAllComponentsFor(TargetName);
-				}
+			for (FName TargetName : TargetNamesToRemove)
+			{
+				RemoveAllComponentsFor(TargetName);
 			}
 		}
 	}
-}
\ No newline at end of file
+}
+
+#endif //WITH_EDITOR
\ No newline at end of file
diff --git a/Engine/Source/Runtime/UMG/Private/Extensions/UIComponentUserWidgetExtension.cpp b/Engine/Source/Runtime/UMG/Private/Extensions/UIComponentUserWidgetExtension.cpp
new file mode 100644
--- /dev/null
+++ b/Engine/Source/Runtime/UMG/Private/Extensions/UIComponentUserWidgetExtension.cpp
@@ -0,0 +1,132 @@
+// Copyright Epic Games, Inc. All Rights Reserved.
+
+#include "Extensions/UIComponentUserWidgetExtension.h"
+
+#include "Blueprint/UserWidget.h"
+#include "CoreMinimal.h"
+#include "Components/Widget.h"
+#include "Extensions/UIComponent.h"
+#include "Extensions/UIComponentContainer.h"
+
+void UUIComponentUserWidgetExtension::PreConstruct(bool bIsDesignTime)
+{
+	Super::Construct();
+
+	InitializeComponents();
+	
+	ComponentContainer->ForEachComponent([bIsDesignTime](UUIComponent* Component)
+		{
+			if(ensure(Component))
+			{
+				Component->PreConstruct(bIsDesignTime);
+			}
+		}
+	);
+}
+void UUIComponentUserWidgetExtension::Construct()
+{
+	Super::Construct();
+
+	ComponentContainer->ForEachComponent([](UUIComponent* Component)
+		{
+			if (ensure(Component))
+			{
+				Component->Construct();
+			}
+		}
+	);
+}
+
+void UUIComponentUserWidgetExtension::Destruct()
+{	
+	Super::Destruct();
+
+	ComponentContainer->ForEachComponent([](UUIComponent* Component)
+		{
+			if (ensure(Component))
+			{
+				Component->Destruct();
+			}
+		}
+	);
+}
+
+void UUIComponentUserWidgetExtension::InitializeContainer(UUIComponentContainer* InComponentContainer)
+{
+	check(InComponentContainer != nullptr);
+	check(ComponentContainer == nullptr);
+	ComponentContainer = InComponentContainer;
+}
+
+void UUIComponentUserWidgetExtension::InitializeComponents()
+{
+	const UUserWidget* UserWidget = GetUserWidget();
+	check(UserWidget);
+	check(ComponentContainer);
+	if (ComponentContainer)
+	{
+		ComponentContainer->InitializeComponents(UserWidget);
+	}
+}
+
+#if WITH_EDITOR
+
+void UUIComponentUserWidgetExtension::RenameWidget(const FName& OldVarName, const FName& NewVarName)
+{
+	if (ensure(ComponentContainer))
+	{
+		ComponentContainer->RenameWidget(OldVarName, NewVarName);
+	}
+}
+
+// Used only to create a Component on the PreviewWidget in the editor, based on the Component Archetype object in the WidgetBlueprint.
+void UUIComponentUserWidgetExtension::CreateAndAddComponent(UUIComponent* ArchetypeComponent, FName OwnerName)
+{
+	ensure(ArchetypeComponent);
+	ensure(ComponentContainer);
+	if (ArchetypeComponent && ComponentContainer)
+	{
+		ensure(ComponentContainer && ArchetypeComponent && ArchetypeComponent->HasAllFlags(RF_ArchetypeObject));
+
+		// Create the component with the WBP Component as the Archetype.
+		UObject* Outer = ComponentContainer;	
+		UUIComponent* PreviewWidgetComponent = NewObject<UUIComponent>(Outer,
+			ArchetypeComponent->GetClass(),
+			NAME_None, RF_Transactional,
+			ArchetypeComponent);
+
+		ComponentContainer->AddComponent(OwnerName, PreviewWidgetComponent);
+	}
+}
+
+void UUIComponentUserWidgetExtension::RemoveComponent(const UClass* ComponentClass, FName OwnerName)
+{
+	if(ComponentContainer)
+	{
+		ComponentContainer->RemoveAllComponentsOfType(ComponentClass, OwnerName);
+	}
+}
+#endif // WITH_EDITOR
+
+TArray<UUIComponent*> UUIComponentUserWidgetExtension::GetComponentsFor(UWidget* Target) const
+{
+	ensure(Target);
+
+	TArray<UUIComponent*> Components;
+	if(Target)
+	{
+		ComponentContainer->ForEachComponent([Target, &Components](UUIComponent* Component){			
+			if(ensure(Component) && Component->GetOwner() == Target)			
+			{
+				Components.Push(Component);
+			}
+		});
+	}
+	return Components;	
+}
+
+UUIComponent* UUIComponentUserWidgetExtension::GetComponent(const UClass* ComponentClass, FName OwnerName) const
+{
+	check(ComponentContainer);
+	return ComponentContainer->GetComponent(ComponentClass, OwnerName);
+}
\ No newline at end of file
diff --git a/Engine/Source/Runtime/UMG/Private/Extensions/UIComponentWidgetBlueprintGeneratedClassExtension.cpp b/Engine/Source/Runtime/UMG/Private/Extensions/UIComponentWidgetBlueprintGeneratedClassExtension.cpp
new file mode 100644
--- /dev/null
+++ b/Engine/Source/Runtime/UMG/Private/Extensions/UIComponentWidgetBlueprintGeneratedClassExtension.cpp
@@ -0,0 +1,101 @@
+// Copyright Epic Games, Inc. All Rights Reserved.
+
+#include "Extensions/UIComponentWidgetBlueprintGeneratedClassExtension.h"
+
+#include "Blueprint/UserWidget.h"
+#include "CoreMinimal.h"
+#include "Blueprint/WidgetTree.h"
+#include "Extensions/UIComponent.h"
+#include "Extensions/UIComponentContainer.h"
+#include "Extensions/UIComponentUserWidgetExtension.h"
+
+UUIComponentContainer* UUIComponentWidgetBlueprintGeneratedClassExtension::DuplicateContainer(UUserWidget* UserWidget) const
+{
+	check(ComponentContainer);
+	check(!ComponentContainer->IsEmpty());
+	
+	UObject* Outer = UserWidget;
+		
+	//FObjectInstancingGraph ObjectInstancingGraph;
+	UUIComponentContainer* NewContainer = NewObject<UUIComponentContainer>(Outer,
+		ComponentContainer->GetClass(),
+		NAME_None,
+		RF_Transactional,
+		ComponentContainer,
+		false);
+
+	return NewContainer;
+}
+
+bool UUIComponentWidgetBlueprintGeneratedClassExtension::VerifyContainer(const UUserWidget* UserWidget) const
+{
+	if(const UUIComponentUserWidgetExtension* UserWidgetExtension = UserWidget->GetExtension<UUIComponentUserWidgetExtension>())
+	{
+		TArray<FUIComponentTarget> Components;
+		ComponentContainer->ForEachComponentTarget([&Components](const FUIComponentTarget& ComponentTarget){
+			Components.Push(ComponentTarget);
+		});
+
+		while(Components.Num() > 0)
+		{
+			const FUIComponentTarget& Target = Components[Components.Num()-1];
+			if(!UserWidgetExtension->GetComponent(Target.GetComponent()->GetClass(), Target.GetTargetName()))
+			{
+				return false;
+			}
+			else
+			{
+				Components.RemoveAt(Components.Num()-1);
+			}
+		}
+		return true;
+	}
+	return false;	
+}
+
+#if WITH_EDITOR
+bool UUIComponentWidgetBlueprintGeneratedClassExtension::VerifyAllWidgetsExists(const UWidgetTree* WidgetTree) const
+{	
+	TSet<FName> Targets;
+	ComponentContainer->ForEachComponentTarget([&Targets](const FUIComponentTarget& ComponentTarget){
+		Targets.FindOrAdd(ComponentTarget.GetTargetName());
+	});
+
+	for (FName& WidgetName : Targets)
+	{
+		if (WidgetTree->FindWidget(WidgetName))
+		{
+			continue;
+		}
+		return false;
+	}
+	return true;	
+}
+#endif // WITH_EDITOR
+
+void UUIComponentWidgetBlueprintGeneratedClassExtension::InitializeContainer(UUIComponentContainer* InComponentContainer)
+{
+	check(InComponentContainer != nullptr);
+	check(ComponentContainer == nullptr);
+	ComponentContainer = InComponentContainer;
+}
+
+void UUIComponentWidgetBlueprintGeneratedClassExtension::PreConstruct(UUserWidget* UserWidget, bool IsDesignTime)
+{
+	// This is where we can add the UserWidgetExtension and Duplicate the Component Container from the GeneratedClassExtension
+
+	// If the extension do not exists, Add it.
+	UUIComponentUserWidgetExtension* UserWidgetExtension = UserWidget->GetExtension<UUIComponentUserWidgetExtension>();
+	if (!UserWidgetExtension)
+	{
+		UserWidgetExtension = UserWidget->AddExtension<UUIComponentUserWidgetExtension>();
+		UserWidgetExtension->InitializeContainer(DuplicateContainer(UserWidget));
+	}
+	else
+	{
+		if (UserWidget != UserWidgetExtension->GetOuterUUserWidget())
+		{
+			UserWidgetExtension->Rename(nullptr, UserWidget);
+		}
+	}
+}
diff --git a/Engine/Source/Runtime/UMG/Private/UserWidget.cpp b/Engine/Source/Runtime/UMG/Private/UserWidget.cpp
--- a/Engine/Source/Runtime/UMG/Private/UserWidget.cpp
+++ b/Engine/Source/Runtime/UMG/Private/UserWidget.cpp
@@ -1775,6 +1775,14 @@ void UUserWidget::NativePreConstruct()
 			});
 	}
 
+	// Extension can add other extensions. Use index loop to initialize them all.
+	for (int32 Index = 0; Index < Extensions.Num(); ++Index)
+	{
+		UUserWidgetExtension* Extension = Extensions[Index];
+		check(Extension);
+		Extension->PreConstruct(bIsDesignTime);
+	}
+
 	DesiredFocusWidget.Resolve(WidgetTree);
 
 	PreConstruct(bIsDesignTime);
diff --git a/Engine/Source/Runtime/UMG/Public/Extensions/UIComponent.h b/Engine/Source/Runtime/UMG/Public/Extensions/UIComponent.h
--- a/Engine/Source/Runtime/UMG/Public/Extensions/UIComponent.h
+++ b/Engine/Source/Runtime/UMG/Public/Extensions/UIComponent.h
@@ -20,15 +20,20 @@ class UUIComponent : public UObject
 
 public:
 	/**
-	 * Called when the owner widget is inititalized.
+	 * Called when the owner widget is initialized.
 	 */
-	UMG_API bool Initialize(UWidget* target);
+	UMG_API bool Initialize(UWidget* Target);
 
+	/**
+	 * Called when the owner widget is pre-constructed. Called in both Editor and runtime.
+	 * @param bIsDesignTime True when the Widget is constructed for design time
+	 */
+	UMG_API void PreConstruct(bool bIsDesignTime);
+	
 	/**
 	 * Called when the owner widget is constructed.
 	 */
 	UMG_API void Construct();
-	
 
 	/**
 	 * Called when the owner widget is destructed.
@@ -43,12 +48,11 @@ public:
 
 protected:
 	UMG_API virtual bool OnInitialize();
+	UMG_API virtual void OnPreConstruct(bool bIsDesignTime);
 	UMG_API virtual void OnConstruct();
 	UMG_API virtual void OnDestruct();
 
 private:
 	UPROPERTY(transient, DuplicateTransient)
 	TWeakObjectPtr<UWidget> Owner;
-
-
 };
diff --git a/Engine/Source/Runtime/UMG/Public/Extensions/UIComponentContainer.h b/Engine/Source/Runtime/UMG/Public/Extensions/UIComponentContainer.h
--- a/Engine/Source/Runtime/UMG/Public/Extensions/UIComponentContainer.h
+++ b/Engine/Source/Runtime/UMG/Public/Extensions/UIComponentContainer.h
@@ -6,13 +6,13 @@
 #include "UObject/ObjectMacros.h"
 #include "UObject/Object.h"
 #include "UObject/WeakObjectPtr.h"
-#include "Extensions/UserWidgetExtension.h"
 
 #include "UIComponentContainer.generated.h"
 
 class UUserWidget;
 class UWidget;
 class UUIComponent;
+class UWidgetTree;
 
 USTRUCT()
 struct FUIComponentTarget
@@ -45,32 +45,30 @@ private:
  * Class that holds all the UIComponents for a UUserWidget.
  */
 UCLASS(NotBlueprintType, Experimental)
-class UMG_API UUIComponentContainer : public UUserWidgetExtension
+class UMG_API UUIComponentContainer : public UObject
 {
 	GENERATED_BODY()
 
 public:
-	virtual void Initialize() override;
-	virtual void Construct() override;		
-	virtual void Destruct() override;
-
-	TArray<UUIComponent*> GetExtensionsFor(FName TargetName);
-
 	void AddComponent(FName TargetName, UUIComponent* Component);
 	void RemoveComponent(FName TargetName, UUIComponent* Component);
+	void RemoveAllComponentsOfType(const UClass* ComponentClass, FName TargetName);
 	void RemoveAllComponentsFor(FName TargetName);
 
-	bool IsEmpty() const;
-
-	void RenameWidget(FName OldName, FName NewName);
-
-	void OnPreviewContentChanged(TSharedRef<SWidget> NewContent);
+	UUIComponent* GetComponent(const UClass* ComponentClass, FName TargetName) const; 
+	
+	void ForEachComponent(TFunctionRef<void(UUIComponent*)> Predicate) const;
+	void ForEachComponentTarget(TFunctionRef<void(const FUIComponentTarget&)> Predicate) const;
 
+	void InitializeComponents(const UUserWidget* UserWidget);
+	
+	bool IsEmpty() const;
+#ifdef WITH_EDITOR
+	void RenameWidget(const FName& OldName, const FName& NewName);
+	void CleanupUIComponents(const UWidgetTree* WidgetTree);	
+#endif //WITH_EDITOR
+	
 private:
-	void Resolve();
-	void CleanupUIComponents(UUserWidget* UserWidget);
-
-	// Use a single TArray for the Entire UUserWidget to reduce memory usage.
 	UPROPERTY()
 	TArray<FUIComponentTarget> Components;
 };
diff --git a/Engine/Source/Runtime/UMG/Public/Extensions/UIComponentUserWidgetExtension.h b/Engine/Source/Runtime/UMG/Public/Extensions/UIComponentUserWidgetExtension.h
new file mode 100644
--- /dev/null
+++ b/Engine/Source/Runtime/UMG/Public/Extensions/UIComponentUserWidgetExtension.h
@@ -0,0 +1,51 @@
+// Copyright Epic Games, Inc. All Rights Reserved.
+
+#pragma once
+
+#include "CoreMinimal.h"
+#include "Extensions/UserWidgetExtension.h"
+
+#include "UIComponentUserWidgetExtension.generated.h"
+
+class UUserWidget;
+class UWidget;
+class UUIComponent;
+class UUIComponentContainer;
+
+/**
+ * Class that holds all the UIComponents for a UUserWidget.
+ */
+UCLASS(NotBlueprintType, DefaultToInstanced, Experimental)
+class UMG_API UUIComponentUserWidgetExtension : public UUserWidgetExtension // This Contains all runtime version the components. They should be already resolved.
+{
+	GENERATED_BODY()
+
+public:
+	//virtual void Initialize() override;
+	virtual void PreConstruct(bool bIsDesignTime) override;
+	virtual void Construct() override;
+	virtual void Destruct() override;
+
+	TArray<UUIComponent*> GetComponentsFor(UWidget* Target) const;
+	UUIComponent* GetComponent(const UClass* ComponentClass, FName OwnerName) const;
+
+	void InitializeContainer(UUIComponentContainer* InComponentContainer);
+	
+
+#if WITH_EDITOR
+	void RenameWidget(const FName& OldVarName, const FName& NewVarName);
+	
+	// Used only to create a Component on the PreviewWidget in the editor, based on the Component Archetype object in the WidgetBlueprint.
+	void CreateAndAddComponent(UUIComponent* ArchetypeComponent, FName OwnerName);
+
+	void RemoveComponent(const UClass* ComponentClass, FName OwnerName);
+
+#endif // WITH_EDITOR
+
+private:
+	void InitializeComponents();
+	
+	// Use a single TArray for the Entire UUserWidget to reduce memory usage.
+	UPROPERTY(Transient, DuplicateTransient)
+	TObjectPtr<UUIComponentContainer> ComponentContainer;
+};
diff --git a/Engine/Source/Runtime/UMG/Public/Extensions/UIComponentWidgetBlueprintGeneratedClassExtension.h b/Engine/Source/Runtime/UMG/Public/Extensions/UIComponentWidgetBlueprintGeneratedClassExtension.h
new file mode 100644
--- /dev/null
+++ b/Engine/Source/Runtime/UMG/Public/Extensions/UIComponentWidgetBlueprintGeneratedClassExtension.h
@@ -0,0 +1,37 @@
+// Copyright Epic Games, Inc. All Rights Reserved.
+
+#pragma once
+
+#include "CoreMinimal.h"
+#include "Extensions/WidgetBlueprintGeneratedClassExtension.h"
+
+#include "UIComponentWidgetBlueprintGeneratedClassExtension.generated.h"
+
+class UUserWidget;
+class UWidget;
+class UUIComponent;
+class UUIComponentContainer;
+
+/**
+ *  Class that hold the archetype version of the Components to be duplicated.
+ */
+UCLASS()
+class UMG_API UUIComponentWidgetBlueprintGeneratedClassExtension : public UWidgetBlueprintGeneratedClassExtension
+{
+	GENERATED_BODY()
+public:
+	void InitializeContainer(UUIComponentContainer* InComponentContainer);
+
+	virtual void PreConstruct(UUserWidget* UserWidget, bool IsDesignTime) override;
+
+#if WITH_EDITOR
+	bool VerifyAllWidgetsExists(const UWidgetTree* WidgetTree) const;
+#endif // WITH_EDITOR
+	
+private:
+	UUIComponentContainer* DuplicateContainer(UUserWidget* UserWidget) const;
+	bool VerifyContainer(const UUserWidget* UserWidget) const;
+	
+	UPROPERTY(DuplicateTransient)
+	TObjectPtr<UUIComponentContainer> ComponentContainer;
+};
diff --git a/Engine/Source/Runtime/UMG/Public/Extensions/UserWidgetExtension.h b/Engine/Source/Runtime/UMG/Public/Extensions/UserWidgetExtension.h
--- a/Engine/Source/Runtime/UMG/Public/Extensions/UserWidgetExtension.h
+++ b/Engine/Source/Runtime/UMG/Public/Extensions/UserWidgetExtension.h
@@ -26,6 +26,11 @@ public:
 	{
 	}
 
+	/** Extend the UUserWidget::PreConstruct function */
+	virtual void PreConstruct(bool bIsDesignTime)
+	{		
+	}
+
 	/** Extend the UUserWidget::Construct function */
 	virtual void Construct()
 	{
