PR #12784:  Fixed usage of uninitialized memory in UViewport::Spawn

diff --git a/Engine/Source/Runtime/UMG/Private/Components/Viewport.cpp b/Engine/Source/Runtime/UMG/Private/Components/Viewport.cpp
--- a/Engine/Source/Runtime/UMG/Private/Components/Viewport.cpp
+++ b/Engine/Source/Runtime/UMG/Private/Components/Viewport.cpp
@@ -489,7 +489,7 @@ AActor* UViewport::Spawn(TSubclassOf<AActor> ActorClass)
 		FActorSpawnParameters SpawnParameters;
 		SpawnParameters.SpawnCollisionHandlingOverride = ESpawnActorCollisionHandlingMethod::AlwaysSpawn;
 
-		return World->SpawnActor<AActor>(ActorClass, FVector(0, 0, 0), FRotator(), SpawnParameters);
+		return World->SpawnActor<AActor>(ActorClass, FVector::ZeroVector, FRotator::ZeroRotator, SpawnParameters);
 	}
 
 	// TODO UMG Report spawning actor error before the world is ready.
