Fix Invalidation Root causing Widgets to look wrong.

A local temporary assigned to a reference is only guaranteed to exist as long as the reference exists, and the reference ceases to exist when the FSlateInvalidationContext constructor returns.

diff --git a/Engine/Source/Runtime/UMG/Private/Slate/WidgetRenderer.cpp b/Engine/Source/Runtime/UMG/Private/Slate/WidgetRenderer.cpp
--- a/Engine/Source/Runtime/UMG/Private/Slate/WidgetRenderer.cpp
+++ b/Engine/Source/Runtime/UMG/Private/Slate/WidgetRenderer.cpp
@@ -392,7 +392,8 @@ bool FWidgetRenderer::DrawInvalidationRoot(TSharedRef<SVirtualWindow>& VirtualWi
 			FSlateWindowElementList& WindowElementList = ScopedDrawBuffer.GetDrawBuffer().AddWindowElementList(VirtualWindow);
 
 			// Populates cached element data lists using the provide invalidation root
-			FSlateInvalidationContext InvalidationContext(WindowElementList, FWidgetStyle());
+			FWidgetStyle InitialWidgetStyle;
+			FSlateInvalidationContext InvalidationContext(WindowElementList, InitialWidgetStyle);
 			InvalidationContext.bParentEnabled = true;
 			InvalidationContext.bAllowFastPathUpdate = true;
 			InvalidationContext.LayoutScaleMultiplier = DrawScale;
