Updated ListView to properly send the OnEntryWidgetGenerated event without relying on a timer
This resolves an issue with MVVM where ListView entries won't have their VM bound if no world timer is running (i.e. the game is paused)

diff --git a/Engine/Source/Runtime/UMG/Private/Components/ListViewBase.cpp b/Engine/Source/Runtime/UMG/Private/Components/ListViewBase.cpp
--- a/Engine/Source/Runtime/UMG/Private/Components/ListViewBase.cpp
+++ b/Engine/Source/Runtime/UMG/Private/Components/ListViewBase.cpp
@@ -255,7 +255,11 @@ void UListViewBase::FinishGeneratingEntry(UUserWidget& GeneratedEntry)
 		{
 			if (UWorld* World = GetWorld())
 			{
-				EntryGenAnnouncementTimerHandle = World->GetTimerManager().SetTimerForNextTick(this, &UListViewBase::HandleAnnounceGeneratedEntries);
+				FTSTicker::GetCoreTicker().AddTicker(FTickerDelegate::CreateWeakLambda(this, [this](float InDelta)
+				{
+					HandleAnnounceGeneratedEntries();
+					return false;
+				}));
 			}
 		}
 	}
