UIComponent (Experimental)
- Change the return value for UIComponent::Initialize and OnInitialize to void. This is not used.

diff --git a/Engine/Source/Runtime/UMG/Private/Extensions/UIComponent.cpp b/Engine/Source/Runtime/UMG/Private/Extensions/UIComponent.cpp
--- a/Engine/Source/Runtime/UMG/Private/Extensions/UIComponent.cpp
+++ b/Engine/Source/Runtime/UMG/Private/Extensions/UIComponent.cpp
@@ -7,18 +7,11 @@
 #include "Blueprint/UserWidget.h"
 //#include "UMGPrivate.h"
 
-bool UUIComponent::Initialize(UWidget* Target)
+void UUIComponent::Initialize(UWidget* Target)
 {
 	ensure(Target);
 	Owner = Target;
-	const bool Result = OnInitialize();
-
-	if (!Result)
-	{
-		Owner = nullptr;
-	}
-
-	return Result;
+	OnInitialize();
 }
 
 void UUIComponent::PreConstruct(bool bIsDesignTime)
@@ -42,9 +35,9 @@ TWeakObjectPtr<UWidget> UUIComponent::GetOwner() const
 	return Owner;
 }
 
-bool UUIComponent::OnInitialize()
+void UUIComponent::OnInitialize()
 {
-	return true;
+
 }
 
 void UUIComponent::OnPreConstruct(bool bIsDesignTime)
diff --git a/Engine/Source/Runtime/UMG/Public/Extensions/UIComponent.h b/Engine/Source/Runtime/UMG/Public/Extensions/UIComponent.h
--- a/Engine/Source/Runtime/UMG/Public/Extensions/UIComponent.h
+++ b/Engine/Source/Runtime/UMG/Public/Extensions/UIComponent.h
@@ -22,7 +22,7 @@ public:
 	/**
 	 * Called when the owner widget is initialized.
 	 */
-	UMG_API bool Initialize(UWidget* Target);
+	UMG_API void Initialize(UWidget* Target);
 
 	/**
 	 * Called when the owner widget is pre-constructed. Called in both Editor and runtime.
@@ -47,7 +47,7 @@ public:
 	UMG_API TWeakObjectPtr<UWidget> GetOwner() const;
 
 protected:
-	UMG_API virtual bool OnInitialize();
+	UMG_API virtual void OnInitialize();
 	UMG_API virtual void OnPreConstruct(bool bIsDesignTime);
 	UMG_API virtual void OnConstruct();
 	UMG_API virtual void OnDestruct();
