Restore focus to the category you selected when returning from the quest modal
- This required exposing UUserWidget::SetDesiredFocusWidget to BP

#changelist validated
#virtualized

diff --git a/Engine/Source/Runtime/UMG/Public/Blueprint/UserWidget.h b/Engine/Source/Runtime/UMG/Public/Blueprint/UserWidget.h
--- a/Engine/Source/Runtime/UMG/Public/Blueprint/UserWidget.h
+++ b/Engine/Source/Runtime/UMG/Public/Blueprint/UserWidget.h
@@ -1359,6 +1359,7 @@ public:
 	 * @param Widget Widget to forward the focus to when this widget receives focus
 	 * @return True if the Widget is set properly. Will return false if it's not a child of this UserWidget.
 	 */
+	UFUNCTION(BlueprintCallable, Category = "User Interface|Focus")
 	UMG_API bool SetDesiredFocusWidget(UWidget* Widget);
 
 	/** @returns The Name of the Widget that should receive focus when this UserWidget gets focus. */
