[ListView] Added an OnListViewFinishedScrolling event accessible from blueprint

#rnx

diff --git a/Engine/Source/Runtime/Slate/Private/Widgets/Views/STableViewBase.cpp b/Engine/Source/Runtime/Slate/Private/Widgets/Views/STableViewBase.cpp
--- a/Engine/Source/Runtime/Slate/Private/Widgets/Views/STableViewBase.cpp
+++ b/Engine/Source/Runtime/Slate/Private/Widgets/Views/STableViewBase.cpp
@@ -392,7 +392,9 @@ void STableViewBase::Tick( const FGeometry& AllottedGeometry, const double InCur
 				// We call this rather than just leave bItemsNeedRefresh as true to ensure that EnsureTickToRefresh is registered
 				RequestLayoutRefresh();
 			}
-			else if (CurrentScrollOffset == TargetScrollOffset)
+			else if (CurrentScrollOffset == TargetScrollOffset
+					&& !IsUserScrolling()
+					&& (AllowOverscroll == EAllowOverscroll::No || Overscroll.GetOverscroll(GetTickSpaceGeometry()) == 0.0f))
 			{
 				// Sync the desired scroll offset with the target scroll offset when the scroll ends.
 				// This prevents teleporting the table upon the next scroll. 
diff --git a/Engine/Source/Runtime/Slate/Public/Widgets/Views/STableViewBase.h b/Engine/Source/Runtime/Slate/Public/Widgets/Views/STableViewBase.h
--- a/Engine/Source/Runtime/Slate/Public/Widgets/Views/STableViewBase.h
+++ b/Engine/Source/Runtime/Slate/Public/Widgets/Views/STableViewBase.h
@@ -416,7 +416,10 @@ protected:
 	virtual void NotifyItemScrolledIntoView() = 0;
 
 	/**
-	 * Called when CurrentScrollOffset == TargetScrollOffset at the end of a ::Tick
+	 * Called at the end of a ::Tick when
+	 * - CurrentScrollOffset == TargetScrollOffset 
+	 * - the user is not scrolling
+	 * - the table view is not scrolling from an overscroll position
 	 */
 	virtual void NotifyFinishedScrolling() = 0;
 
diff --git a/Engine/Source/Runtime/Slate/Public/Widgets/Views/STileView.h b/Engine/Source/Runtime/Slate/Public/Widgets/Views/STileView.h
--- a/Engine/Source/Runtime/Slate/Public/Widgets/Views/STileView.h
+++ b/Engine/Source/Runtime/Slate/Public/Widgets/Views/STileView.h
@@ -77,6 +77,8 @@ public:
 		SLATE_EVENT( FOnWidgetToBeRemoved, OnTileReleased )
 
 		SLATE_EVENT( FOnTableViewScrolled, OnTileViewScrolled )
+		
+		SLATE_EVENT( FOnFinishedScrolling, OnFinishedScrolling )
 
 		SLATE_EVENT( FOnItemScrolledIntoView, OnItemScrolledIntoView )
 
diff --git a/Engine/Source/Runtime/Slate/Public/Widgets/Views/STreeView.h b/Engine/Source/Runtime/Slate/Public/Widgets/Views/STreeView.h
--- a/Engine/Source/Runtime/Slate/Public/Widgets/Views/STreeView.h
+++ b/Engine/Source/Runtime/Slate/Public/Widgets/Views/STreeView.h
@@ -159,6 +159,8 @@ public:
 
 		SLATE_EVENT( FOnTableViewScrolled, OnTreeViewScrolled )
 
+		SLATE_EVENT( FOnFinishedScrolling, OnFinishedScrolling )
+
 		SLATE_EVENT( FOnItemScrolledIntoView, OnItemScrolledIntoView )
 
 		SLATE_EVENT( FOnGetChildren, OnGetChildren )
diff --git a/Engine/Source/Runtime/UMG/Private/Components/ListView.cpp b/Engine/Source/Runtime/UMG/Private/Components/ListView.cpp
--- a/Engine/Source/Runtime/UMG/Private/Components/ListView.cpp
+++ b/Engine/Source/Runtime/UMG/Private/Components/ListView.cpp
@@ -362,6 +362,12 @@ void UListView::OnListViewScrolledInternal(float ItemOffset, float DistanceRemai
 	BP_OnListViewScrolled.Broadcast(ItemOffset, DistanceRemaining);
 }
 
+void UListView::OnListViewFinishedScrollingInternal()
+{
+	ITypedUMGListView::OnListViewFinishedScrollingInternal();
+	BP_OnListViewFinishedScrolling.Broadcast();
+}
+
 void UListView::InitHorizontalEntrySpacing(float InHorizontalEntrySpacing)
 {
 	ensureMsgf(!MyListView.IsValid(), TEXT("The SWidget is already created."));
diff --git a/Engine/Source/Runtime/UMG/Public/Components/ListView.h b/Engine/Source/Runtime/UMG/Public/Components/ListView.h
--- a/Engine/Source/Runtime/UMG/Public/Components/ListView.h
+++ b/Engine/Source/Runtime/UMG/Public/Components/ListView.h
@@ -13,6 +13,7 @@ DECLARE_DYNAMIC_MULTICAST_DELEGATE_TwoParams(FOnListItemSelectionChangedDynamic,
 DECLARE_DYNAMIC_MULTICAST_DELEGATE_TwoParams(FOnItemIsHoveredChangedDynamic, UObject*, Item, bool, bIsHovered);
 DECLARE_DYNAMIC_MULTICAST_DELEGATE_TwoParams(FOnListItemScrolledIntoViewDynamic, UObject*, Item, UUserWidget*, Widget);
 DECLARE_DYNAMIC_MULTICAST_DELEGATE_TwoParams(FOnListViewScrolledDynamic, float, ItemOffset, float, DistanceRemaining);
+DECLARE_DYNAMIC_MULTICAST_DELEGATE(FOnListViewFinishedScrollingDynamic);
 
 /**
  * A virtualized list that allows up to thousands of items to be displayed.
@@ -158,6 +159,7 @@ protected:
 	UMG_API virtual void OnSelectionChangedInternal(UObject* FirstSelectedItem) override;
 	UMG_API virtual void OnItemScrolledIntoViewInternal(UObject* Item, UUserWidget& EntryWidget) override;
 	UMG_API virtual void OnListViewScrolledInternal(float ItemOffset, float DistanceRemaining) override;
+	UMG_API virtual void OnListViewFinishedScrollingInternal() override;
 
 	UMG_API void HandleOnEntryInitializedInternal(UObject* Item, const TSharedRef<ITableRow>& TableRow);
 
@@ -331,4 +333,7 @@ private:
 	
 	UPROPERTY(BlueprintAssignable, Category = Events, meta = (DisplayName = "On List View Scrolled"))
 	FOnListViewScrolledDynamic BP_OnListViewScrolled;	
+	
+	UPROPERTY(BlueprintAssignable, Category = Events, meta = (DisplayName = "On List View Finished Scrolling"))
+	FOnListViewFinishedScrollingDynamic BP_OnListViewFinishedScrolling;	
 };
diff --git a/Engine/Source/Runtime/UMG/Public/Components/ListViewBase.h b/Engine/Source/Runtime/UMG/Public/Components/ListViewBase.h
--- a/Engine/Source/Runtime/UMG/Public/Components/ListViewBase.h
+++ b/Engine/Source/Runtime/UMG/Public/Components/ListViewBase.h
@@ -314,6 +314,7 @@ protected:
 			.OnIsSelectableOrNavigable_UObject(Implementer, &UListViewBaseT::HandleIsSelectableOrNavigable)
 			.OnItemScrolledIntoView_UObject(Implementer, &UListViewBaseT::HandleItemScrolledIntoView)
 			.OnTileViewScrolled_UObject(Implementer, &UListViewBaseT::HandleListViewScrolled)
+			.OnFinishedScrolling_UObject(Implementer, &UListViewBaseT::HandleFinishedScrolling)
 			.OnMouseButtonClick_UObject(Implementer, &UListViewBaseT::HandleItemClicked)
 			.OnMouseButtonDoubleClick_UObject(Implementer, &UListViewBaseT::HandleItemDoubleClicked);
 	}
@@ -349,6 +350,7 @@ protected:
 			.OnRowReleased_UObject(Implementer, &UListViewBaseT::HandleRowReleased)
 			.OnItemScrolledIntoView_UObject(Implementer, &UListViewBaseT::HandleItemScrolledIntoView)
 			.OnTreeViewScrolled_UObject(Implementer, &UListViewBaseT::HandleListViewScrolled)
+			.OnFinishedScrolling_UObject(Implementer, &UListViewBaseT::HandleFinishedScrolling)
 			.OnMouseButtonClick_UObject(Implementer, &UListViewBaseT::HandleItemClicked)
 			.OnMouseButtonDoubleClick_UObject(Implementer, &UListViewBaseT::HandleItemDoubleClicked)
 			.OnGetChildren_UObject(Implementer, &UListViewBaseT::HandleGetChildren)
@@ -389,6 +391,7 @@ protected:
 	virtual bool OnIsSelectableOrNavigableInternal(ItemType FirstSelectedItem) { return OnIsItemSelectableOrNavigable().IsBound() ? OnIsItemSelectableOrNavigable().Execute(FirstSelectedItem) : true; }
 	virtual void OnItemScrolledIntoViewInternal(ItemType Item, UUserWidget& EntryWidget) {}
 	virtual void OnListViewScrolledInternal(float ItemOffset, float DistanceRemaining) {}
+	virtual void OnListViewFinishedScrollingInternal() {}
 	virtual void OnItemExpansionChangedInternal(ItemType Item, bool bIsExpanded) {}
 
 private:
@@ -444,6 +447,7 @@ private:
 
 	void HandleFinishedScrolling()
 	{
+		OnListViewFinishedScrollingInternal();
 		OnFinishedScrolling().Broadcast();
 	}
 
