UMG Extensions: Correction to UserWidget Extension Preconstruct to be called when the Extension is added after contruction was already called.

diff --git a/Engine/Source/Runtime/UMG/Private/UserWidget.cpp b/Engine/Source/Runtime/UMG/Private/UserWidget.cpp
--- a/Engine/Source/Runtime/UMG/Private/UserWidget.cpp
+++ b/Engine/Source/Runtime/UMG/Private/UserWidget.cpp
@@ -88,6 +88,7 @@ UUserWidget::UUserWidget(const FObjectInitializer& ObjectInitializer)
 	, bHasScriptImplementedTick(true)
 	, bHasScriptImplementedPaint(true)
 	, bInitialized(false)
+	, bAreExtensionsPreConstructed(false)
 	, bAreExtensionsConstructed(false)
 	, bStoppingAllAnimations(false)
 	, TickFrequency(EWidgetTickFrequency::Auto)
@@ -1775,10 +1776,12 @@ void UUserWidget::NativePreConstruct()
 			});
 	}
 
+	bAreExtensionsPreConstructed = true;
 	// Extension can add other extensions. Use index loop to initialize them all.
-	for (int32 Index = 0; Index < Extensions.Num(); ++Index)
+	TArray<UUserWidgetExtension*, TInlineAllocator<32>> LocalExtensions;
+	LocalExtensions.Append(Extensions);
+	for (UUserWidgetExtension* Extension : LocalExtensions)
 	{
-		UUserWidgetExtension* Extension = Extensions[Index];
 		check(Extension);
 		Extension->PreConstruct(bIsDesignTime);
 	}
@@ -1827,6 +1830,7 @@ void UUserWidget::NativeDestruct()
 
 	// Extension can remove other extensions.
 	bAreExtensionsConstructed = false; // To prevent calling Destruct on the same extension if it's removed by another extension.
+	bAreExtensionsPreConstructed = false;
 	if (Extensions.Num() > 0)
 	{
 		TArray<UUserWidgetExtension*, TInlineAllocator<32>> LocalExtensions;
@@ -2644,6 +2648,13 @@ UUserWidgetExtension* UUserWidget::AddExtension(TSubclassOf<UUserWidgetExtension
 	{
 		Extension->Initialize();
 	}
+
+	if (bAreExtensionsPreConstructed)
+	{
+		const bool bIsDesignTime = IsDesignTime();
+		Extension->PreConstruct(bIsDesignTime);
+	}
+	
 	if (bAreExtensionsConstructed)
 	{
 		Extension->Construct();
diff --git a/Engine/Source/Runtime/UMG/Public/Blueprint/UserWidget.h b/Engine/Source/Runtime/UMG/Public/Blueprint/UserWidget.h
--- a/Engine/Source/Runtime/UMG/Public/Blueprint/UserWidget.h
+++ b/Engine/Source/Runtime/UMG/Public/Blueprint/UserWidget.h
@@ -1504,6 +1504,9 @@ private:
 	/** Has this widget been initialized by its class yet? */
 	uint8 bInitialized : 1;
 
+	/** Has this widget been preconstructed and we need to call PreConstruct on new extension. */
+	uint8 bAreExtensionsPreConstructed : 1;
+	
 	/** Has this widget been constructed and we need to call Construct on new extension. */
 	uint8 bAreExtensionsConstructed : 1;
 
