First Person: Enabled additional primitive types to work with FP (WidgetComponent, Paper2D, ParticleSystem). Also added a new setting to DynamicMeshBuilder to enable it to use first person rendering. That class could probably benefit from collapsing all the overloaded functions into a single function taking a struct, but that is out of scope of this CL.

diff --git a/Engine/Plugins/2D/Paper2D/Source/Paper2D/Private/PaperRenderSceneProxy.cpp b/Engine/Plugins/2D/Paper2D/Source/Paper2D/Private/PaperRenderSceneProxy.cpp
--- a/Engine/Plugins/2D/Paper2D/Source/Paper2D/Private/PaperRenderSceneProxy.cpp
+++ b/Engine/Plugins/2D/Paper2D/Source/Paper2D/Private/PaperRenderSceneProxy.cpp
@@ -411,6 +411,7 @@ void FPaperRenderSceneProxy::GetNewBatchMeshes(const FSceneView* View, int32 Vie
 
 				Settings.CastShadow = bCastShadow;
 				Settings.bDisableBackfaceCulling = bDrawTwoSided;
+				Settings.bIsFirstPerson = IsFirstPerson();
 
 				FDynamicMeshDrawOffset DrawOffset;
 				DrawOffset.FirstIndex = Batch.VertexOffset;
diff --git a/Engine/Plugins/FX/Niagara/Source/Niagara/Private/NiagaraComponent.cpp b/Engine/Plugins/FX/Niagara/Source/Niagara/Private/NiagaraComponent.cpp
--- a/Engine/Plugins/FX/Niagara/Source/Niagara/Private/NiagaraComponent.cpp
+++ b/Engine/Plugins/FX/Niagara/Source/Niagara/Private/NiagaraComponent.cpp
@@ -437,7 +437,7 @@ TUniformBuffer<FPrimitiveUniformShaderParameters>* FNiagaraSceneProxy::GetCustom
 				.UseVolumetricLightmap(bHasPrecomputedVolumetricLightmap)
 				.UseSingleSampleShadowFromStationaryLights(UseSingleSampleShadowFromStationaryLights())
 				.HasPixelAnimation(AnyMaterialHasPixelAnimation())
-				.IsFirstPerson(bIsFirstPerson);
+				.IsFirstPerson(IsFirstPerson());
 		if ( InstanceBounds.IsValid )
 		{
 			UBBuilder.InstanceLocalBounds(InstanceBounds);
diff --git a/Engine/Source/Runtime/Engine/Private/DynamicMeshBuilder.cpp b/Engine/Source/Runtime/Engine/Private/DynamicMeshBuilder.cpp
--- a/Engine/Source/Runtime/Engine/Private/DynamicMeshBuilder.cpp
+++ b/Engine/Source/Runtime/Engine/Private/DynamicMeshBuilder.cpp
@@ -883,6 +883,7 @@ void FDynamicMeshBuilder::GetMesh(
 					.LocalBounds(FBoxSphereBounds(EForceInit::ForceInit))
 					.ReceivesDecals(Settings.bReceivesDecals)
 					.OutputVelocity(true)
+					.IsFirstPerson(Settings.bIsFirstPerson)
 				.Build();
 
 			
diff --git a/Engine/Source/Runtime/Engine/Private/ParticleSystemSceneProxy.cpp b/Engine/Source/Runtime/Engine/Private/ParticleSystemSceneProxy.cpp
--- a/Engine/Source/Runtime/Engine/Private/ParticleSystemSceneProxy.cpp
+++ b/Engine/Source/Runtime/Engine/Private/ParticleSystemSceneProxy.cpp
@@ -492,6 +492,7 @@ void FParticleSystemSceneProxy::UpdateWorldSpacePrimitiveUniformBuffer(FRHIComma
 				.UseVolumetricLightmap(GetScene().HasPrecomputedVolumetricLightmap_RenderThread())
 				.CustomPrimitiveData(GetCustomPrimitiveData())
 				.HasPixelAnimation(AnyMaterialHasPixelAnimation())
+				.IsFirstPerson(IsFirstPerson())
 			.Build()
 		);
 	}
diff --git a/Engine/Source/Runtime/Engine/Private/PrimitiveSceneProxy.cpp b/Engine/Source/Runtime/Engine/Private/PrimitiveSceneProxy.cpp
--- a/Engine/Source/Runtime/Engine/Private/PrimitiveSceneProxy.cpp
+++ b/Engine/Source/Runtime/Engine/Private/PrimitiveSceneProxy.cpp
@@ -903,7 +903,7 @@ void FPrimitiveSceneProxy::BuildUniformShaderParameters(FPrimitiveUniformShaderP
 			.RayTracingFarField(IsRayTracingFarField())
 			.RayTracingHasGroupId(GetRayTracingGroupId() != FPrimitiveSceneProxy::InvalidRayTracingGroupId)
 			.MeshPaintTextureDescriptor(GetMeshPaintTextureDescriptor())
-			.IsFirstPerson(bIsFirstPerson);
+			.IsFirstPerson(IsFirstPerson());
 
 	if (PrimitiveSceneInfo != nullptr)
 	{
diff --git a/Engine/Source/Runtime/Engine/Public/DynamicMeshBuilder.h b/Engine/Source/Runtime/Engine/Public/DynamicMeshBuilder.h
--- a/Engine/Source/Runtime/Engine/Public/DynamicMeshBuilder.h
+++ b/Engine/Source/Runtime/Engine/Public/DynamicMeshBuilder.h
@@ -122,6 +122,7 @@ struct FDynamicMeshBuilderSettings
 	bool bUseSelectionOutline = true;
 	bool bCanApplyViewModeOverrides = false;
 	bool bUseWireframeSelectionColoring = false;
+	bool bIsFirstPerson = false;
 };
 
 /**
diff --git a/Engine/Source/Runtime/UMG/Private/Components/WidgetComponent.cpp b/Engine/Source/Runtime/UMG/Private/Components/WidgetComponent.cpp
--- a/Engine/Source/Runtime/UMG/Private/Components/WidgetComponent.cpp
+++ b/Engine/Source/Runtime/UMG/Private/Components/WidgetComponent.cpp
@@ -404,6 +404,7 @@ public:
 							Settings.bDisableBackfaceCulling = false;
 							Settings.bReceivesDecals = true;
 							Settings.bUseSelectionOutline = true;
+							Settings.bIsFirstPerson = IsFirstPerson();
 							MeshBuilder.GetMesh(ViewportLocalToWorld, PreviousLocalToWorld, ParentMaterialProxy, SDPG_World, Settings, nullptr, ViewIndex, Collector, FHitProxyId());
 						}
 					}
@@ -486,6 +487,7 @@ public:
 							Settings.bDisableBackfaceCulling = false;
 							Settings.bReceivesDecals = true;
 							Settings.bUseSelectionOutline = true;
+							Settings.bIsFirstPerson = IsFirstPerson();
 							MeshBuilder.GetMesh(ViewportLocalToWorld, PreviousLocalToWorld, ParentMaterialProxy, SDPG_World, Settings, nullptr, ViewIndex, Collector, FHitProxyId());
 						}
 					}
