Removed  redundant deprecation guard.

#rnx

diff --git a/Engine/Source/Runtime/UMG/Private/Components/WidgetComponent.cpp b/Engine/Source/Runtime/UMG/Private/Components/WidgetComponent.cpp
--- a/Engine/Source/Runtime/UMG/Private/Components/WidgetComponent.cpp
+++ b/Engine/Source/Runtime/UMG/Private/Components/WidgetComponent.cpp
@@ -1784,9 +1784,7 @@ void UWidgetComponent::SetSlateWidget(const TSharedPtr<SWidget>& InSlateWidget)
 void UWidgetComponent::UpdateWidget()
 {
 	// Don't do any work if Slate is not initialized
-	PRAGMA_DISABLE_DEPRECATION_WARNINGS
 	if (FSlateApplication::IsInitialized() && IsValidChecked(this))
-	PRAGMA_ENABLE_DEPRECATION_WARNINGS
 	{
 		if (Space == EWidgetSpace::World)
 		{
