Fix to Crash when adding the first UI Component to a Widget Blueprint

diff --git a/Engine/Source/Runtime/UMG/Private/Extensions/UIComponent.cpp b/Engine/Source/Runtime/UMG/Private/Extensions/UIComponent.cpp
--- a/Engine/Source/Runtime/UMG/Private/Extensions/UIComponent.cpp
+++ b/Engine/Source/Runtime/UMG/Private/Extensions/UIComponent.cpp
@@ -5,11 +5,11 @@
 #include "CoreGlobals.h"
 #include "Components/Widget.h"
 #include "Blueprint/UserWidget.h"
-//#include "UMGPrivate.h"
 
 void UUIComponent::Initialize(UWidget* Target)
 {
 	ensure(Target);
+	ensure(Owner.IsExplicitlyNull());
 	Owner = Target;
 	OnInitialize();
 }
diff --git a/Engine/Source/Runtime/UMG/Private/Extensions/UIComponentUserWidgetExtension.cpp b/Engine/Source/Runtime/UMG/Private/Extensions/UIComponentUserWidgetExtension.cpp
--- a/Engine/Source/Runtime/UMG/Private/Extensions/UIComponentUserWidgetExtension.cpp
+++ b/Engine/Source/Runtime/UMG/Private/Extensions/UIComponentUserWidgetExtension.cpp
@@ -10,45 +10,51 @@
 
 void UUIComponentUserWidgetExtension::PreConstruct(bool bIsDesignTime)
 {
-	Super::Construct();
+	Super::PreConstruct(bIsDesignTime);
 
-	InitializeComponents();
-	
-	ComponentContainer->ForEachComponent([bIsDesignTime](UUIComponent* Component)
-		{
-			if(ensure(Component))
+	if (ComponentContainer)
+	{		
+		InitializeComponents();
+		ComponentContainer->ForEachComponent([bIsDesignTime](UUIComponent* Component)
 			{
-				Component->PreConstruct(bIsDesignTime);
+				if(ensure(Component))
+				{
+					Component->PreConstruct(bIsDesignTime);
+				}
 			}
-		}
-	);
+		);
+	}
 }
 void UUIComponentUserWidgetExtension::Construct()
 {
 	Super::Construct();
-
-	ComponentContainer->ForEachComponent([](UUIComponent* Component)
-		{
-			if (ensure(Component))
+	if (ComponentContainer)
+	{
+		ComponentContainer->ForEachComponent([](UUIComponent* Component)
 			{
-				Component->Construct();
+				if (ensure(Component))
+				{
+					Component->Construct();
+				}
 			}
-		}
-	);
+		);
+	}
 }
 
 void UUIComponentUserWidgetExtension::Destruct()
 {	
 	Super::Destruct();
-
-	ComponentContainer->ForEachComponent([](UUIComponent* Component)
-		{
-			if (ensure(Component))
+	if (ComponentContainer)
+	{
+		ComponentContainer->ForEachComponent([](UUIComponent* Component)
 			{
-				Component->Destruct();
+				if (ensure(Component))
+				{
+					Component->Destruct();
+				}
 			}
-		}
-	);
+		);
+	}
 }
 
 void UUIComponentUserWidgetExtension::InitializeContainer(UUIComponentContainer* InComponentContainer)
@@ -60,11 +66,10 @@ void UUIComponentUserWidgetExtension::InitializeContainer(UUIComponentContainer*
 
 void UUIComponentUserWidgetExtension::InitializeComponents()
 {
-	const UUserWidget* UserWidget = GetUserWidget();
-	check(UserWidget);
-	check(ComponentContainer);
 	if (ComponentContainer)
 	{
+		const UUserWidget* UserWidget = GetUserWidget();
+		check(UserWidget);
 		ComponentContainer->InitializeComponents(UserWidget);
 	}
 }
