UIComponent: change function signature to match the declaration

diff --git a/Engine/Source/Runtime/UMG/Private/Extensions/UIComponent.cpp b/Engine/Source/Runtime/UMG/Private/Extensions/UIComponent.cpp
--- a/Engine/Source/Runtime/UMG/Private/Extensions/UIComponent.cpp
+++ b/Engine/Source/Runtime/UMG/Private/Extensions/UIComponent.cpp
@@ -78,7 +78,7 @@ bool UUIComponent::RemoveFieldValueChangedDelegate(UE::FieldNotification::FField
 	return false;
 }
 
-int32 UUIComponent::RemoveAllFieldValueChangedDelegates(const void* InUserObject)
+int32 UUIComponent::RemoveAllFieldValueChangedDelegates(FDelegateUserObjectConst InUserObject)
 {
 	if (InUserObject)
 	{
@@ -87,7 +87,7 @@ int32 UUIComponent::RemoveAllFieldValueChangedDelegates(const void* InUserObject
 	return false;
 }
 
-int32 UUIComponent::RemoveAllFieldValueChangedDelegates(UE::FieldNotification::FFieldId InFieldId, const void* InUserObject)
+int32 UUIComponent::RemoveAllFieldValueChangedDelegates(UE::FieldNotification::FFieldId InFieldId, FDelegateUserObjectConst InUserObject)
 {
 	if (InFieldId.IsValid() && InUserObject)
 	{
diff --git a/Engine/Source/Runtime/UMG/Public/Extensions/UIComponent.h b/Engine/Source/Runtime/UMG/Public/Extensions/UIComponent.h
--- a/Engine/Source/Runtime/UMG/Public/Extensions/UIComponent.h
+++ b/Engine/Source/Runtime/UMG/Public/Extensions/UIComponent.h
@@ -3,6 +3,8 @@
 #pragma once
 
 #include "CoreMinimal.h"
+#include "Delegates/IDelegateInstance.h"
+#include "FieldNotificationDeclaration.h"
 #include "FieldNotificationDelegate.h"
 #include "UObject/ObjectMacros.h"
 #include "IFieldNotificationClassDescriptor.h"
