UMG: PR #12826: Slate - Added properly working absolute <-> local transform BP functions

#rnx

diff --git a/Engine/Source/Runtime/UMG/Private/SlateBlueprintLibrary.cpp b/Engine/Source/Runtime/UMG/Private/SlateBlueprintLibrary.cpp
--- a/Engine/Source/Runtime/UMG/Private/SlateBlueprintLibrary.cpp
+++ b/Engine/Source/Runtime/UMG/Private/SlateBlueprintLibrary.cpp
@@ -143,6 +143,26 @@ void USlateBlueprintLibrary::ScreenToViewport(UObject* WorldContextObject, FVect
 	USlateBlueprintLibrary::AbsoluteToViewport(WorldContextObject, AbsolutePosition, PixelPosition, ViewportPosition);
 }
 
+float USlateBlueprintLibrary::Scalar_AbsoluteToLocal(const FGeometry& Geometry, float AbsoluteScalar)
+{
+	return Inverse(Geometry.GetAccumulatedRenderTransform()).TransformVector(FVector2f(AbsoluteScalar, 0.0f)).Size();
+}
+
+float USlateBlueprintLibrary::Scalar_LocalToAbsolute(const FGeometry& Geometry, float LocalScalar)
+{
+	return Geometry.GetAccumulatedRenderTransform().TransformVector(FVector2f(LocalScalar, 0.0f)).Size();
+}
+
+FVector2D USlateBlueprintLibrary::Vector_AbsoluteToLocal(const FGeometry& Geometry, FVector2D AbsoluteVector)
+{
+	return Inverse(Geometry.GetAccumulatedRenderTransform()).TransformVector(AbsoluteVector);
+}
+
+FVector2D USlateBlueprintLibrary::Vector_LocalToAbsolute(const FGeometry& Geometry, FVector2D LocalVector)
+{
+	return Geometry.GetAccumulatedRenderTransform().TransformVector(LocalVector);
+}
+
 float USlateBlueprintLibrary::TransformScalarAbsoluteToLocal(const FGeometry& Geometry, float AbsoluteScalar)
 {
 	return Geometry.GetAccumulatedRenderTransform().TransformVector(FVector2f(AbsoluteScalar, 0.0f)).Size();
diff --git a/Engine/Source/Runtime/UMG/Public/Blueprint/SlateBlueprintLibrary.h b/Engine/Source/Runtime/UMG/Public/Blueprint/SlateBlueprintLibrary.h
--- a/Engine/Source/Runtime/UMG/Public/Blueprint/SlateBlueprintLibrary.h
+++ b/Engine/Source/Runtime/UMG/Public/Blueprint/SlateBlueprintLibrary.h
@@ -26,21 +26,23 @@ public:
 	static UMG_API bool IsUnderLocation(const FGeometry& Geometry, const FVector2D& AbsoluteCoordinate);
 
 	/**
+	 * Transforms absolute coordinates into local coordinates
+	 *
 	 * Absolute coordinates could be either desktop or window space depending on what space the root of the widget hierarchy is in.
 	 *
 	 * @return Transforms AbsoluteCoordinate into the local space of this Geometry.
 	 */
-	UFUNCTION(BlueprintPure, Category="User Interface|Geometry")
+	UFUNCTION(BlueprintPure, Category="User Interface|Geometry", meta = (DisplayName = "Absolute to Local (Coordinates)", Keywords = "transform"))
 	static UMG_API FVector2D AbsoluteToLocal(const FGeometry& Geometry, FVector2D AbsoluteCoordinate);
 
 	/**
-	 * Translates local coordinates into absolute coordinates
+	 * Transforms local coordinates into absolute coordinates
 	 *
 	 * Absolute coordinates could be either desktop or window space depending on what space the root of the widget hierarchy is in.
 	 *
 	 * @return  Absolute coordinates
 	 */
-	UFUNCTION(BlueprintPure, Category="User Interface|Geometry")
+	UFUNCTION(BlueprintPure, Category="User Interface|Geometry", meta = (DisplayName = "Local to Absolute (Coordinates)", Keywords = "transform"))
 	static UMG_API FVector2D LocalToAbsolute(const FGeometry& Geometry, FVector2D LocalCoordinate);
 
 	/** Returns the local top/left of the geometry in local space. */
@@ -55,20 +57,40 @@ public:
 	UFUNCTION(BlueprintPure, Category="User Interface|Geometry")
 	static UMG_API FVector2D GetAbsoluteSize(const FGeometry& Geometry);
 
+	/** Transforms a scalar from absolute space to local space  */
+	UFUNCTION(BlueprintPure, Category = "User Interface|Geometry", meta = (DisplayName = "Absolute to Local (Scalar)", Keywords = "transform"))
+	static UMG_API float Scalar_AbsoluteToLocal(const FGeometry& Geometry, float AbsoluteScalar);
+
+	/** Transforms a scalar from local space to absolute space  */
+	UFUNCTION(BlueprintPure, Category = "User Interface|Geometry", meta = (DisplayName = "Local to Absolute (Scalar)", Keywords = "transform"))
+	static UMG_API float Scalar_LocalToAbsolute(const FGeometry& Geometry, float LocalScalar);
+
+	/** Transforms a vector from absolute space to local space  */
+	UFUNCTION(BlueprintPure, Category = "User Interface|Geometry", meta = (DisplayName = "Absolute to Local (Vector)", Keywords = "transform"))
+	static UMG_API FVector2D Vector_AbsoluteToLocal(const FGeometry& Geometry, FVector2D AbsoluteVector);
+
+	/** Transforms a vector from local space to absolute space  */
+	UFUNCTION(BlueprintPure, Category = "User Interface|Geometry", meta = (DisplayName = "Local to Absolute (Vector)", Keywords = "transform"))
+	static UMG_API FVector2D Vector_LocalToAbsolute(const FGeometry& Geometry, FVector2D LocalVector);
+
 	/**  */
-	UFUNCTION(BlueprintPure, Category = "User Interface|Geometry")
+	UE_DEPRECATED(5.6, "TransformScalarAbsoluteToLocal returns inverted results and has been replaced by Scalar_LocalToAbsolute")
+	UFUNCTION(BlueprintPure, Category = "User Interface|Geometry", meta = (DeprecatedFunction, DeprecationMessage = "Returns inverted results. Replace with 'Local to Absolute (Scalar)'"))
 	static UMG_API float TransformScalarAbsoluteToLocal(const FGeometry& Geometry, float AbsoluteScalar);
 
 	/**  */
-	UFUNCTION(BlueprintPure, Category = "User Interface|Geometry")
+	UE_DEPRECATED(5.6, "TransformScalarLocalToAbsolute returns inverted results and has been replaced by Scalar_AbsoluteToLocal")
+	UFUNCTION(BlueprintPure, Category = "User Interface|Geometry", meta = (DeprecatedFunction, DeprecationMessage = "Returns inverted results. Replace with 'Absolute To Local (Scalar)'"))
 	static UMG_API float TransformScalarLocalToAbsolute(const FGeometry& Geometry, float LocalScalar);
 
 	/**  */
-	UFUNCTION(BlueprintPure, Category = "User Interface|Geometry")
+	UE_DEPRECATED(5.6, "TransformVectorAbsoluteToLocal returns inverted results and has been replaced by Vector_LocalToAbsolute")
+	UFUNCTION(BlueprintPure, Category = "User Interface|Geometry", meta = (DeprecatedFunction, DeprecationMessage = "Returns inverted results. Replace with 'Local to Absolute (Vector)'"))
 	static UMG_API FVector2D TransformVectorAbsoluteToLocal(const FGeometry& Geometry, FVector2D AbsoluteVector);
 
 	/**  */
-	UFUNCTION(BlueprintPure, Category = "User Interface|Geometry")
+	UE_DEPRECATED(5.6, "TransformVectorLocalToAbsolute returns inverted results and has been replaced by Vector_AbsoluteToLocal")
+	UFUNCTION(BlueprintPure, Category = "User Interface|Geometry", meta = (DeprecatedFunction, DeprecationMessage = "Returns inverted results. Replace with 'Absolute To Local (Vector)'"))
 	static UMG_API FVector2D TransformVectorLocalToAbsolute(const FGeometry& Geometry, FVector2D LocalVector);
 
 	/** Returns whether brushes A and B are identical. */
