ComboBox Widgets: Fixed potential null access when widget becomes invalid as a result of OnSelectionChanged

diff --git a/Engine/Source/Runtime/UMG/Private/Components/ComboBoxKey.cpp b/Engine/Source/Runtime/UMG/Private/Components/ComboBoxKey.cpp
--- a/Engine/Source/Runtime/UMG/Private/Components/ComboBoxKey.cpp
+++ b/Engine/Source/Runtime/UMG/Private/Components/ComboBoxKey.cpp
@@ -286,41 +286,44 @@ PRAGMA_ENABLE_DEPRECATION_WARNINGS
 
 void UComboBoxKey::GenerateContent()
 {
-	bool bSetToNull = true;
-	if (IsDesignTime())
+	if (ComboBoxContent.IsValid())
 	{
-		if (!SelectedOption.IsNone())
+		bool bSetToNull = true;
+		if (IsDesignTime())
 		{
-			ComboBoxContent->SetContent(SNew(STextBlock)
-				.Text(FText::FromName(SelectedOption)));
-			bSetToNull = false;
-		}
-	}
-	else
-	{
-		// Call the user's delegate to see if they want to generate a custom widget bound to the data source.
-		if (OnGenerateContentWidget.IsBound())
-		{
-			if (UWidget* Widget = OnGenerateContentWidget.Execute(SelectedOption))
+			if (!SelectedOption.IsNone())
 			{
-				ComboBoxContent->SetContent(Widget->TakeWidget());
+				ComboBoxContent->SetContent(SNew(STextBlock)
+					.Text(FText::FromName(SelectedOption)));
 				bSetToNull = false;
 			}
 		}
 		else
 		{
-			// Warn the user that they need to generate a widget for the item
-			ComboBoxContent->SetContent(
-				SNew(STextBlock)
-				.Text(FText::Format(LOCTEXT("WarningMissingOnGenerateContentWidget", "{0} Missing OnGenerateContentWidget"), FText::FromName(SelectedOption)))
+			// Call the user's delegate to see if they want to generate a custom widget bound to the data source.
+			if (OnGenerateContentWidget.IsBound())
+			{
+				if (UWidget* Widget = OnGenerateContentWidget.Execute(SelectedOption))
+				{
+					ComboBoxContent->SetContent(Widget->TakeWidget());
+					bSetToNull = false;
+				}
+			}
+			else
+			{
+				// Warn the user that they need to generate a widget for the item
+				ComboBoxContent->SetContent(
+					SNew(STextBlock)
+					.Text(FText::Format(LOCTEXT("WarningMissingOnGenerateContentWidget", "{0} Missing OnGenerateContentWidget"), FText::FromName(SelectedOption)))
 				);
-			bSetToNull = false;
+				bSetToNull = false;
+			}
 		}
-	}
 
-	if (bSetToNull)
-	{
-		ComboBoxContent->SetContent(SNullWidget::NullWidget);
+		if (bSetToNull)
+		{
+			ComboBoxContent->SetContent(SNullWidget::NullWidget);
+		}
 	}
 }
 
diff --git a/Engine/Source/Runtime/UMG/Private/Components/ComboBoxString.cpp b/Engine/Source/Runtime/UMG/Private/Components/ComboBoxString.cpp
--- a/Engine/Source/Runtime/UMG/Private/Components/ComboBoxString.cpp
+++ b/Engine/Source/Runtime/UMG/Private/Components/ComboBoxString.cpp
@@ -421,7 +421,7 @@ void UComboBoxString::UpdateOrGenerateWidget(TSharedPtr<FString> Item)
 		const FString StringItem = Item.IsValid() ? *Item : FString();
 		DefaultComboBoxContent.Pin()->SetText(FText::FromString(StringItem));
 	}
-	else
+	else if (ComboBoxContent.IsValid())
 	{
 		DefaultComboBoxContent.Reset();
 		ComboBoxContent->SetContent(HandleGenerateWidget(Item));
