UMG: Fix SetAnimationCurrentTime not working because UnwarpedPosition also needs to be set/reset.

diff --git a/Engine/Source/Runtime/UMG/Public/Animation/UMGSequencePlayer.h b/Engine/Source/Runtime/UMG/Public/Animation/UMGSequencePlayer.h
--- a/Engine/Source/Runtime/UMG/Public/Animation/UMGSequencePlayer.h
+++ b/Engine/Source/Runtime/UMG/Public/Animation/UMGSequencePlayer.h
@@ -40,7 +40,7 @@ public:
 	/** Reverses a running animation */
 	UMG_API void Reverse();
 
-	void SetCurrentTime(float InTime) { TimeCursorPosition = AnimationResolution.AsFrameTime(InTime); }
+	void SetCurrentTime(float InTime) { TimeCursorPosition = UnwarpedPosition = AnimationResolution.AsFrameTime(InTime); }
 	FQualifiedFrameTime GetCurrentTime() const { return FQualifiedFrameTime(TimeCursorPosition, AnimationResolution); }
 
 	/** @return The current animation being played */
