Widget Component: Fixes for Remote Player Widget Components on Listen Servers when accessing Local Players
PR #12990

diff --git a/Engine/Source/Runtime/UMG/Private/Components/WidgetComponent.cpp b/Engine/Source/Runtime/UMG/Private/Components/WidgetComponent.cpp
--- a/Engine/Source/Runtime/UMG/Private/Components/WidgetComponent.cpp
+++ b/Engine/Source/Runtime/UMG/Private/Components/WidgetComponent.cpp
@@ -1726,7 +1726,10 @@ ULocalPlayer* UWidgetComponent::GetOwnerPlayer() const
 	{
 		if (APlayerController* OwningController = Cast<APlayerController>(OwningPawn->GetController()))
 		{
-			return OwningController->GetLocalPlayer();
+			if (ULocalPlayer* OwningLocalPlayer = OwningController->GetLocalPlayer())
+			{
+				return OwningLocalPlayer;
+			}
 		}
 	}
 	
