[SCommonButtonTableRow] Update the selection state based on the OwnerTable selected items in situations where the SObjectTableRow selection check missed it.

#rnx

diff --git a/Engine/Plugins/Runtime/CommonUI/Source/CommonUI/Public/SCommonButtonTableRow.h b/Engine/Plugins/Runtime/CommonUI/Source/CommonUI/Public/SCommonButtonTableRow.h
--- a/Engine/Plugins/Runtime/CommonUI/Source/CommonUI/Public/SCommonButtonTableRow.h
+++ b/Engine/Plugins/Runtime/CommonUI/Source/CommonUI/Public/SCommonButtonTableRow.h
@@ -106,6 +106,24 @@ protected:
 		}
 	}
 
+	virtual void DetectItemSelectionChanged() override
+	{
+		SObjectTableRow<ItemType>::DetectItemSelectionChanged();
+		
+		TSharedRef<ITypedTableView<ItemType>> OwnerTable = this->OwnerTablePtr.Pin().ToSharedRef();
+		if (const TObjectPtrWrapTypeOf<ItemType>* MyItemPtr = this->GetItemForThis(OwnerTable))
+		{
+			// Selection changes at the list level can happen directly or in response to another item being selected.
+			// Regardless, just make sure the button's selection state is inline with the item's
+			const UCommonButtonBase* CommonButton = Cast<UCommonButtonBase>(this->WidgetObject);
+			const bool bIsItemSelected = OwnerTable->Private_IsItemSelected(*MyItemPtr);
+			if (CommonButton && CommonButton->GetSelected() != bIsItemSelected)
+			{
+				OnItemSelectionChanged(bIsItemSelected);
+			}
+		}
+	}
+
 	virtual void OnItemSelectionChanged(bool bIsItemSelected) override
 	{
 		SObjectTableRow<ItemType>::OnItemSelectionChanged(bIsItemSelected);
diff --git a/Engine/Source/Runtime/UMG/Public/Slate/SObjectTableRow.h b/Engine/Source/Runtime/UMG/Public/Slate/SObjectTableRow.h
--- a/Engine/Source/Runtime/UMG/Public/Slate/SObjectTableRow.h
+++ b/Engine/Source/Runtime/UMG/Public/Slate/SObjectTableRow.h
@@ -117,21 +117,7 @@ public:
 
 	EActiveTimerReturnType DetectItemSelectionChanged(double InCurrentTime, float InDeltaTime)
 	{
-		// List views were built assuming the use of attributes on rows to check on selection status, so there is no
-		// clean way to inform individual rows of changes to the selection state of their current items.
-		// Since event-based selection changes are only really needed in a game scenario, we (crudely) monitor it here to generate events.
-		// If desired, per-item selection events could be added as a longer-term todo
-		TSharedPtr<ITypedTableView<ItemType>> OwnerTable = OwnerTablePtr.Pin();
-		const TObjectPtrWrapTypeOf<ItemType>* MyItemPtr = OwnerTable ? GetItemForThis(OwnerTable.ToSharedRef()) : nullptr;
-		if (MyItemPtr)
-		{
-			if (bIsAppearingSelected != OwnerTable->Private_IsItemSelected(*MyItemPtr))
-			{
-				bIsAppearingSelected = !bIsAppearingSelected;
-				OnItemSelectionChanged(bIsAppearingSelected);
-			}
-		}
-
+		DetectItemSelectionChanged();
 		return EActiveTimerReturnType::Continue;
 	}
 
@@ -520,6 +506,24 @@ protected:
 		}
 	}
 
+	virtual void DetectItemSelectionChanged()
+	{
+		// List views were built assuming the use of attributes on rows to check on selection status, so there is no
+		// clean way to inform individual rows of changes to the selection state of their current items.
+		// Since event-based selection changes are only really needed in a game scenario, we (crudely) monitor it here to generate events.
+		// If desired, per-item selection events could be added as a longer-term todo
+		TSharedPtr<ITypedTableView<ItemType>> OwnerTable = OwnerTablePtr.Pin();
+		const TObjectPtrWrapTypeOf<ItemType>* MyItemPtr = OwnerTable ? GetItemForThis(OwnerTable.ToSharedRef()) : nullptr;
+		if (MyItemPtr)
+		{
+			if (bIsAppearingSelected != OwnerTable->Private_IsItemSelected(*MyItemPtr))
+			{
+				bIsAppearingSelected = !bIsAppearingSelected;
+				OnItemSelectionChanged(bIsAppearingSelected);
+			}
+		}
+	}
+
 	virtual void OnItemSelectionChanged(bool bIsItemSelected)
 	{
 		if (WidgetObject)
