[Engine] Added new variables to allow individual EditableText to dictate whether to use the IntegratedKeyboard or not.
[RN] Added "EnableIntegratedKeyboard" variable to EditableText. This allow the Integrated Keyboard to be enabled per EditableText.

diff --git a/Engine/Source/Runtime/ApplicationCore/Private/IOS/IOSView.cpp b/Engine/Source/Runtime/ApplicationCore/Private/IOS/IOSView.cpp
--- a/Engine/Source/Runtime/ApplicationCore/Private/IOS/IOSView.cpp
+++ b/Engine/Source/Runtime/ApplicationCore/Private/IOS/IOSView.cpp
@@ -749,6 +749,15 @@ self.accessibilityElements = @[Window.accessibilityContainer];
 	});
 }
 
+-(bool)ShouldUseIntegratedKeyboard
+{
+	bool bUseIntegratedKeyboard = false;
+	GConfig->GetBool(TEXT("/Script/IOSRuntimeSettings.IOSRuntimeSettings"), TEXT("bUseIntegratedKeyboard"), bUseIntegratedKeyboard, GEngineIni);
+	
+	// get notifications when the keyboard is in view
+	return FParse::Param(FCommandLine::Get(), TEXT("NewKeyboard")) || bUseIntegratedKeyboard;
+}
+
 -(BOOL)becomeFirstResponder
 {
 	volatile int32 ShowCount = KeyboardShowCount;
@@ -1032,15 +1041,12 @@ self.accessibilityElements = @[Window.accessibilityContainer];
 {
 #if !PLATFORM_TVOS
 	KeyboardShowCount = 0;
-	
-	bool bUseIntegratedKeyboard = false;
-	GConfig->GetBool(TEXT("/Script/IOSRuntimeSettings.IOSRuntimeSettings"), TEXT("bUseIntegratedKeyboard"), bUseIntegratedKeyboard, GEngineIni);
 
 	bool bEnableVirtualKeyboardVisibilityEvent = false;
 	GConfig->GetBool(TEXT("/Script/IOSRuntimeSettings.IOSRuntimeSettings"), TEXT("bEnableVirtualKeyboardVisibilityEvent"), bEnableVirtualKeyboardVisibilityEvent, GEngineIni);
-	
-	// get notifications when the keyboard is in view
-	bIsUsingIntegratedKeyboard = FParse::Param(FCommandLine::Get(), TEXT("NewKeyboard")) || bUseIntegratedKeyboard;
+
+	bIsUsingIntegratedKeyboard = [self ShouldUseIntegratedKeyboard];
+
 	if (bEnableVirtualKeyboardVisibilityEvent || bIsUsingIntegratedKeyboard)
 	{
 		[[NSNotificationCenter defaultCenter] addObserver:self
diff --git a/Engine/Source/Runtime/ApplicationCore/Public/IOS/IOSView.h b/Engine/Source/Runtime/ApplicationCore/Public/IOS/IOSView.h
--- a/Engine/Source/Runtime/ApplicationCore/Public/IOS/IOSView.h
+++ b/Engine/Source/Runtime/ApplicationCore/Public/IOS/IOSView.h
@@ -109,6 +109,7 @@ APPLICATIONCORE_API
 -(void)ActivateKeyboard:(bool)bInSendEscapeOnClose;
 -(void)ActivateKeyboard:(bool)bInSendEscapeOnClose keyboardConfig:(FKeyboardConfig)KeyboardConfig;
 -(void)DeactivateKeyboard;
+-(bool)ShouldUseIntegratedKeyboard;
 
 // callable from outside to fake locations
 -(void)HandleTouchAtLoc:(CGPoint)Loc PrevLoc:(CGPoint)PrevLoc TouchIndex:(int)TouchIndex Force:(float)Force Type:(TouchType)Type TouchesArray:(TArray<TouchInput>&)TouchesArray;
diff --git a/Engine/Source/Runtime/IOS/IOSRuntimeSettings/Classes/IOSRuntimeSettings.h b/Engine/Source/Runtime/IOS/IOSRuntimeSettings/Classes/IOSRuntimeSettings.h
--- a/Engine/Source/Runtime/IOS/IOSRuntimeSettings/Classes/IOSRuntimeSettings.h
+++ b/Engine/Source/Runtime/IOS/IOSRuntimeSettings/Classes/IOSRuntimeSettings.h
@@ -528,6 +528,10 @@ public:
 	// Whether or not we should receive the keyboard's visibility event (Shown/Hidden) regardless of the IntegratedKeyboard being activated or not.
 	UPROPERTY(GlobalConfig, EditAnywhere, Category = Input)
 	bool bEnableVirtualKeyboardVisibilityEvent;
+	
+	// Whether or not the EditableText Widget is allowed to enable IntegratedKeyboard. bUseIntegratedKeyboard enabled will disregard this setting.
+	UPROPERTY(GlobalConfig, EditAnywhere, Meta=(EditCondition = "!bUseIntegratedKeyboard"), Category = Input)
+	bool bAllowWidgetEnablingIntegratedKeyboard;
 
 	/** Sample rate to run the audio mixer with. */
 	UPROPERTY(config, EditAnywhere, Category = "Audio", Meta = (DisplayName = "Audio Mixer Sample Rate"))
diff --git a/Engine/Source/Runtime/Slate/Private/Framework/Text/IOS/IOSPlatformTextField.cpp b/Engine/Source/Runtime/Slate/Private/Framework/Text/IOS/IOSPlatformTextField.cpp
--- a/Engine/Source/Runtime/Slate/Private/Framework/Text/IOS/IOSPlatformTextField.cpp
+++ b/Engine/Source/Runtime/Slate/Private/Framework/Text/IOS/IOSPlatformTextField.cpp
@@ -67,6 +67,19 @@ FIOSPlatformTextField::~FIOSPlatformTextField()
 void FIOSPlatformTextField::ShowVirtualKeyboard(bool bShow, int32 UserIndex, TSharedPtr<IVirtualKeyboardEntry> TextEntryWidget)
 {
 	FIOSView* View = [IOSAppDelegate GetDelegate].IOSView;
+
+	bool bAllowWidgetEnablingIntegratedKeyboard = false;
+	GConfig->GetBool(TEXT("/Script/IOSRuntimeSettings.IOSRuntimeSettings"), TEXT("bAllowWidgetEnablingIntegratedKeyboard"), bAllowWidgetEnablingIntegratedKeyboard, GEngineIni);
+
+	bool bShouldUseIntegratedKeyboardGlobally = [View ShouldUseIntegratedKeyboard];
+
+	// Prevent any changes to View->bIsUsingIntegratedKeyboard in the event bShouldUseIntegratedKeyboardGlobally is enabled as it is already intialized during the KeyboardInitialization and only needs to be done once throughout the app lifetime.
+	// bShouldUseIntegratedKeyboardGlobally should always have priority over individually toggled asset (bAllowWidgetEnablingIntegratedKeyboard).
+	if (bShow && !bShouldUseIntegratedKeyboardGlobally && bAllowWidgetEnablingIntegratedKeyboard)
+	{
+		View->bIsUsingIntegratedKeyboard = TextEntryWidget.IsValid() ? TextEntryWidget->IsIntegratedKeyboardEnabled() : false;
+	}
+	 
 	if (View->bIsUsingIntegratedKeyboard)
 	{
 		if (bShow)
diff --git a/Engine/Source/Runtime/Slate/Private/Widgets/Input/SEditableText.cpp b/Engine/Source/Runtime/Slate/Private/Widgets/Input/SEditableText.cpp
--- a/Engine/Source/Runtime/Slate/Private/Widgets/Input/SEditableText.cpp
+++ b/Engine/Source/Runtime/Slate/Private/Widgets/Input/SEditableText.cpp
@@ -47,6 +47,7 @@ void SEditableText::Construct( const FArguments& InArgs )
 	VirtualKeyboardDismissAction = InArgs._VirtualKeyboardDismissAction;
 	OnKeyCharHandler = InArgs._OnKeyCharHandler;
 	OnKeyDownHandler = InArgs._OnKeyDownHandler;
+	bEnableIntegratedKeyboard = InArgs._EnableIntegratedKeyboard;
 
 	Font = InArgs._Font;
 	ColorAndOpacity = InArgs._ColorAndOpacity;
@@ -389,6 +390,11 @@ void SEditableText::SetAllowContextMenu(const TAttribute< bool >& InAllowContext
 	bAllowContextMenu = InAllowContextMenu;
 }
 
+void SEditableText::SetEnableIntegratedKeyboard(const TAttribute<bool>& InEnableIntegratedKeyboard)
+{
+	bEnableIntegratedKeyboard = InEnableIntegratedKeyboard;
+}
+
 void SEditableText::SetVirtualKeyboardDismissAction(TAttribute< EVirtualKeyboardDismissAction > InVirtualKeyboardDismissAction)
 {
 	VirtualKeyboardDismissAction = InVirtualKeyboardDismissAction;
@@ -535,6 +541,11 @@ bool SEditableText::IsMultiLineTextEdit() const
 	return false;
 }
 
+bool SEditableText::IsIntegratedKeyboardEnabled() const
+{
+	return bEnableIntegratedKeyboard.Get(false);
+}
+
 bool SEditableText::ShouldJumpCursorToEndWhenFocused() const
 {
 	return bIsCaretMovedWhenGainFocus.Get(false);
diff --git a/Engine/Source/Runtime/Slate/Private/Widgets/Input/SVirtualKeyboardEntry.cpp b/Engine/Source/Runtime/Slate/Private/Widgets/Input/SVirtualKeyboardEntry.cpp
--- a/Engine/Source/Runtime/Slate/Private/Widgets/Input/SVirtualKeyboardEntry.cpp
+++ b/Engine/Source/Runtime/Slate/Private/Widgets/Input/SVirtualKeyboardEntry.cpp
@@ -33,6 +33,7 @@ void SVirtualKeyboardEntry::Construct( const FArguments& InArgs )
 	MinDesiredWidth = InArgs._MinDesiredWidth;
 	KeyboardType = InArgs._KeyboardType;
 	VirtualKeyboardOptions = InArgs._VirtualKeyboardOptions;
+	bEnableIntegratedKeyboard = InArgs._EnableIntegratedKeyboard;
 }
 
 void SVirtualKeyboardEntry::SetText(const TAttribute< FText >& InNewText)
diff --git a/Engine/Source/Runtime/Slate/Private/Widgets/Text/SMultiLineEditableText.cpp b/Engine/Source/Runtime/Slate/Private/Widgets/Text/SMultiLineEditableText.cpp
--- a/Engine/Source/Runtime/Slate/Private/Widgets/Text/SMultiLineEditableText.cpp
+++ b/Engine/Source/Runtime/Slate/Private/Widgets/Text/SMultiLineEditableText.cpp
@@ -298,6 +298,11 @@ bool SMultiLineEditableText::IsMultiLineTextEdit() const
 	return bAllowMultiLine.Get(true);
 }
 
+bool SMultiLineEditableText::IsIntegratedKeyboardEnabled() const
+{
+	return false;
+}
+
 bool SMultiLineEditableText::ShouldJumpCursorToEndWhenFocused() const
 {
 	return false;
diff --git a/Engine/Source/Runtime/Slate/Private/Widgets/Text/SlateEditableTextLayout.cpp b/Engine/Source/Runtime/Slate/Private/Widgets/Text/SlateEditableTextLayout.cpp
--- a/Engine/Source/Runtime/Slate/Private/Widgets/Text/SlateEditableTextLayout.cpp
+++ b/Engine/Source/Runtime/Slate/Private/Widgets/Text/SlateEditableTextLayout.cpp
@@ -3865,6 +3865,13 @@ bool FSlateEditableTextLayout::FVirtualKeyboardEntry::IsMultilineEntry() const
 	return OwnerLayout->OwnerWidget->IsMultiLineTextEdit();
 }
 
+bool FSlateEditableTextLayout::FVirtualKeyboardEntry::IsIntegratedKeyboardEnabled() const
+{
+	check(IsInGameThread());
+	
+	return OwnerLayout->OwnerWidget->IsIntegratedKeyboardEnabled();
+}
+
 
 TSharedRef<FSlateEditableTextLayout::FTextInputMethodContext> FSlateEditableTextLayout::FTextInputMethodContext::Create(FSlateEditableTextLayout& InOwnerLayout)
 {
diff --git a/Engine/Source/Runtime/Slate/Public/Widgets/Input/IVirtualKeyboardEntry.h b/Engine/Source/Runtime/Slate/Public/Widgets/Input/IVirtualKeyboardEntry.h
--- a/Engine/Source/Runtime/Slate/Public/Widgets/Input/IVirtualKeyboardEntry.h
+++ b/Engine/Source/Runtime/Slate/Public/Widgets/Input/IVirtualKeyboardEntry.h
@@ -101,4 +101,11 @@ public:
 	* @return Whether the entry is multi-line
 	*/
 	virtual bool IsMultilineEntry() const = 0;
+	
+	/**
+    * Returns whether the entry is using IntegratedKeyboard
+    *
+    * @return Whether the entry is using IntegratedKeyboard
+    */
+	virtual bool IsIntegratedKeyboardEnabled() const { return false; }
 };
diff --git a/Engine/Source/Runtime/Slate/Public/Widgets/Input/SEditableText.h b/Engine/Source/Runtime/Slate/Public/Widgets/Input/SEditableText.h
--- a/Engine/Source/Runtime/Slate/Public/Widgets/Input/SEditableText.h
+++ b/Engine/Source/Runtime/Slate/Public/Widgets/Input/SEditableText.h
@@ -122,6 +122,9 @@ public:
 
 		/** Whether the context menu can be opened  */
 		SLATE_ATTRIBUTE(bool, AllowContextMenu)
+	
+		/** Whether the IntegratedKeyboard is enabled  */
+		SLATE_ATTRIBUTE(bool, EnableIntegratedKeyboard)
 
 		/** Delegate to call before a context menu is opened. User returns the menu content or null to the disable context menu */
 		SLATE_EVENT(FOnContextMenuOpening, OnContextMenuOpening)
@@ -238,6 +241,9 @@ public:
 	/** See the AllowContextMenu attribute */
 	SLATE_API void SetAllowContextMenu(const TAttribute< bool >& InAllowContextMenu);
 
+	/** See the EnableIntegratedKeyboard attribute */
+	SLATE_API void SetEnableIntegratedKeyboard(const TAttribute< bool >& InEnableIntegratedKeyboard);
+
 	/** Set the VirtualKeyboardDismissAction attribute */
 	SLATE_API void SetVirtualKeyboardDismissAction(TAttribute< EVirtualKeyboardDismissAction > InVirtualKeyboardDismissAction);
 
@@ -424,6 +430,7 @@ public:
 	SLATE_API virtual bool IsTextReadOnly() const override;
 	SLATE_API virtual bool IsTextPassword() const override;
 	SLATE_API virtual bool IsMultiLineTextEdit() const override;
+	SLATE_API virtual bool IsIntegratedKeyboardEnabled() const override;
 	//~ End ISlateEditableTextWidget Interface
 
 protected:
@@ -495,6 +502,9 @@ protected:
 	/** Whether to disable the context menu */
 	TAttribute<bool> bAllowContextMenu;
 
+	/** Whether to enable integrated keyboard */
+	TAttribute<bool> bEnableIntegratedKeyboard;
+
 	/** Delegate to call before a context menu is opened */
 	FOnContextMenuOpening OnContextMenuOpening;
 
diff --git a/Engine/Source/Runtime/Slate/Public/Widgets/Input/SVirtualKeyboardEntry.h b/Engine/Source/Runtime/Slate/Public/Widgets/Input/SVirtualKeyboardEntry.h
--- a/Engine/Source/Runtime/Slate/Public/Widgets/Input/SVirtualKeyboardEntry.h
+++ b/Engine/Source/Runtime/Slate/Public/Widgets/Input/SVirtualKeyboardEntry.h
@@ -32,6 +32,7 @@ public:
 		, _MinDesiredWidth( 0.0f )
 		, _KeyboardType ( EKeyboardType::Keyboard_Default )
 		, _VirtualKeyboardOptions ( FVirtualKeyboardOptions() )
+		, _EnableIntegratedKeyboard ( false )
 		{}
 
 		/** Sets the text content for this editable text widget */
@@ -67,6 +68,9 @@ public:
 		/** Sets additional arguments to be used by the virtual keyboard summoned by this widget */
 		SLATE_ARGUMENT( FVirtualKeyboardOptions, VirtualKeyboardOptions )
 
+		/** Whether the IntegratedKeyboard is enabled  */
+		SLATE_ARGUMENT( bool, EnableIntegratedKeyboard )
+
 		SLATE_END_ARGS()
 
 
@@ -147,6 +151,13 @@ public:
 
 		return false;
 	}
+
+	virtual bool IsIntegratedKeyboardEnabled() const override
+	{
+		check(IsInGameThread());
+
+		return bEnableIntegratedKeyboard.Get();
+	}
 	//~ End IVirtualKeyboardEntry Interface
 
 protected:
@@ -210,6 +221,9 @@ private:
 
 	TAttribute<EKeyboardType> KeyboardType;
 
+	/** Whether to enable integrated keyboard */
+	TAttribute<bool> bEnableIntegratedKeyboard;
+
 	FVirtualKeyboardOptions VirtualKeyboardOptions;
 
 	bool bNeedsUpdate;
diff --git a/Engine/Source/Runtime/Slate/Public/Widgets/Text/ISlateEditableTextWidget.h b/Engine/Source/Runtime/Slate/Public/Widgets/Text/ISlateEditableTextWidget.h
--- a/Engine/Source/Runtime/Slate/Public/Widgets/Text/ISlateEditableTextWidget.h
+++ b/Engine/Source/Runtime/Slate/Public/Widgets/Text/ISlateEditableTextWidget.h
@@ -237,6 +237,9 @@ public:
 	/** Are we able to insert the given character into our text? */
 	virtual bool CanTypeCharacter(const TCHAR InChar) const = 0;
 
+	/** Is the widget currently has IntegratedKeyboard enabled? */
+	virtual bool IsIntegratedKeyboardEnabled() const = 0;
+
 	/**
 	 * Ensure that we will get a Tick() soon (either due to having active focus, or something having changed progmatically and requiring an update)
 	 * Does nothing if the active tick timer is already enabled
diff --git a/Engine/Source/Runtime/Slate/Public/Widgets/Text/SMultiLineEditableText.h b/Engine/Source/Runtime/Slate/Public/Widgets/Text/SMultiLineEditableText.h
--- a/Engine/Source/Runtime/Slate/Public/Widgets/Text/SMultiLineEditableText.h
+++ b/Engine/Source/Runtime/Slate/Public/Widgets/Text/SMultiLineEditableText.h
@@ -506,6 +506,7 @@ public:
 	SLATE_API virtual bool IsTextReadOnly() const override;
 	SLATE_API virtual bool IsTextPassword() const override;
 	SLATE_API virtual bool IsMultiLineTextEdit() const override;
+	SLATE_API virtual bool IsIntegratedKeyboardEnabled() const override;
 	//~ End ISlateEditableTextWidget Interface
 
 protected:
diff --git a/Engine/Source/Runtime/Slate/Public/Widgets/Text/SlateEditableTextLayout.h b/Engine/Source/Runtime/Slate/Public/Widgets/Text/SlateEditableTextLayout.h
--- a/Engine/Source/Runtime/Slate/Public/Widgets/Text/SlateEditableTextLayout.h
+++ b/Engine/Source/Runtime/Slate/Public/Widgets/Text/SlateEditableTextLayout.h
@@ -451,6 +451,7 @@ private:
 		virtual EKeyboardType GetVirtualKeyboardType() const override;
 		virtual FVirtualKeyboardOptions GetVirtualKeyboardOptions() const override;
 		virtual bool IsMultilineEntry() const override;
+		virtual bool IsIntegratedKeyboardEnabled() const override;
 
 	private:
 		FVirtualKeyboardEntry(FSlateEditableTextLayout& InOwnerLayout);
diff --git a/Engine/Source/Runtime/UMG/Private/Components/EditableText.cpp b/Engine/Source/Runtime/UMG/Private/Components/EditableText.cpp
--- a/Engine/Source/Runtime/UMG/Private/Components/EditableText.cpp
+++ b/Engine/Source/Runtime/UMG/Private/Components/EditableText.cpp
@@ -50,6 +50,7 @@ PRAGMA_DISABLE_DEPRECATION_WARNINGS
 	ClearKeyboardFocusOnCommit = true;
 	SelectAllTextOnCommit = false;
 	AllowContextMenu = true;
+	EnableIntegratedKeyboard = false;
 	VirtualKeyboardTrigger = EVirtualKeyboardTrigger::OnFocusByPointer;
 	VirtualKeyboardDismissAction = EVirtualKeyboardDismissAction::TextChangeOnDismiss;
 	SetClipping(EWidgetClipping::ClipToBounds);
@@ -89,7 +90,8 @@ TSharedRef<SWidget> UEditableText::RebuildWidget()
 		.VirtualKeyboardTrigger(VirtualKeyboardTrigger)
 		.VirtualKeyboardDismissAction(VirtualKeyboardDismissAction)
 		.Justification(Justification)
-		.OverflowPolicy(OverflowPolicy);
+		.OverflowPolicy(OverflowPolicy)
+		.EnableIntegratedKeyboard(EnableIntegratedKeyboard);
 	PRAGMA_ENABLE_DEPRECATION_WARNINGS
 	
 	return MyEditableText.ToSharedRef();
@@ -118,6 +120,7 @@ void UEditableText::SynchronizeProperties()
 	MyEditableText->SetVirtualKeyboardDismissAction(VirtualKeyboardDismissAction);
 	MyEditableText->SetJustification(Justification);
 	MyEditableText->SetOverflowPolicy(OverflowPolicy);
+	MyEditableText->SetEnableIntegratedKeyboard(EnableIntegratedKeyboard);
 
 	PRAGMA_ENABLE_DEPRECATION_WARNINGS
 
@@ -373,6 +376,21 @@ void UEditableText::SetFontOutlineMaterial(UMaterialInterface* InMaterial)
 	}
 }
 
+bool UEditableText::GetEnableIntegratedKeyboard() const
+{
+	return EnableIntegratedKeyboard;
+}
+
+void UEditableText::SetEnableIntegratedKeyboard(bool InEnableIntegratedKeyboard)
+{
+	EnableIntegratedKeyboard = InEnableIntegratedKeyboard;
+	
+	if (MyEditableText.IsValid())
+	{
+		MyEditableText->SetEnableIntegratedKeyboard(EnableIntegratedKeyboard);
+	}
+}
+
 void UEditableText::HandleOnTextChanged(const FText& InText)
 {
 	if (SetTextInternal(InText))
diff --git a/Engine/Source/Runtime/UMG/Public/Components/EditableText.h b/Engine/Source/Runtime/UMG/Public/Components/EditableText.h
--- a/Engine/Source/Runtime/UMG/Public/Components/EditableText.h
+++ b/Engine/Source/Runtime/UMG/Public/Components/EditableText.h
@@ -241,6 +241,9 @@ public:
 	UFUNCTION(BlueprintCallable, Category = "Appearance")
 	UMG_API void SetFontOutlineMaterial(UMaterialInterface* InMaterial);
 
+	UMG_API bool GetEnableIntegratedKeyboard() const;
+	UMG_API void SetEnableIntegratedKeyboard(bool InEnableIntegratedKeyboard);
+
 public:
 	
 	//~ Begin UWidget Interface
@@ -277,4 +280,8 @@ protected:
 private:
 	/** @return true if the text was changed, or false if identical. */
 	bool SetTextInternal(const FText& InText);
+	
+	/** Whether the integrated keyboard is enabled for this editable text. */
+	UPROPERTY(EditAnywhere, BlueprintReadWrite, Getter, Setter, Category = Mobile, AdvancedDisplay, meta = (AllowPrivateAccess = "true"))
+	bool EnableIntegratedKeyboard;
 };
