Fixed: PSO pre-caching for Widget and MaterialBillboard components on platforms that do not support 'Manual Vertex Fetch'

diff --git a/Engine/Source/Runtime/Engine/Private/Components/MaterialBillboardComponent.cpp b/Engine/Source/Runtime/Engine/Private/Components/MaterialBillboardComponent.cpp
--- a/Engine/Source/Runtime/Engine/Private/Components/MaterialBillboardComponent.cpp
+++ b/Engine/Source/Runtime/Engine/Private/Components/MaterialBillboardComponent.cpp
@@ -406,7 +406,23 @@ void UMaterialBillboardComponent::CollectPSOPrecacheData(const FPSOPrecacheParam
 			FMaterialInterfacePSOPrecacheParams& ComponentParams = OutParams[OutParams.AddDefaulted()];
 			ComponentParams.MaterialInterface = MaterialInterface;
 			ComponentParams.PSOPrecacheParams = PrecachePSOParams;
-			ComponentParams.VertexFactoryDataList.Add(FPSOPrecacheVertexFactoryData(VFType));
+
+			if (RHISupportsManualVertexFetch(GMaxRHIShaderPlatform))
+			{
+				ComponentParams.VertexFactoryDataList.Add(FPSOPrecacheVertexFactoryData(VFType));
+			}
+			else
+			{
+				FVertexDeclarationElementList VtxElements;
+				//Position Float3, 2x PackedNormal, 1x UVs Half2, Color, InstancingID
+				VtxElements.Add(FVertexElement(0, 0, VET_Float3, 0, sizeof(FVector3f)));
+				VtxElements.Add(FVertexElement(1, 0, VET_PackedNormal, 1, sizeof(FPackedNormal) * 2));
+				VtxElements.Add(FVertexElement(1, sizeof(FPackedNormal), VET_PackedNormal, 2, sizeof(FPackedNormal) * 2));
+				VtxElements.Add(FVertexElement(2, 0, VET_Half2, 3, 4));
+				VtxElements.Add(FVertexElement(3, 0, VET_Color, 4, sizeof(FColor)));
+				VtxElements.Add(FVertexElement(4, 0, VET_UInt, 13, sizeof(uint32), true));
+				ComponentParams.VertexFactoryDataList.Add(FPSOPrecacheVertexFactoryData(VFType, VtxElements));
+			}
 		}
 	}
 }
diff --git a/Engine/Source/Runtime/UMG/Private/Components/WidgetComponent.cpp b/Engine/Source/Runtime/UMG/Private/Components/WidgetComponent.cpp
--- a/Engine/Source/Runtime/UMG/Private/Components/WidgetComponent.cpp
+++ b/Engine/Source/Runtime/UMG/Private/Components/WidgetComponent.cpp
@@ -710,8 +710,24 @@ void UWidgetComponent::CollectPSOPrecacheData(const FPSOPrecacheParams& BasePrec
 		FMaterialInterfacePSOPrecacheParams& ComponentParams = OutParams[OutParams.AddDefaulted()];
 		ComponentParams.Priority = EPSOPrecachePriority::High;
 		ComponentParams.MaterialInterface = MaterialInstance;
-		ComponentParams.VertexFactoryDataList.Add(FPSOPrecacheVertexFactoryData(&FLocalVertexFactory::StaticType));
 		ComponentParams.PSOPrecacheParams = BasePrecachePSOParams;
+		
+		if (RHISupportsManualVertexFetch(GMaxRHIShaderPlatform))
+		{
+			ComponentParams.VertexFactoryDataList.Add(FPSOPrecacheVertexFactoryData(&FLocalVertexFactory::StaticType));
+		}
+		else
+		{
+			FVertexDeclarationElementList VtxElements;
+			//Position Float3, 2x PackedNormal, 1x UVs Float2, Color, InstancingID
+			VtxElements.Add(FVertexElement(0, 0, VET_Float3, 0, sizeof(FVector3f)));
+			VtxElements.Add(FVertexElement(1, 0, VET_PackedNormal, 1, sizeof(FPackedNormal) * 2));
+			VtxElements.Add(FVertexElement(1, sizeof(FPackedNormal), VET_PackedNormal, 2, sizeof(FPackedNormal) * 2));
+			VtxElements.Add(FVertexElement(2, 0, VET_Float2, 3, sizeof(FVector2f)));
+			VtxElements.Add(FVertexElement(3, 0, VET_Color, 4, sizeof(FColor)));
+			VtxElements.Add(FVertexElement(4, 0, VET_UInt, 13, sizeof(uint32), true));
+			ComponentParams.VertexFactoryDataList.Add(FPSOPrecacheVertexFactoryData(&FLocalVertexFactory::StaticType, VtxElements));
+		}
 	}
 }
 
