USlider: expose bPreventThrottling to the user

#rnx

diff --git a/Engine/Source/Runtime/UMG/Private/Components/Slider.cpp b/Engine/Source/Runtime/UMG/Private/Components/Slider.cpp
--- a/Engine/Source/Runtime/UMG/Private/Components/Slider.cpp
+++ b/Engine/Source/Runtime/UMG/Private/Components/Slider.cpp
@@ -54,7 +54,8 @@ PRAGMA_DISABLE_DEPRECATION_WARNINGS
 		.OnMouseCaptureEnd(BIND_UOBJECT_DELEGATE(FSimpleDelegate, HandleOnMouseCaptureEnd))
 		.OnControllerCaptureBegin(BIND_UOBJECT_DELEGATE(FSimpleDelegate, HandleOnControllerCaptureBegin))
 		.OnControllerCaptureEnd(BIND_UOBJECT_DELEGATE(FSimpleDelegate, HandleOnControllerCaptureEnd))
-		.OnValueChanged(BIND_UOBJECT_DELEGATE(FOnFloatValueChanged, HandleOnValueChanged));
+		.OnValueChanged(BIND_UOBJECT_DELEGATE(FOnFloatValueChanged, HandleOnValueChanged))
+		.PreventThrottling(bPreventThrottling);
 PRAGMA_ENABLE_DEPRECATION_WARNINGS
 	return MySlider.ToSharedRef();
 }
diff --git a/Engine/Source/Runtime/UMG/Public/Components/Slider.h b/Engine/Source/Runtime/UMG/Public/Components/Slider.h
--- a/Engine/Source/Runtime/UMG/Public/Components/Slider.h
+++ b/Engine/Source/Runtime/UMG/Public/Components/Slider.h
@@ -95,6 +95,10 @@ public:
 	UPROPERTY(EditAnywhere, BlueprintReadOnly, Category="Interaction")
 	bool IsFocusable;
 
+	/** If refresh requests for the viewport should happen for all value changes **/
+	UPROPERTY(EditAnywhere, BlueprintReadOnly, Category="Interaction")
+	bool bPreventThrottling = false;
+
 public:
 
 	/** Invoked when the mouse is pressed and a capture begins. */
