Allow Blueprint to override functionality of ITypedUMGListView::OnIsSelectableOrNavigableInternal via UListView::BP_OnIsItemSelectableOrNavigable. This enables list views in widgets to have disabled / unavailable options that cannot be navigated to without requiring a base C++ class derivation for that list view.
#tests PIE Tested a widget owning a list view and setting an override for its BP_OnIsItemSelectableOrNavigable via Create Event on Event Construct. Blueprint logic got executed in place of C++ logic, allowing special behavior regarding the specific entry type for that list view.

diff --git a/Engine/Source/Runtime/UMG/Private/Components/ListView.cpp b/Engine/Source/Runtime/UMG/Private/Components/ListView.cpp
--- a/Engine/Source/Runtime/UMG/Private/Components/ListView.cpp
+++ b/Engine/Source/Runtime/UMG/Private/Components/ListView.cpp
@@ -368,6 +368,11 @@ void UListView::OnListViewFinishedScrollingInternal()
 	BP_OnListViewFinishedScrolling.Broadcast();
 }
 
+bool UListView::OnIsSelectableOrNavigableInternal(UObject* FirstSelectedItem)
+{
+	return BP_OnIsItemSelectableOrNavigable.IsBound() ? BP_OnIsItemSelectableOrNavigable.Execute(FirstSelectedItem) : ITypedUMGListView::OnIsSelectableOrNavigableInternal(FirstSelectedItem);
+}
+
 void UListView::InitHorizontalEntrySpacing(float InHorizontalEntrySpacing)
 {
 	ensureMsgf(!MyListView.IsValid(), TEXT("The SWidget is already created."));
diff --git a/Engine/Source/Runtime/UMG/Public/Components/ListView.h b/Engine/Source/Runtime/UMG/Public/Components/ListView.h
--- a/Engine/Source/Runtime/UMG/Public/Components/ListView.h
+++ b/Engine/Source/Runtime/UMG/Public/Components/ListView.h
@@ -14,6 +14,7 @@ DECLARE_DYNAMIC_MULTICAST_DELEGATE_TwoParams(FOnItemIsHoveredChangedDynamic, UOb
 DECLARE_DYNAMIC_MULTICAST_DELEGATE_TwoParams(FOnListItemScrolledIntoViewDynamic, UObject*, Item, UUserWidget*, Widget);
 DECLARE_DYNAMIC_MULTICAST_DELEGATE_TwoParams(FOnListViewScrolledDynamic, float, ItemOffset, float, DistanceRemaining);
 DECLARE_DYNAMIC_MULTICAST_DELEGATE(FOnListViewFinishedScrollingDynamic);
+DECLARE_DYNAMIC_DELEGATE_RetVal_OneParam(bool, FOnIsItemSelectableOrNavigableDynamic, UObject*, Item);
 
 /**
  * A virtualized list that allows up to thousands of items to be displayed.
@@ -160,6 +161,7 @@ protected:
 	UMG_API virtual void OnItemScrolledIntoViewInternal(UObject* Item, UUserWidget& EntryWidget) override;
 	UMG_API virtual void OnListViewScrolledInternal(float ItemOffset, float DistanceRemaining) override;
 	UMG_API virtual void OnListViewFinishedScrollingInternal() override;
+	UMG_API virtual bool OnIsSelectableOrNavigableInternal(UObject* FirstSelectedItem) override;
 
 	UMG_API void HandleOnEntryInitializedInternal(UObject* Item, const TSharedRef<ITableRow>& TableRow);
 
@@ -336,4 +338,7 @@ private:
 	
 	UPROPERTY(BlueprintAssignable, Category = Events, meta = (DisplayName = "On List View Finished Scrolling"))
 	FOnListViewFinishedScrollingDynamic BP_OnListViewFinishedScrolling;	
+
+	UPROPERTY(BlueprintReadWrite, Category = Events, meta = (AllowPrivateAccess, DisplayName = "On Is Item Selectable Or Navigable"))
+	FOnIsItemSelectableOrNavigableDynamic BP_OnIsItemSelectableOrNavigable;
 };
