UMG: Fix copy-pasted comment on UBackgroundBlur::CornerRadius

#rnx

diff --git a/Engine/Source/Runtime/UMG/Public/Components/BackgroundBlur.h b/Engine/Source/Runtime/UMG/Public/Components/BackgroundBlur.h
--- a/Engine/Source/Runtime/UMG/Public/Components/BackgroundBlur.h
+++ b/Engine/Source/Runtime/UMG/Public/Components/BackgroundBlur.h
@@ -59,10 +59,7 @@ public:
 	int32 BlurRadius;
 
 	UE_DEPRECATED(5.1, "Direct access to CornerRadius is deprecated. Please use the getter or setter.")
-	/**
-	 * This is the number of pixels which will be weighted in each direction from any given pixel when computing the blur
-	 * A larger value is more costly but allows for stronger blurs.  
-	 */
+	/** Radius in Slate Units applied to the blur widget at each corner. X = Top Left, Y = Top Right, Z = Bottom Right, W = Bottom Left */
 	UPROPERTY(EditAnywhere, AdvancedDisplay, BlueprintReadWrite, Getter, Setter, BlueprintSetter=SetCornerRadius, Category = Appearance)
 	FVector4 CornerRadius;
 
