Expose IsScrolling for blueprint usage
#12580

#tests Made a scrollbox, changed a script to access GetIsScrolling from it
#rnx

diff --git a/Engine/Source/Runtime/Slate/Public/Widgets/Layout/SScrollBox.h b/Engine/Source/Runtime/Slate/Public/Widgets/Layout/SScrollBox.h
--- a/Engine/Source/Runtime/Slate/Public/Widgets/Layout/SScrollBox.h
+++ b/Engine/Source/Runtime/Slate/Public/Widgets/Layout/SScrollBox.h
@@ -292,6 +292,8 @@ public:
 	/** Gets the scroll offset of the bottom of the ScrollBox in Slate Units. */
 	SLATE_API float GetScrollOffsetOfEnd() const;
 
+	SLATE_API bool GetIsScrolling() const { return bIsScrolling; }
+
 	SLATE_API void SetScrollOffset( float NewScrollOffset );
 
 	SLATE_API void ScrollToStart();
diff --git a/Engine/Source/Runtime/UMG/Private/Components/ScrollBox.cpp b/Engine/Source/Runtime/UMG/Private/Components/ScrollBox.cpp
--- a/Engine/Source/Runtime/UMG/Private/Components/ScrollBox.cpp
+++ b/Engine/Source/Runtime/UMG/Private/Components/ScrollBox.cpp
@@ -255,6 +255,11 @@ void UScrollBox::ScrollWidgetIntoView(UWidget* WidgetToFind, bool AnimateScroll,
 	}
 }
 
+bool UScrollBox::GetIsScrolling() const
+{
+	return MyScrollBox.IsValid() && MyScrollBox->GetIsScrolling();
+}
+
 #if WITH_EDITORONLY_DATA
 
 void UScrollBox::Serialize(FArchive& Ar)
diff --git a/Engine/Source/Runtime/UMG/Public/Components/ScrollBox.h b/Engine/Source/Runtime/UMG/Public/Components/ScrollBox.h
--- a/Engine/Source/Runtime/UMG/Public/Components/ScrollBox.h
+++ b/Engine/Source/Runtime/UMG/Public/Components/ScrollBox.h
@@ -284,6 +284,10 @@ public:
 	UFUNCTION(BlueprintCallable, Category="Widget")
 	UMG_API void ScrollWidgetIntoView(UWidget* WidgetToFind, bool AnimateScroll = true, EDescendantScrollDestination ScrollDestination = EDescendantScrollDestination::IntoView, float Padding = 0);
 
+	/** Returns true when the widget is currently inertial scrolling */
+	UFUNCTION(BlueprintPure, Category="Widget")
+	UMG_API bool GetIsScrolling() const;
+
 	//~ Begin UWidget Interface
 	UMG_API virtual void SynchronizeProperties() override;
 	//~ End UWidget Interface
