Update ActiveWidgetIndex after adding or removing a slot
It prevents from getting out of sync with SWidgetSwitcher's WidgetIndex, which can be modified when adding or removing a slot.

diff --git a/Engine/Source/Runtime/UMG/Private/Components/WidgetSwitcher.cpp b/Engine/Source/Runtime/UMG/Private/Components/WidgetSwitcher.cpp
--- a/Engine/Source/Runtime/UMG/Private/Components/WidgetSwitcher.cpp
+++ b/Engine/Source/Runtime/UMG/Private/Components/WidgetSwitcher.cpp
@@ -106,6 +106,11 @@ void UWidgetSwitcher::OnSlotAdded(UPanelSlot* InSlot)
 	if ( MyWidgetSwitcher.IsValid() )
 	{
 		CastChecked<UWidgetSwitcherSlot>(InSlot)->BuildSlot(MyWidgetSwitcher.ToSharedRef());
+
+		// Adding the slot might have changed the active widget index.
+PRAGMA_DISABLE_DEPRECATION_WARNINGS
+		ActiveWidgetIndex = MyWidgetSwitcher->GetActiveWidgetIndex();
+PRAGMA_ENABLE_DEPRECATION_WARNINGS
 	}
 }
 
@@ -118,6 +123,11 @@ void UWidgetSwitcher::OnSlotRemoved(UPanelSlot* InSlot)
 		if ( Widget.IsValid() )
 		{
 			MyWidgetSwitcher->RemoveSlot(Widget.ToSharedRef());
+
+			// Removing the slot might have changed the active widget index.
+PRAGMA_DISABLE_DEPRECATION_WARNINGS
+			ActiveWidgetIndex = MyWidgetSwitcher->GetActiveWidgetIndex();
+PRAGMA_ENABLE_DEPRECATION_WARNINGS
 		}
 	}
 }
