UMG: refactor widget animations to not use a player object anymore

This change removes the use of UUMGSequencePlayer except in places that require backwards compatibility. These player objects, and the IMovieScenePlayer interface, are getting deprecated in favor of more lightweight "runner" structs that can be easily packed in memory. A future optimization might even be to move all these runner structures into the UMG Sequence Tick Manager.

As much care as possible as been taken for backwards compatibility. Almost everything should still work _except_ the rare case of calling GetUserTag or SetUserTag on the returned player object from a PlayAnimation BP node that is *not* done from *inside* the user widget itself. That is, such calls, if done in another Blueprint, wouldn't get automatically fixed up and result in BP compilation errors. It should only take a couple clicks to fix them manually.

A special note about the virtual methods OnAnimationStartedPlaying and OnAnimationFinishedPlaying on the user widget class: these used to take the player object. In order to only create these legacy player objects when needed, these virtual methods aren't called unless the user widget's actual class hierarchy contains a native C++ class. When that happens, we have no way of knowing if that sub-class has overriden those two callbacks, so by precaution we do create a player object and call them. If these methods are, in fact, not overriden, the sub-class can override GetLegacyPlayerRequirement and declare that it doesn't need them called, nor a player object created just for that purpose. And if the methods have been overriden, that code should be upgrade to instead override the versions that take the runner structure.

diff --git a/Engine/Plugins/Runtime/CommonUI/Source/CommonUI/Public/CommonUserWidget.h b/Engine/Plugins/Runtime/CommonUI/Source/CommonUI/Public/CommonUserWidget.h
--- a/Engine/Plugins/Runtime/CommonUI/Source/CommonUI/Public/CommonUserWidget.h
+++ b/Engine/Plugins/Runtime/CommonUI/Source/CommonUI/Public/CommonUserWidget.h
@@ -48,6 +48,8 @@ public:
 	void AddActionBinding(FUIActionBindingHandle ActionBinding);
 
 protected:
+	virtual ERequiresLegacyPlayer GetLegacyPlayerRequirement() const override { return ERequiresLegacyPlayer::No; }
+
 	virtual void OnWidgetRebuilt() override;
 	virtual void NativeDestruct() override;
 	
diff --git a/Engine/Source/Editor/UMGEditor/Private/Nodes/K2Node_PlayAnimation.cpp b/Engine/Source/Editor/UMGEditor/Private/Nodes/K2Node_PlayAnimation.cpp
--- a/Engine/Source/Editor/UMGEditor/Private/Nodes/K2Node_PlayAnimation.cpp
+++ b/Engine/Source/Editor/UMGEditor/Private/Nodes/K2Node_PlayAnimation.cpp
@@ -21,4 +21,20 @@ UK2Node_PlayAnimationTimeRange::UK2Node_PlayAnimationTimeRange(const FObjectInit
 	ProxyClass = UWidgetAnimationPlayCallbackProxy::StaticClass();
 }
 
+UK2Node_PlayAnimation2::UK2Node_PlayAnimation2(const FObjectInitializer& ObjectInitializer)
+	: Super(ObjectInitializer)
+{
+	ProxyFactoryFunctionName = GET_FUNCTION_NAME_CHECKED(UWidgetAnimationPlayCallbackProxy, NewPlayAnimationProxyObject);
+	ProxyFactoryClass = UWidgetAnimationPlayCallbackProxy::StaticClass();
+	ProxyClass = UWidgetAnimationPlayCallbackProxy::StaticClass();
+}
+
+UK2Node_PlayAnimationTimeRange2::UK2Node_PlayAnimationTimeRange2(const FObjectInitializer& ObjectInitializer)
+	: Super(ObjectInitializer)
+{
+	ProxyFactoryFunctionName = GET_FUNCTION_NAME_CHECKED(UWidgetAnimationPlayCallbackProxy, NewPlayAnimationTimeRangeProxyObject);
+	ProxyFactoryClass = UWidgetAnimationPlayCallbackProxy::StaticClass();
+	ProxyClass = UWidgetAnimationPlayCallbackProxy::StaticClass();
+}
+
 #undef LOCTEXT_NAMESPACE
diff --git a/Engine/Source/Editor/UMGEditor/Private/Nodes/K2Node_PlayAnimation.h b/Engine/Source/Editor/UMGEditor/Private/Nodes/K2Node_PlayAnimation.h
--- a/Engine/Source/Editor/UMGEditor/Private/Nodes/K2Node_PlayAnimation.h
+++ b/Engine/Source/Editor/UMGEditor/Private/Nodes/K2Node_PlayAnimation.h
@@ -16,3 +16,16 @@ class UMGEDITOR_API UK2Node_PlayAnimationTimeRange : public UK2Node_BaseAsyncTas
 {
 	GENERATED_UCLASS_BODY()
 };
+
+UCLASS()
+class UMGEDITOR_API UK2Node_PlayAnimation2 : public UK2Node_BaseAsyncTask
+{
+	GENERATED_UCLASS_BODY()
+};
+
+UCLASS()
+class UMGEDITOR_API UK2Node_PlayAnimationTimeRange2 : public UK2Node_BaseAsyncTask
+{
+	GENERATED_UCLASS_BODY()
+};
+
diff --git a/Engine/Source/Editor/UMGEditor/Private/WidgetBlueprint.cpp b/Engine/Source/Editor/UMGEditor/Private/WidgetBlueprint.cpp
--- a/Engine/Source/Editor/UMGEditor/Private/WidgetBlueprint.cpp
+++ b/Engine/Source/Editor/UMGEditor/Private/WidgetBlueprint.cpp
@@ -11,6 +11,8 @@
 #include "EdGraph/EdGraph.h"
 #include "Blueprint/WidgetTree.h"
 #include "Animation/WidgetAnimation.h"
+#include "Animation/WidgetAnimationHandle.h"
+#include "Animation/UMGSequencePlayer.h"
 
 #include "Kismet2/StructureEditorUtils.h"
 
@@ -1063,6 +1065,71 @@ void UWidgetBlueprint::PostLoad()
 			Graph->Schema = UWidgetGraphSchema::StaticClass();
 		}
 	}
+
+	if ( GetLinkerCustomVersion(FFortniteMainBranchObjectVersion::GUID) < FFortniteMainBranchObjectVersion::UpgradeWidgetBlueprintLegacySequencePlayer )
+	{
+		FixupLegacySequencePlayerCalls();
+	}
+}
+
+void UWidgetBlueprint::FixupLegacySequencePlayerCalls()
+{
+	TArray<UEdGraph*> Graphs;
+	GetAllGraphs(Graphs);
+	for (UEdGraph* Graph : Graphs)
+	{
+		TArray<UK2Node_CallFunction*> CallFunctionNodes;
+		Graph->GetNodesOfClass(CallFunctionNodes);
+		for (UK2Node_CallFunction* CallFunctionNode : CallFunctionNodes)
+		{
+			// Old calls to GetUserTag and SetUserTag were done on the UUMGSequencePlayer object returned by
+			// calls to PlayAnimation and its variants. These would be hooked up on the "self" pin.
+			//
+			// The new calls are done one the animation handle structure, via functions on 
+			// UWidgetAnimationHandleFunctionLibrary. In this case, since these functions are static, the
+			// "self" pin is hidden, and we want calls to be done on the "Target" pin which didn't exist
+			// before, and doesn't exist yet as we are being called in PostLoad.
+			// 
+			// Core redirectors can't handle this case so we fix them up here. We rename "self" to Target
+			// so that connections are preserved.
+			//
+			// TODO: this wouldn't fix calls done _outside_ of the widget, those will be broken.
+			//
+			const FName FunctionName = CallFunctionNode->GetFunctionName();
+			const bool bMatchesAnimationFunctionName = (
+					FunctionName == GET_FUNCTION_NAME_CHECKED(UUMGSequencePlayer, GetUserTag) ||
+					FunctionName == GET_FUNCTION_NAME_CHECKED(UUMGSequencePlayer, SetUserTag));
+			if (!bMatchesAnimationFunctionName)
+			{
+				continue;
+			}
+
+			UEdGraphPin* SelfPin = GetDefault<UEdGraphSchema_K2>()->FindSelfPin(*CallFunctionNode, EGPD_Input);
+			const bool bMatchesAnimationClass = (
+					SelfPin &&
+					SelfPin->PinType.PinCategory == UEdGraphSchema_K2::PC_Object &&
+					SelfPin->PinType.PinSubCategoryObject == UUMGSequencePlayer::StaticClass());
+			if (!bMatchesAnimationClass)
+			{
+				continue;
+			}
+
+			// Redirect from UUMGSequencePlayer::XxxUserTag to UWidgetAnimationHandleFunctionLibrary::XxxUserTag.
+			CallFunctionNode->FunctionReference.SetExternalMember(
+					FunctionName, UWidgetAnimationHandleFunctionLibrary::StaticClass());
+
+			// IMPORTANT: the name "Target" must match the parameter name in the GetUserTag and 
+			// SetUserTag methods!
+			UEdGraphPin* TargetPin = CallFunctionNode->FindPin(TEXT("Target"));
+			if (SelfPin && !SelfPin->LinkedTo.IsEmpty() && !TargetPin)
+			{
+				SelfPin->PinName = TEXT("Target");
+#if WITH_EDITORONLY_DATA
+				SelfPin->PinFriendlyName = FText();
+#endif
+			}
+		}
+	}
 }
 
 void UWidgetBlueprint::PostDuplicate(bool bDuplicateForPIE)
diff --git a/Engine/Source/Editor/UMGEditor/Public/WidgetBlueprint.h b/Engine/Source/Editor/UMGEditor/Public/WidgetBlueprint.h
--- a/Engine/Source/Editor/UMGEditor/Public/WidgetBlueprint.h
+++ b/Engine/Source/Editor/UMGEditor/Public/WidgetBlueprint.h
@@ -332,8 +332,13 @@ public:
 	virtual const UWidgetEditingProjectSettings* GetRelevantSettings() const;
 
 protected:
+
 	virtual void LoadModulesRequiredForCompilation() override;
 
+private:
+
+	void FixupLegacySequencePlayerCalls();
+
 private:
 	/**
 	 * The desired tick frequency set by the user on the UserWidget's CDO.
diff --git a/Engine/Source/Runtime/Core/Public/UObject/FortniteMainBranchObjectVersions.inl b/Engine/Source/Runtime/Core/Public/UObject/FortniteMainBranchObjectVersions.inl
--- a/Engine/Source/Runtime/Core/Public/UObject/FortniteMainBranchObjectVersions.inl
+++ b/Engine/Source/Runtime/Core/Public/UObject/FortniteMainBranchObjectVersions.inl
@@ -613,5 +613,8 @@ FORTNITE_MAIN_VERSION(VerseRightToLeftHandedness, 199)
 // Added additional data required to record and represent particle data from the game thread (Kinematic targets, and SQ rejection reasons)
 FORTNITE_MAIN_VERSION(AdditionalGameThreadDataSupportInChaosVisualDebugger, 200)
 
+// Upgrade UMG widget blueprints using legacy animation API
+FORTNITE_MAIN_VERSION(UpgradeWidgetBlueprintLegacySequencePlayer, 201)
+
 // -----<new versions can be added above this line>-------------------------------------------------
 #undef FORTNITE_MAIN_VERSION
\ No newline at end of file
diff --git a/Engine/Source/Runtime/MovieScene/Public/IMovieScenePlayer.h b/Engine/Source/Runtime/MovieScene/Public/IMovieScenePlayer.h
--- a/Engine/Source/Runtime/MovieScene/Public/IMovieScenePlayer.h
+++ b/Engine/Source/Runtime/MovieScene/Public/IMovieScenePlayer.h
@@ -37,12 +37,12 @@ namespace UE::MovieScene
 	/**
 	 * Playback capability for storing an IMovieScenePlayer unique index.
 	 */
-	struct FPlayerIndexPlaybackCapability
+	struct MOVIESCENE_API FPlayerIndexPlaybackCapability
 	{
 		UE_DECLARE_MOVIESCENE_PLAYBACK_CAPABILITY(FPlayerIndexPlaybackCapability)
 
-		static MOVIESCENE_API IMovieScenePlayer* GetPlayer(TSharedRef<const FSharedPlaybackState> Owner);
-		static MOVIESCENE_API uint16 GetPlayerIndex(TSharedRef<const FSharedPlaybackState> Owner);
+		static IMovieScenePlayer* GetPlayer(TSharedRef<const FSharedPlaybackState> Owner);
+		static uint16 GetPlayerIndex(TSharedRef<const FSharedPlaybackState> Owner);
 
 		FPlayerIndexPlaybackCapability(uint16 InPlayerIndex)
 			: PlayerIndex(InPlayerIndex)
diff --git a/Engine/Source/Runtime/UMG/Private/Animation/UMGSequencePlayer.cpp b/Engine/Source/Runtime/UMG/Private/Animation/UMGSequencePlayer.cpp
--- a/Engine/Source/Runtime/UMG/Private/Animation/UMGSequencePlayer.cpp
+++ b/Engine/Source/Runtime/UMG/Private/Animation/UMGSequencePlayer.cpp
@@ -1,628 +1,265 @@
 // Copyright Epic Games, Inc. All Rights Reserved.
 
 #include "Animation/UMGSequencePlayer.h"
-#include "MovieScene.h"
-#include "UMGPrivate.h"
+
+#include "Animation/UMGSequenceTickManager.h"
 #include "Animation/WidgetAnimation.h"
-#include "MovieSceneTimeHelpers.h"
-#include "Evaluation/MovieSceneSequenceHierarchy.h"
-#include "Evaluation/MovieScenePlayback.h"
-#include "Channels/MovieSceneTimeWarpChannel.h"
-#include "EntitySystem/MovieSceneEntitySystemRunner.h"
+#include "Animation/WidgetAnimationState.h"
 #include "EntitySystem/MovieSceneEntitySystemLinker.h"
-#include "Animation/UMGSequenceTickManager.h"
-#include "ProfilingDebugging/CsvProfiler.h"
+#include "EntitySystem/MovieSceneEntitySystemRunner.h"
+#include "IMovieScenePlaybackClient.h"
+#include "MovieSceneFwd.h"
 
 #include UE_INLINE_GENERATED_CPP_BY_NAME(UMGSequencePlayer)
 
-extern TAutoConsoleVariable<bool> CVarUserWidgetUseParallelAnimation;
-
-namespace UE::UMG
-{
-
-	bool GAsyncAnimationControlFlow = true;
-	FAutoConsoleVariableRef CVarAsyncAnimationControlFlow(
-		TEXT("UMG.AsyncAnimationControlFlow"),
-		GAsyncAnimationControlFlow,
-		TEXT("(Default: true) Whether to perform animation control flow functions (Play, Pause, Stop etc) asynchronously.")
-	);
-
-	bool GVarAnimationMarkers = false;
-	FAutoConsoleVariableRef CVarAnimationMarkers(
-		TEXT("UMG.AnimationMarkers"),
-		GVarAnimationMarkers,
-		TEXT("(Default: false) Whether to emit profiling frame markers for starting and stopping UMG animations.")
-	);
-
-} // namespace UE::UMG
-
 UUMGSequencePlayer::UUMGSequencePlayer(const FObjectInitializer& ObjectInitializer)
 	: Super(ObjectInitializer)
 {
-	PlayerStatus = EMovieScenePlayerStatus::Stopped;
-	TimeCursorPosition = UnwarpedPosition = FFrameTime(0);
-	PlaybackSpeed = 1;
-	bRestoreState = false;
-	Animation = nullptr;
-	bIsStopping = false;
-	bIsBeginningPlay = false;
-	bCompleteOnPostEvaluation = false;
-	UserTag = NAME_None;
-	BlockedDeltaTimeCompensation = 0.f;
 }
 
-void UUMGSequencePlayer::InitSequencePlayer(UWidgetAnimation& InAnimation, UUserWidget& InUserWidget)
+void UUMGSequencePlayer::InitSequencePlayer(FWidgetAnimationState& InState)
 {
-	Animation = &InAnimation;
-	UserWidget = &InUserWidget;
-
-	UMovieScene* MovieScene = Animation->GetMovieScene();
+	WidgetAnimationHandle = InState.GetAnimationHandle();
 
-	// Cache the time range of the sequence to determine when we stop
-	Duration = UE::MovieScene::DiscreteSize(MovieScene->GetPlaybackRange());
-	AnimationResolution = MovieScene->GetTickResolution();
-	AbsolutePlaybackStart = UE::MovieScene::DiscreteInclusiveLower(MovieScene->GetPlaybackRange());
+	RootTemplateInstance.Initialize(InState.GetSharedPlaybackState().ToSharedRef());
 }
 
-UMovieSceneEntitySystemLinker* UUMGSequencePlayer::ConstructEntitySystemLinker()
+void UUMGSequencePlayer::InitSequencePlayer(UWidgetAnimation& InAnimation, UUserWidget& InUserWidget)
 {
-	UUserWidget* Widget = UserWidget.Get();
-	if (ensure(Widget) && !EnumHasAnyFlags(Animation->GetFlags(), EMovieSceneSequenceFlags::BlockingEvaluation))
-	{
-		if (!ensure(Widget->AnimationTickManager))
-		{
-			// @todo: There should be no possible way that the animation tick manager is null here, but there is a very low-rate
-			// crash caused by it being null that is very hard to track down, so patching with a band-aid for now.
-			Widget->AnimationTickManager = UUMGSequenceTickManager::Get(Widget);
-			Widget->AnimationTickManager->AddWidget(Widget);
-		}
-
-		return Widget->AnimationTickManager->GetLinker();
-	}
-
-	return UMovieSceneEntitySystemLinker::CreateLinker(Widget ? Widget->GetWorld() : nullptr, UE::MovieScene::EEntitySystemLinkerRole::UMG);
+	// Nothing, stubbed out.
+	ensureMsgf(false, TEXT("UUMGSequencePlayer shouldn't be doing anything except wrap a widget animation runner."));
 }
 
 void UUMGSequencePlayer::Tick(float DeltaTime)
 {
-	if (IsEvaluating())
-	{
-		BlockedDeltaTimeCompensation += DeltaTime;
-		return;
-	}
-
-	if (bIsStopping || bIsBeginningPlay)
+	if (FWidgetAnimationState* AnimState = WidgetAnimationHandle.GetAnimationState())
 	{
-		return;
-	}
-
-	DeltaTime += BlockedDeltaTimeCompensation;
-	BlockedDeltaTimeCompensation = 0.f;
-
-	if ( PlayerStatus == EMovieScenePlayerStatus::Playing )
-	{
-		FFrameTime DeltaFrameTime = (bIsPlayingForward ? DeltaTime * PlaybackSpeed : -DeltaTime * PlaybackSpeed) * AnimationResolution;
-
-		UnwarpedPosition += DeltaFrameTime;
-
-		FFrameTime LastTimePosition = TimeCursorPosition;
-
-		const FMovieSceneSequenceHierarchy* Hierarchy = RootTemplateInstance.GetHierarchy();
-		FMovieSceneSequenceTransform RootTransform;
-		if (Hierarchy && Hierarchy->GetRootTransform().FindFirstWarpDomain() == UE::MovieScene::ETimeWarpChannelDomain::PlayRate)
-		{
-			RootTransform = Hierarchy->GetRootTransform();
-		}
-
-		TimeCursorPosition = RootTransform.TransformTime(UnwarpedPosition);
-
-		// Check if we crossed over bounds
-		const bool bCrossedLowerBound = TimeCursorPosition < FFrameTime(0);
-		const bool bCrossedUpperBound = TimeCursorPosition >= FFrameTime(Duration);
-		const bool bCrossedEndTime = bIsPlayingForward
-			? LastTimePosition < EndTime && EndTime <= TimeCursorPosition
-			: LastTimePosition > EndTime && EndTime >= TimeCursorPosition;
-
-		// Increment the num loops if we crossed any bounds.
-		if (bCrossedLowerBound || bCrossedUpperBound || (bCrossedEndTime && NumLoopsCompleted >= NumLoopsToPlay - 1))
-		{
-			NumLoopsCompleted++;
-		}
-
-		// Did the animation complete
-		const bool bCompleted = NumLoopsToPlay != 0 && NumLoopsCompleted >= NumLoopsToPlay;
-
-		// Handle times and see if we need to loop or ping-pong. If looping/ping-ponging,
-		// we update the sequence for the last little bit before it happens.
-		const FFrameTime LastValidFrame(Duration-1, 0.99999994f);
-		if (bCrossedLowerBound)
-		{
-			if (bCompleted)
-			{
-				TimeCursorPosition = UnwarpedPosition = FFrameTime(0);
-			}
-			else
-			{
-				UpdateInternal(LastTimePosition, FFrameTime(0), false);
-
-				if (PlayMode == EUMGSequencePlayMode::PingPong)
-				{
-					bIsPlayingForward = !bIsPlayingForward;
-					TimeCursorPosition = FMath::Abs(TimeCursorPosition);
-					LastTimePosition = FFrameTime(0);
-				}
-				else
-				{
-					TimeCursorPosition += FFrameTime(Duration);
-					LastTimePosition = LastValidFrame;
-				}
-			}
-		}
-		else if (bCrossedUpperBound)
-		{
-			if (bCompleted)
-			{
-				TimeCursorPosition = LastValidFrame;
-			}
-			else
-			{
-				UpdateInternal(LastTimePosition, LastValidFrame, false);
-
-				if (PlayMode == EUMGSequencePlayMode::PingPong)
-				{
-					bIsPlayingForward = !bIsPlayingForward;
-					TimeCursorPosition = LastValidFrame - (TimeCursorPosition - FFrameTime(Duration));
-					LastTimePosition = LastValidFrame;
-				}
-				else
-				{
-					TimeCursorPosition = TimeCursorPosition - FFrameTime(Duration);
-					LastTimePosition = FFrameTime(0);
-				}
-			}
-		}
-		else if (bCrossedEndTime)
-		{
-			if (bCompleted)
-			{
-				TimeCursorPosition = EndTime;
-			}
-		}
-
-		if (bCrossedLowerBound || bCrossedUpperBound || bCrossedEndTime)
-		{
-			UnwarpedPosition = RootTransform.Inverse().TryTransformTime(TimeCursorPosition).Get(TimeCursorPosition);
-		}
-
-		bCompleteOnPostEvaluation = bCompleted;
-
-		const bool bHasJumped = (bCrossedLowerBound || bCrossedUpperBound || bCrossedEndTime);
-		UpdateInternal(LastTimePosition, TimeCursorPosition, bHasJumped);
-	}
-}
-
-void UUMGSequencePlayer::UpdateInternal(FFrameTime LastTimePosition, FFrameTime NextTimePosition, bool bHasJumped)
-{
-	if (RootTemplateInstance.IsValid())
-	{
-		UMovieScene* MovieScene = Animation->GetMovieScene();
-
-		FMovieSceneContext Context(
-				FMovieSceneEvaluationRange(
-					AbsolutePlaybackStart + TimeCursorPosition,
-					AbsolutePlaybackStart + LastTimePosition,
-					AnimationResolution),
-				PlayerStatus);
-		Context.SetHasJumped(bHasJumped);
-
-		TSharedPtr<FMovieSceneEntitySystemRunner> Runner = RootTemplateInstance.GetRunner();
-		if (Runner)
-		{
-			UE::MovieScene::ERunnerUpdateFlags UpdateFlags = UE::MovieScene::ERunnerUpdateFlags::None;
-			FSimpleDelegate OnComplete;
-
-			if (bCompleteOnPostEvaluation)
-			{
-				bIsStopping = true;
-				bCompleteOnPostEvaluation = false;
-
-				OnComplete = FSimpleDelegate::CreateUObject(this, &UUMGSequencePlayer::HandleLatentStop);
-				UpdateFlags = UE::MovieScene::ERunnerUpdateFlags::Flush;
-			}
-
-			Runner->QueueUpdate(Context, RootTemplateInstance.GetRootInstanceHandle(), MoveTemp(OnComplete), UpdateFlags);
-
-			if (Runner == SynchronousRunner || !UE::UMG::GAsyncAnimationControlFlow)
-			{
-				Runner->Flush();
-				ApplyLatentActions();
-			}
-		}
-	}
-}
-
-void UUMGSequencePlayer::PlayInternal(double StartAtTime, double EndAtTime, int32 InNumLoopsToPlay, EUMGSequencePlayMode::Type InPlayMode, float InPlaybackSpeed, bool bInRestoreState)
-{
-	UUserWidget* Widget = UserWidget.Get();
-	UUMGSequenceTickManager* TickManager = Widget ? ToRawPtr(Widget->AnimationTickManager) : nullptr;
-
-	if (UE::UMG::GVarAnimationMarkers && Animation && Widget)
-	{
-		CSV_EVENT_GLOBAL(TEXT("Play Animation [%s::%s]"), *Widget->GetName(), *Animation->GetName());
-	}
-
-	RootTemplateInstance.Initialize(*Animation, *this, nullptr);
-
-	if (bInRestoreState)
-	{
-		RootTemplateInstance.EnableGlobalPreAnimatedStateCapture();
-	}
-
-	bIsBeginningPlay = true;
-	bRestoreState = bInRestoreState;
-	PlaybackSpeed = FMath::Abs(InPlaybackSpeed);
-	PlayMode = InPlayMode;
-
-	FFrameTime LastValidFrame(Duration-1, 0.99999994f);
-
-	if (PlayMode == EUMGSequencePlayMode::Reverse)
-	{
-		// When playing in reverse count subtract the start time from the end.
-		TimeCursorPosition = LastValidFrame - StartAtTime * AnimationResolution;
-	}
-	else
-	{
-		TimeCursorPosition = StartAtTime * AnimationResolution;
-	}
-
-	// Clamp the start time and end time to be within the bounds
-	TimeCursorPosition = FMath::Clamp(TimeCursorPosition, FFrameTime(0), LastValidFrame);
-
-	const FMovieSceneSequenceHierarchy* Hierarchy = RootTemplateInstance.GetHierarchy();
-	if (Hierarchy && Hierarchy->GetRootTransform().FindFirstWarpDomain() == UE::MovieScene::ETimeWarpChannelDomain::PlayRate)
-	{
-		UnwarpedPosition = Hierarchy->GetRootTransform().Inverse().TryTransformTime(TimeCursorPosition).Get(TimeCursorPosition);
-	}
-	else
-	{
-		UnwarpedPosition = TimeCursorPosition;
-	}
-
-	EndTime = FMath::Clamp(EndAtTime * AnimationResolution, FFrameTime(0), LastValidFrame);
-
-	if ( PlayMode == EUMGSequencePlayMode::PingPong )
-	{
-		// When animating in ping-pong mode double the number of loops to play so that a loop is a complete forward/reverse cycle.
-		NumLoopsToPlay = 2 * InNumLoopsToPlay;
-	}
-	else
-	{
-		NumLoopsToPlay = InNumLoopsToPlay;
-	}
-
-	NumLoopsCompleted = 0;
-	bIsPlayingForward = InPlayMode != EUMGSequencePlayMode::Reverse;
-
-	PlayerStatus = EMovieScenePlayerStatus::Playing;
-
-	// Playback assumes the start frame has already been evaulated, so we also want to evaluate any events on the start frame here.
-	TSharedPtr<FMovieSceneEntitySystemRunner> Runner = RootTemplateInstance.GetRunner();
-	if (Runner)
-	{
-		const FMovieSceneContext Context(FMovieSceneEvaluationRange(AbsolutePlaybackStart + TimeCursorPosition, AbsolutePlaybackStart + TimeCursorPosition, AnimationResolution), PlayerStatus);
-
-		auto OnBegunPlay = [this]
-		{
-			this->bIsBeginningPlay = false;
-		};
-
-		// We queue an update instead of immediately flushing the entire linker so that we don't incur a cascade of flushes on frames when multiple animations are played
-		// In rare cases where the linker must be flushed immediately PreTick, the queue should be manually flushed 
-		Runner->QueueUpdate(Context, RootTemplateInstance.GetRootInstanceHandle(), FSimpleDelegate::CreateWeakLambda(this, OnBegunPlay), UE::MovieScene::ERunnerUpdateFlags::Flush);
-
-		if (Runner == SynchronousRunner || !UE::UMG::GAsyncAnimationControlFlow)
-		{
-			Runner->Flush();
-		}
+		AnimState->Tick(DeltaTime);
 	}
 }
 
 void UUMGSequencePlayer::Play(float StartAtTime, int32 InNumLoopsToPlay, EUMGSequencePlayMode::Type InPlayMode, float InPlaybackSpeed, bool bInRestoreState)
 {
-	if (NeedsQueueLatentAction())
+	if (FWidgetAnimationState* AnimState = WidgetAnimationHandle.GetAnimationState())
 	{
-		QueueLatentAction(FMovieSceneSequenceLatentActionDelegate::CreateUObject(
-					this, &UUMGSequencePlayer::Play, StartAtTime, InNumLoopsToPlay, InPlayMode, InPlaybackSpeed, bInRestoreState));
-		return;
+		FWidgetAnimationStatePlayParams PlayParams;
+		PlayParams.StartAtTime = StartAtTime;
+		PlayParams.NumLoopsToPlay = InNumLoopsToPlay;
+		PlayParams.PlayMode = InPlayMode;
+		PlayParams.PlaybackSpeed = InPlaybackSpeed;
+		PlayParams.bRestoreState = bInRestoreState;
+		AnimState->Play(PlayParams);
 	}
-
-	PlayInternal(StartAtTime, 0.0, InNumLoopsToPlay, InPlayMode, InPlaybackSpeed, bInRestoreState);
 }
 
 void UUMGSequencePlayer::PlayTo(float StartAtTime, float EndAtTime, int32 InNumLoopsToPlay, EUMGSequencePlayMode::Type InPlayMode, float InPlaybackSpeed, bool bInRestoreState)
 {
-	if (NeedsQueueLatentAction())
+	if (FWidgetAnimationState* AnimState = WidgetAnimationHandle.GetAnimationState())
 	{
-		QueueLatentAction(FMovieSceneSequenceLatentActionDelegate::CreateUObject(
-					this, &UUMGSequencePlayer::PlayTo, StartAtTime, EndAtTime, InNumLoopsToPlay, InPlayMode, InPlaybackSpeed, bInRestoreState));
-		return;
+		FWidgetAnimationStatePlayParams PlayParams;
+		PlayParams.StartAtTime = StartAtTime;
+		PlayParams.EndOffset = EndAtTime;
+		PlayParams.NumLoopsToPlay = InNumLoopsToPlay;
+		PlayParams.PlayMode = InPlayMode;
+		PlayParams.PlaybackSpeed = InPlaybackSpeed;
+		PlayParams.bRestoreState = bInRestoreState;
+		AnimState->Play(PlayParams);
 	}
-
-	PlayInternal(StartAtTime, EndAtTime, InNumLoopsToPlay, InPlayMode, InPlaybackSpeed, bInRestoreState);
 }
 
 void UUMGSequencePlayer::Pause()
 {
-	if (NeedsQueueLatentAction())
+	if (FWidgetAnimationState* AnimState = WidgetAnimationHandle.GetAnimationState())
 	{
-		QueueLatentAction(FMovieSceneSequenceLatentActionDelegate::CreateUObject(this, &UUMGSequencePlayer::Pause));
-		return;
-	}
-
-	// Purposely don't trigger any OnFinished events
-	PlayerStatus = EMovieScenePlayerStatus::Stopped;
-
-	// Evaluate the sequence at its current time, with a status of 'stopped' to ensure that animated state pauses correctly. (ie. audio sounds should stop/pause)
-	TSharedPtr<FMovieSceneEntitySystemRunner> RunnerToUpdate = RootTemplateInstance.HasEverUpdated() ? RootTemplateInstance.GetRunner() : nullptr;
-	if (RunnerToUpdate)
-	{
-		const FMovieSceneContext Context(FMovieSceneEvaluationRange(AbsolutePlaybackStart + TimeCursorPosition, AbsolutePlaybackStart + TimeCursorPosition, AnimationResolution), PlayerStatus);
-		RunnerToUpdate->QueueUpdate(Context, RootTemplateInstance.GetRootInstanceHandle(), UE::MovieScene::ERunnerUpdateFlags::Flush);
-
-		if (RunnerToUpdate == SynchronousRunner || !UE::UMG::GAsyncAnimationControlFlow)
-		{
-			RunnerToUpdate->Flush();
-		}
+		AnimState->Pause();
 	}
 }
 
 void UUMGSequencePlayer::Reverse()
 {
-	if (PlayerStatus == EMovieScenePlayerStatus::Playing)
+	if (FWidgetAnimationState* AnimState = WidgetAnimationHandle.GetAnimationState())
 	{
-		bIsPlayingForward = !bIsPlayingForward;
+		AnimState->Reverse();
 	}
 }
 
 void UUMGSequencePlayer::Stop()
 {
-	using namespace UE::MovieScene;
-
-	if (NeedsQueueLatentAction())
+	if (FWidgetAnimationState* AnimState = WidgetAnimationHandle.GetAnimationState())
 	{
-		QueueLatentAction(FMovieSceneSequenceLatentActionDelegate::CreateUObject(this, &UUMGSequencePlayer::Stop));
-		return;
+		AnimState->Stop();
 	}
+}
 
-	if (PlayerStatus == EMovieScenePlayerStatus::Stopped)
+void UUMGSequencePlayer::SetCurrentTime(float InTime)
+{
+	if (FWidgetAnimationState* AnimState = WidgetAnimationHandle.GetAnimationState())
 	{
-		return;
+		AnimState->SetCurrentTime(InTime);
 	}
+}
 
-	PlayerStatus = EMovieScenePlayerStatus::Stopped;
-
-	UMovieSceneSequence* MovieSceneSequence = RootTemplateInstance.GetSequence(MovieSceneSequenceID::Root);
-	const bool bIsRunningWithTickManager = CVarUserWidgetUseParallelAnimation.GetValueOnGameThread();
-	const bool bIsSequenceBlocking = EnumHasAnyFlags(MovieSceneSequence->GetFlags(), EMovieSceneSequenceFlags::BlockingEvaluation);
-	const bool bIsAsync = bIsRunningWithTickManager && !bIsSequenceBlocking && UE::UMG::GAsyncAnimationControlFlow == true;
-
-	UUserWidget* Widget = UserWidget.Get();
-	UUMGSequenceTickManager* TickManager = Widget ? ToRawPtr(Widget->AnimationTickManager) : nullptr;
-
-	TimeCursorPosition = UnwarpedPosition = FFrameTime(0);
-
-	if (!TickManager || !RootTemplateInstance.IsValid())
+FQualifiedFrameTime UUMGSequencePlayer::GetCurrentTime() const
+{
+	if (FWidgetAnimationState* AnimState = WidgetAnimationHandle.GetAnimationState())
 	{
-		HandleLatentStop();
-		return;
+		return AnimState->GetCurrentTime();
 	}
+	return FQualifiedFrameTime();
+}
 
-	if (!RootTemplateInstance.HasEverUpdated())
+const UWidgetAnimation* UUMGSequencePlayer::GetAnimation() const
+{
+	if (FWidgetAnimationState* AnimState = WidgetAnimationHandle.GetAnimationState())
 	{
-		TickManager->ClearLatentActions(this);
-		LatentActions.Empty();
-		HandleLatentStop();
-		return;
+		AnimState->GetAnimation();
 	}
+	return nullptr;
+}
 
-	const FMovieSceneContext Context(FMovieSceneEvaluationRange(AbsolutePlaybackStart, AnimationResolution), PlayerStatus);
-
-	// Prevent any other updates
-	bIsStopping = true;
-
-	if (TSharedPtr<FMovieSceneEntitySystemRunner> Runner = RootTemplateInstance.GetRunner())
+FName UUMGSequencePlayer::GetUserTag() const
+{
+	if (FWidgetAnimationState* AnimState = WidgetAnimationHandle.GetAnimationState())
 	{
-		// It would be better if we did not have to actively re-evaluate the animation at the start time when Stop is called
-		// so we don't have to queue an additional update, but this has to be included for legacy reasons.
-		Runner->QueueUpdate(Context, RootTemplateInstance.GetRootInstanceHandle(), FSimpleDelegate::CreateUObject(this, &UUMGSequencePlayer::HandleLatentStop), UE::MovieScene::ERunnerUpdateFlags::Flush);
-
-		if (Runner == SynchronousRunner || !UE::UMG::GAsyncAnimationControlFlow)
-		{
-			Runner->Flush();
-		}
+		return AnimState->GetUserTag();
 	}
+	return NAME_None;
 }
 
-void UUMGSequencePlayer::HandleLatentStop()
+void UUMGSequencePlayer::SetUserTag(FName InUserTag)
 {
-	if (TSharedPtr<FMovieSceneEntitySystemRunner> Runner = RootTemplateInstance.GetRunner())
+	if (FWidgetAnimationState* AnimState = WidgetAnimationHandle.GetAnimationState())
 	{
-		Runner->QueueFinalUpdate(RootTemplateInstance.GetRootInstanceHandle());
-
-		// Even if our request to Finish the instance was queued, we can wait until the next flush for those effects to be seen
-		// This will most likely happen immediately anyway since the runner will keep looping until its queue is empty,
-		// And we are already inside an active evaluation
+		return AnimState->SetUserTag(InUserTag);
 	}
+}
 
-	PlayerStatus = EMovieScenePlayerStatus::Stopped;
-
-	PRAGMA_DISABLE_DEPRECATION_WARNINGS
-	RootTemplateInstance.ResetDirectorInstances();
-	PRAGMA_ENABLE_DEPRECATION_WARNINGS
-
-	if (bRestoreState)
+void UUMGSequencePlayer::SetNumLoopsToPlay(int32 InNumLoopsToPlay)
+{
+	if (FWidgetAnimationState* AnimState = WidgetAnimationHandle.GetAnimationState())
 	{
-		RestorePreAnimatedState();
+		AnimState->SetNumLoopsToPlay(InNumLoopsToPlay);
 	}
+}
 
-	bIsStopping = false;
-
-	UUserWidget* Widget = UserWidget.Get();
-	if (Widget)
+void UUMGSequencePlayer::SetPlaybackSpeed(float InPlaybackSpeed)
+{
+	if (FWidgetAnimationState* AnimState = WidgetAnimationHandle.GetAnimationState())
 	{
-		if (UE::UMG::GVarAnimationMarkers && Animation)
-		{
-			CSV_EVENT_GLOBAL(TEXT("Stop Animation [%s::%s]"), *Widget->GetName(), *Animation->GetName());
-		}
-
-		Widget->OnAnimationFinishedPlaying(*this);
+		AnimState->SetPlaybackSpeed(InPlaybackSpeed);
 	}
-
-	OnSequenceFinishedPlayingEvent.Broadcast(*this);
 }
 
-
-void UUMGSequencePlayer::SetNumLoopsToPlay(int32 InNumLoopsToPlay)
+bool UUMGSequencePlayer::IsPlayingForward() const
 {
-	if (PlayMode == EUMGSequencePlayMode::PingPong)
+	if (FWidgetAnimationState* AnimState = WidgetAnimationHandle.GetAnimationState())
 	{
-		NumLoopsToPlay = (2 * InNumLoopsToPlay);
+		return AnimState->IsPlayingForward();
 	}
-	else
+	return true;
+}
+
+bool UUMGSequencePlayer::IsStopping() const
+{
+	if (FWidgetAnimationState* AnimState = WidgetAnimationHandle.GetAnimationState())
 	{
-		NumLoopsToPlay = InNumLoopsToPlay;
+		return AnimState->IsStopping();
 	}
+	return false;
 }
 
-void UUMGSequencePlayer::SetPlaybackSpeed(float InPlaybackSpeed)
+FMovieSceneRootEvaluationTemplateInstance& UUMGSequencePlayer::GetEvaluationTemplate()
 {
-	PlaybackSpeed = InPlaybackSpeed;
+	return RootTemplateInstance;
 }
 
-EMovieScenePlayerStatus::Type UUMGSequencePlayer::GetPlaybackStatus() const
+UMovieSceneEntitySystemLinker* UUMGSequencePlayer::ConstructEntitySystemLinker()
 {
-	return PlayerStatus;
+	ensureMsgf(false, TEXT("This legacy player should never have to construct a linker."));
+	return nullptr;
 }
 
-UObject* UUMGSequencePlayer::GetPlaybackContext() const
+UObject* UUMGSequencePlayer::AsUObject()
 {
-	return UserWidget.Get();
+	return this;
 }
 
-TArray<UObject*> UUMGSequencePlayer::GetEventContexts() const
+EMovieScenePlayerStatus::Type UUMGSequencePlayer::GetPlaybackStatus() const
 {
-	TArray<UObject*> EventContexts;
-	if (UserWidget.IsValid())
+	if (FWidgetAnimationState* AnimState = WidgetAnimationHandle.GetAnimationState())
 	{
-		EventContexts.Add(UserWidget.Get());
+		return AnimState->GetPlaybackStatus();
 	}
-	return EventContexts;
+	return EMovieScenePlayerStatus::Stopped;
 }
 
 void UUMGSequencePlayer::SetPlaybackStatus(EMovieScenePlayerStatus::Type InPlaybackStatus)
 {
-	PlayerStatus = InPlaybackStatus;
-}
-
-void UUMGSequencePlayer::PopulateUpdateFlags(UE::MovieScene::ESequenceInstanceUpdateFlags& OutFlags)
-{
-	// Leave empty so Pre and PostEvaluation are not called
+	if (FWidgetAnimationState* AnimState = WidgetAnimationHandle.GetAnimationState())
+	{
+		AnimState->SetPlaybackStatus(InPlaybackStatus);
+	}
 }
 
-bool UUMGSequencePlayer::NeedsQueueLatentAction() const
+IMovieScenePlaybackClient* UUMGSequencePlayer::GetPlaybackClient()
 {
-	return IsEvaluating();
+	if (FWidgetAnimationState* AnimState = WidgetAnimationHandle.GetAnimationState())
+	{
+		return AnimState->GetSharedPlaybackState()->FindCapability<IMovieScenePlaybackClient>();
+	}
+	return nullptr;
 }
 
-void UUMGSequencePlayer::QueueLatentAction(FMovieSceneSequenceLatentActionDelegate Delegate)
+FMovieSceneSpawnRegister& UUMGSequencePlayer::GetSpawnRegister()
 {
-	if (CVarUserWidgetUseParallelAnimation.GetValueOnGameThread())
+	if (FWidgetAnimationState* AnimState = WidgetAnimationHandle.GetAnimationState())
 	{
-		UUserWidget* Widget = UserWidget.Get();
-		UUMGSequenceTickManager* TickManager = Widget ? ToRawPtr(Widget->AnimationTickManager) : nullptr;
-
-		if (ensure(TickManager))
+		if (FMovieSceneSpawnRegister* SpawnRegister = AnimState->GetSharedPlaybackState()->FindCapability<FMovieSceneSpawnRegister>())
 		{
-			TickManager->AddLatentAction(Delegate);
+			return *SpawnRegister;
 		}
 	}
-	else
-	{
-		LatentActions.Add(Delegate);
-	}
+	return IMovieScenePlayer::GetSpawnRegister();
 }
 
-void UUMGSequencePlayer::ApplyLatentActions()
+UObject* UUMGSequencePlayer::GetPlaybackContext() const
 {
-	if (CVarUserWidgetUseParallelAnimation.GetValueOnGameThread())
-	{
-		UUserWidget* Widget = UserWidget.Get();
-		UUMGSequenceTickManager* TickManager = Widget ? ToRawPtr(Widget->AnimationTickManager) : nullptr;
-		if (TickManager)
-		{
-			TickManager->RunLatentActions();
-		}
-	}
-	else
+	if (FWidgetAnimationState* AnimState = WidgetAnimationHandle.GetAnimationState())
 	{
-		while (LatentActions.Num() > 0)
-		{
-			const FMovieSceneSequenceLatentActionDelegate& Delegate = LatentActions[0];
-			Delegate.ExecuteIfBound();
-			LatentActions.RemoveAt(0);
-		}
+		return AnimState->GetUserWidget();
 	}
+	return nullptr;
 }
 
-void UUMGSequencePlayer::RemoveEvaluationData()
+void UUMGSequencePlayer::InitializeRootInstance(TSharedRef<UE::MovieScene::FSharedPlaybackState> NewSharedPlaybackState)
 {
-	using namespace UE::MovieScene;
-
-	UMovieSceneEntitySystemLinker* Linker           = RootTemplateInstance.GetEntitySystemLinker();
-	FSequenceInstance*             SequenceInstance = RootTemplateInstance.FindInstance(MovieSceneSequenceID::Root);
+	ensureMsgf(false, TEXT("The legacy player should never initialize sequences: it only wraps already initialized ones."));
+}
 
-	if (SequenceInstance && Linker)
+void UUMGSequencePlayer::RemoveEvaluationData()
+{
+	if (FWidgetAnimationState* AnimState = WidgetAnimationHandle.GetAnimationState())
 	{
-		TSharedPtr<FMovieSceneEntitySystemRunner> Runner = RootTemplateInstance.GetRunner();
-		if (Runner && Runner->IsCurrentlyEvaluating())
-		{
-			Runner->FlushOutstanding();
-		}
-
-		SequenceInstance->Ledger.UnlinkEverything(Linker);
-		SequenceInstance->InvalidateCachedData();
+		return AnimState->RemoveEvaluationData();
 	}
-
-	PRAGMA_DISABLE_DEPRECATION_WARNINGS
-	RootTemplateInstance.ResetDirectorInstances();
-	PRAGMA_ENABLE_DEPRECATION_WARNINGS
 }
 
 void UUMGSequencePlayer::TearDown()
 {
-	RootTemplateInstance.TearDown();
 }
 
 void UUMGSequencePlayer::BeginDestroy()
 {
-	RootTemplateInstance.TearDown();
-
-	// Remove any latent actions added by this player.
-	if (CVarUserWidgetUseParallelAnimation.GetValueOnGameThread())
-	{
-		UUserWidget* Widget = UserWidget.Get();
-		UUMGSequenceTickManager* TickManager = Widget ? ToRawPtr(Widget->AnimationTickManager) : nullptr;
-		if (TickManager)
-		{
-			TickManager->ClearLatentActions(this);
-		}
-	}
-	else
-	{
-		LatentActions.Empty();
-	}
-
 	Super::BeginDestroy();
 }
 
+UUMGSequencePlayer::FOnSequenceFinishedPlaying& UUMGSequencePlayer::OnSequenceFinishedPlaying()
+{
+	return OnSequenceFinishedPlayingEvent;
+}
+
+void UUMGSequencePlayer::BroadcastSequenceFinishedPlaying()
+{
+	OnSequenceFinishedPlayingEvent.Broadcast(*this);
+}
 
diff --git a/Engine/Source/Runtime/UMG/Private/Animation/UMGSequenceTickManager.cpp b/Engine/Source/Runtime/UMG/Private/Animation/UMGSequenceTickManager.cpp
--- a/Engine/Source/Runtime/UMG/Private/Animation/UMGSequenceTickManager.cpp
+++ b/Engine/Source/Runtime/UMG/Private/Animation/UMGSequenceTickManager.cpp
@@ -40,8 +40,6 @@ namespace UE::UMG
 	);
 } // namespace UE::UMG
 
-extern TAutoConsoleVariable<bool> CVarUserWidgetUseParallelAnimation;
-
 UUMGSequenceTickManager::UUMGSequenceTickManager(const FObjectInitializer& Init)
 	: Super(Init)
 	, bIsTicking(false)
@@ -119,11 +117,6 @@ void UUMGSequenceTickManager::BeginDestroy()
 
 void UUMGSequenceTickManager::TickWidgetAnimations(float DeltaSeconds)
 {
-	if (!CVarUserWidgetUseParallelAnimation.GetValueOnGameThread())
-	{
-		return;
-	}
-
 	if (bIsTicking)
 	{
 		return;
@@ -237,7 +230,7 @@ void UUMGSequenceTickManager::TickWidgetAnimations(float DeltaSeconds)
 			if (UserWidget)
 			{
 				// If this widget no longer has any animations playing, it doesn't need to be ticked any more
-				if (UserWidget->ActiveSequencePlayers.Num() == 0)
+				if (!UserWidget->IsAnyAnimationPlaying())
 				{
 					UserWidget->UpdateCanTick();
 					UserWidget->AnimationTickManager = nullptr;
diff --git a/Engine/Source/Runtime/UMG/Private/Animation/WidgetAnimationHandle.cpp b/Engine/Source/Runtime/UMG/Private/Animation/WidgetAnimationHandle.cpp
new file mode 100644
--- /dev/null
+++ b/Engine/Source/Runtime/UMG/Private/Animation/WidgetAnimationHandle.cpp
@@ -0,0 +1,77 @@
+// Copyright Epic Games, Inc. All Rights Reserved.
+
+#include "Animation/WidgetAnimationHandle.h"
+
+#include "Animation/UMGSequencePlayer.h"
+#include "Animation/WidgetAnimationState.h"
+#include "Blueprint/UserWidget.h"
+#include "EdGraph/EdGraphPin.h"
+
+FWidgetAnimationHandle::FWidgetAnimationHandle()
+{
+}
+
+FWidgetAnimationHandle::FWidgetAnimationHandle(UUserWidget* InUserWidget, int32 InStateIndex, uint32 InStateSerial)
+	: WeakUserWidget(InUserWidget)
+	, StateIndex(InStateIndex)
+	, StateSerial(InStateSerial)
+{
+}
+
+bool FWidgetAnimationHandle::IsValid() const
+{
+	return (WeakUserWidget.IsValid() && StateIndex >= 0);
+}
+
+UUMGSequencePlayer* FWidgetAnimationHandle::GetSequencePlayer() const
+{
+	if (FWidgetAnimationState* State = GetAnimationState())
+	{
+		return State->GetOrCreateLegacyPlayer();
+	}
+	return nullptr;
+}
+
+FWidgetAnimationState* FWidgetAnimationHandle::GetAnimationState() const
+{
+	if (UUserWidget* UserWidget = WeakUserWidget.Get())
+	{
+		if (UserWidget->ActiveAnimations.IsValidIndex(StateIndex))
+		{
+			FWidgetAnimationState& ActiveState = UserWidget->ActiveAnimations[StateIndex];
+			if (ActiveState.SerialNumber == StateSerial)
+			{
+				return &ActiveState;
+			}
+		}
+	}
+	return nullptr;
+}
+
+FName FWidgetAnimationHandle::GetUserTag() const
+{
+	if (FWidgetAnimationState* State = GetAnimationState())
+	{
+		return State->GetUserTag();
+	}
+	return NAME_None;
+}
+
+void FWidgetAnimationHandle::SetUserTag(FName InUserTag)
+{
+	if (FWidgetAnimationState* State = GetAnimationState())
+	{
+		State->SetUserTag(InUserTag);
+	}
+}
+
+FName UWidgetAnimationHandleFunctionLibrary::GetUserTag(const FWidgetAnimationHandle& Target)
+{
+	return Target.GetUserTag();
+}
+
+void UWidgetAnimationHandleFunctionLibrary::SetUserTag(FWidgetAnimationHandle& Target, FName InUserTag)
+{
+	Target.SetUserTag(InUserTag);
+}
+
diff --git a/Engine/Source/Runtime/UMG/Private/Animation/WidgetAnimationPlayCallbackProxy.cpp b/Engine/Source/Runtime/UMG/Private/Animation/WidgetAnimationPlayCallbackProxy.cpp
--- a/Engine/Source/Runtime/UMG/Private/Animation/WidgetAnimationPlayCallbackProxy.cpp
+++ b/Engine/Source/Runtime/UMG/Private/Animation/WidgetAnimationPlayCallbackProxy.cpp
@@ -11,19 +11,67 @@
 
 #define LOCTEXT_NAMESPACE "UMG"
 
-UWidgetAnimationPlayCallbackProxy* UWidgetAnimationPlayCallbackProxy::CreatePlayAnimationProxyObject(class UUMGSequencePlayer*& Result, class UUserWidget* Widget, UWidgetAnimation* InAnimation, float StartAtTime, int32 NumLoopsToPlay, EUMGSequencePlayMode::Type PlayMode, float PlaybackSpeed)
+UWidgetAnimationPlayCallbackProxy* UWidgetAnimationPlayCallbackProxy::CreatePlayAnimationProxyObject(
+		UUMGSequencePlayer*& Result,
+		UUserWidget* Widget,
+		UWidgetAnimation* InAnimation,
+		float StartAtTime,
+		int32 NumLoopsToPlay,
+		EUMGSequencePlayMode::Type PlayMode,
+		float PlaybackSpeed)
+{
+	FWidgetAnimationHandle Unused;
+	UWidgetAnimationPlayCallbackProxy* Proxy = NewPlayAnimationProxyObject(Unused, Widget, InAnimation, StartAtTime, NumLoopsToPlay, PlayMode, PlaybackSpeed);
+	Result = Proxy->WidgetAnimationHandle.GetSequencePlayer();
+	return Proxy;
+}
+
+UWidgetAnimationPlayCallbackProxy* UWidgetAnimationPlayCallbackProxy::NewPlayAnimationProxyObject(
+		FWidgetAnimationHandle& Result,
+		UUserWidget* Widget,
+		UWidgetAnimation* InAnimation,
+		float StartAtTime,
+		int32 NumLoopsToPlay,
+		EUMGSequencePlayMode::Type PlayMode,
+		float PlaybackSpeed)
 {
 	UWidgetAnimationPlayCallbackProxy* Proxy = NewObject<UWidgetAnimationPlayCallbackProxy>();
 	Proxy->SetFlags(RF_StrongRefOnFrame);
-	Result = Proxy->ExecutePlayAnimation(Widget, InAnimation, StartAtTime, NumLoopsToPlay, PlayMode, PlaybackSpeed);
+	Proxy->ExecutePlayAnimation(Widget, InAnimation, StartAtTime, NumLoopsToPlay, PlayMode, PlaybackSpeed);
+	Result = Proxy->WidgetAnimationHandle;
 	return Proxy;
 }
 
-UWidgetAnimationPlayCallbackProxy* UWidgetAnimationPlayCallbackProxy::CreatePlayAnimationTimeRangeProxyObject(class UUMGSequencePlayer*& Result, class UUserWidget* Widget, UWidgetAnimation* InAnimation, float StartAtTime, float EndAtTime, int32 NumLoopsToPlay, EUMGSequencePlayMode::Type PlayMode, float PlaybackSpeed)
+UWidgetAnimationPlayCallbackProxy* UWidgetAnimationPlayCallbackProxy::CreatePlayAnimationTimeRangeProxyObject(
+		UUMGSequencePlayer*& Result,
+		UUserWidget* Widget,
+		UWidgetAnimation* InAnimation,
+		float StartAtTime,
+		float EndAtTime,
+		int32 NumLoopsToPlay,
+		EUMGSequencePlayMode::Type PlayMode,
+		float PlaybackSpeed)
+{
+	FWidgetAnimationHandle Unused;
+	UWidgetAnimationPlayCallbackProxy* Proxy = NewPlayAnimationTimeRangeProxyObject(Unused, Widget, InAnimation, StartAtTime, EndAtTime, NumLoopsToPlay, PlayMode, PlaybackSpeed);
+	Result = Proxy->WidgetAnimationHandle.GetSequencePlayer();
+	return Proxy;
+}
+
+UWidgetAnimationPlayCallbackProxy* UWidgetAnimationPlayCallbackProxy::NewPlayAnimationTimeRangeProxyObject(
+		FWidgetAnimationHandle& Result,
+		UUserWidget* Widget,
+		UWidgetAnimation* InAnimation,
+		float StartAtTime,
+		float EndAtTime,
+		int32 NumLoopsToPlay,
+		EUMGSequencePlayMode::Type PlayMode,
+		float PlaybackSpeed)
 {
 	UWidgetAnimationPlayCallbackProxy* Proxy = NewObject<UWidgetAnimationPlayCallbackProxy>();
 	Proxy->SetFlags(RF_StrongRefOnFrame);
-	Result = Proxy->ExecutePlayAnimationTimeRange(Widget, InAnimation, StartAtTime, EndAtTime, NumLoopsToPlay, PlayMode, PlaybackSpeed);
+	Proxy->ExecutePlayAnimationTimeRange(Widget, InAnimation, StartAtTime, EndAtTime, NumLoopsToPlay, PlayMode, PlaybackSpeed);
+	Result = Proxy->WidgetAnimationHandle;
 	return Proxy;
 }
 
@@ -32,41 +80,39 @@ UWidgetAnimationPlayCallbackProxy::UWidgetAnimationPlayCallbackProxy(const FObje
 {
 }
 
-class UUMGSequencePlayer* UWidgetAnimationPlayCallbackProxy::ExecutePlayAnimation(class UUserWidget* Widget, UWidgetAnimation* InAnimation, float StartAtTime, int32 NumLoopsToPlay, EUMGSequencePlayMode::Type PlayMode, float PlaybackSpeed)
+void UWidgetAnimationPlayCallbackProxy::ExecutePlayAnimation(UUserWidget* Widget, UWidgetAnimation* InAnimation, float StartAtTime, int32 NumLoopsToPlay, EUMGSequencePlayMode::Type PlayMode, float PlaybackSpeed)
 {
 	if (!Widget)
 	{
-		return nullptr;
+		return;
 	}
 
-	UUMGSequencePlayer* Player = Widget->PlayAnimation(InAnimation, StartAtTime, NumLoopsToPlay, PlayMode, PlaybackSpeed);
-	if (Player)
+	WidgetAnimationHandle = Widget->PlayAnimation(InAnimation, StartAtTime, NumLoopsToPlay, PlayMode, PlaybackSpeed);
+	FWidgetAnimationState* State = WidgetAnimationHandle.GetAnimationState();
+	if (State)
 	{
-		Player->OnSequenceFinishedPlaying().AddUObject(this, &UWidgetAnimationPlayCallbackProxy::OnSequenceFinished);
+		OnFinishedHandle = State->GetOnWidgetAnimationFinished().AddUObject(this, &UWidgetAnimationPlayCallbackProxy::OnSequenceFinished);
 	}
-
-	return Player;
 }
 
-class UUMGSequencePlayer* UWidgetAnimationPlayCallbackProxy::ExecutePlayAnimationTimeRange(class UUserWidget* Widget, UWidgetAnimation* InAnimation, float StartAtTime, float EndAtTime, int32 NumLoopsToPlay, EUMGSequencePlayMode::Type PlayMode, float PlaybackSpeed)
+void UWidgetAnimationPlayCallbackProxy::ExecutePlayAnimationTimeRange(UUserWidget* Widget, UWidgetAnimation* InAnimation, float StartAtTime, float EndAtTime, int32 NumLoopsToPlay, EUMGSequencePlayMode::Type PlayMode, float PlaybackSpeed)
 {
 	if (!Widget)
 	{
-		return nullptr;
+		return;
 	}
 
-	UUMGSequencePlayer* Player = Widget->PlayAnimationTimeRange(InAnimation, StartAtTime, EndAtTime, NumLoopsToPlay, PlayMode, PlaybackSpeed);
-	if (Player)
+	WidgetAnimationHandle = Widget->PlayAnimationTimeRange(InAnimation, StartAtTime, EndAtTime, NumLoopsToPlay, PlayMode, PlaybackSpeed);
+	FWidgetAnimationState* State = WidgetAnimationHandle.GetAnimationState();
+	if (State)
 	{
-		OnFinishedHandle = Player->OnSequenceFinishedPlaying().AddUObject(this, &UWidgetAnimationPlayCallbackProxy::OnSequenceFinished);
+		OnFinishedHandle = State->GetOnWidgetAnimationFinished().AddUObject(this, &UWidgetAnimationPlayCallbackProxy::OnSequenceFinished);
 	}
-
-	return Player;
 }
 
-void UWidgetAnimationPlayCallbackProxy::OnSequenceFinished(class UUMGSequencePlayer& Player)
+void UWidgetAnimationPlayCallbackProxy::OnSequenceFinished(FWidgetAnimationState& State)
 {
-	Player.OnSequenceFinishedPlaying().Remove(OnFinishedHandle);
+	State.GetOnWidgetAnimationFinished().Remove(OnFinishedHandle);
 
 	// We delay the Finish broadcast to next frame.
 	FTSTicker::FDelegateHandle TickerHandle = FTSTicker::GetCoreTicker().AddTicker(FTickerDelegate::CreateUObject(this, &UWidgetAnimationPlayCallbackProxy::OnAnimationFinished));
@@ -80,5 +126,6 @@ bool UWidgetAnimationPlayCallbackProxy::OnAnimationFinished(float /*DeltaTime*/)
 	// Returning false, disable the ticker.
 	return false;
 }
+
 #undef LOCTEXT_NAMESPACE
 
diff --git a/Engine/Source/Runtime/UMG/Private/Animation/WidgetAnimationState.cpp b/Engine/Source/Runtime/UMG/Private/Animation/WidgetAnimationState.cpp
new file mode 100644
--- /dev/null
+++ b/Engine/Source/Runtime/UMG/Private/Animation/WidgetAnimationState.cpp
@@ -0,0 +1,674 @@
+// Copyright Epic Games, Inc. All Rights Reserved.
+
+#include "Animation/WidgetAnimationState.h"
+
+#include "Animation/UMGSequencePlayer.h"
+#include "Animation/UMGSequenceTickManager.h"
+#include "Animation/WidgetAnimation.h"
+#include "Animation/WidgetAnimationHandle.h"
+#include "Blueprint/UserWidget.h"
+#include "EntitySystem/MovieSceneEntitySystemLinker.h"
+#include "EntitySystem/MovieSceneEntitySystemRunner.h"
+#include "Evaluation/MovieSceneEvaluationState.h"
+#include "Evaluation/MovieScenePlayback.h"
+#include "Evaluation/MovieSceneSequenceHierarchy.h"
+#include "MovieScene.h"
+#include "MovieSceneFwd.h"
+#include "MovieSceneLegacyPlayer.h"
+#include "MovieSceneTimeHelpers.h"
+#include "ProfilingDebugging/CsvProfiler.h"
+#include "UMGPrivate.h"
+
+namespace UE::UMG
+{
+
+	bool GVarAnimationDetailedLogging = false;
+	FAutoConsoleVariableRef CVarAnimationDetailedLogging(
+		TEXT("UMG.AnimationDetailedLogging"),
+		GVarAnimationDetailedLogging,
+		TEXT("(Default: false) Whether to print out detailed information about UMG animations.")
+	);
+
+	bool GVarAnimationMarkers = false;
+	FAutoConsoleVariableRef CVarAnimationMarkers(
+		TEXT("UMG.AnimationMarkers"),
+		GVarAnimationMarkers,
+		TEXT("(Default: false) Whether to emit profiling frame markers for starting and stopping UMG animations.")
+	);
+
+	struct FUMGLegacyPlayerProvider : public UE::MovieScene::ILegacyPlayerProviderPlaybackCapability
+	{
+		FWidgetAnimationState* State = nullptr;
+
+		FUMGLegacyPlayerProvider(FWidgetAnimationState* InState)
+			: State(InState)
+		{
+			check(State);
+		}
+
+		virtual IMovieScenePlayer* CreateLegacyPlayer(TSharedRef<UE::MovieScene::FSharedPlaybackState> InSharedPlaybackState) override
+		{
+			return State->GetOrCreateLegacyPlayer();
+		}
+	};
+
+} // namespace UE::UMG
+
+FWidgetAnimationState::FWidgetAnimationState()
+{
+	UserTag = NAME_None;
+	bRestoreState = false;
+	bIsBeginningPlay = false;
+	bIsStopping = false;
+}
+
+void FWidgetAnimationState::Initialize(UWidgetAnimation* InAnimation, UUserWidget* InUserWidget, int32 InIndex, int32 InSerialNumber)
+{
+	using namespace UE::UMG;
+	using namespace UE::MovieScene;
+
+	check(InAnimation && InUserWidget);
+
+	Animation = InAnimation;
+	WeakUserWidget = InUserWidget;
+	StateIndex = InIndex;
+	SerialNumber = InSerialNumber;
+
+	check(InUserWidget->AnimationTickManager);
+
+	const bool bNeedsLegacyPlayer = NeedsLegacyPlayer();
+
+	UMovieSceneEntitySystemLinker* Linker = InUserWidget->AnimationTickManager->GetLinker();
+	FInstanceRegistry* InstanceRegistry = Linker->GetInstanceRegistry();
+	FRootInstanceHandle RootInstanceHandle = InstanceRegistry->AllocateRootInstance(*InAnimation, InUserWidget);
+	{
+		FSequenceInstance& RootInstance = InstanceRegistry->MutateInstance(RootInstanceHandle);
+		TSharedRef<FSharedPlaybackState> SharedPlaybackState = RootInstance.GetSharedPlaybackState();
+		WeakPlaybackState = SharedPlaybackState;
+
+		SharedPlaybackState->AddCapability<FMovieSceneEvaluationState>();
+		if (bNeedsLegacyPlayer)
+		{
+			GetOrCreateLegacyPlayer();
+			check(LegacyPlayer);
+			SharedPlaybackState->AddCapability<FPlayerIndexPlaybackCapability>(LegacyPlayer->GetUniqueIndex());
+		}
+		else
+		{
+			SharedPlaybackState->AddCapability<FUMGLegacyPlayerProvider>(this);
+		}
+
+		RootInstance.Initialize();
+	}
+
+	PlaybackManager.Initialize(InAnimation);
+	PlaybackManager.SetDissectLooping(EMovieSceneLoopDissection::DissectOne);
+
+#if !NO_LOGGING
+	if (UE::UMG::GVarAnimationDetailedLogging)
+	{
+		UE_LOG(LogUMG, Log, TEXT("Animation: Initializing '%s' [%d], instance [%d]%s"), 
+				*GetNameSafe(Animation), StateIndex, RootInstanceHandle.InstanceID,
+				(LegacyPlayer ? TEXT(", with legacy player") : TEXT("")));
+	}
+#endif
+}
+
+bool FWidgetAnimationState::NeedsLegacyPlayer() const
+{
+	UUserWidget* UserWidget = WeakUserWidget.Get();
+	if (!UserWidget)
+	{
+		return false;
+	}
+
+	const ERequiresLegacyPlayer Requirement = UserWidget->GetLegacyPlayerRequirement();
+	if (Requirement == ERequiresLegacyPlayer::AutoDetect)
+	{
+		// Require a legacy player if the widget has a native C++ sub-class. This is because that
+		// sub-class could override one of the virtual methods that takes a legacy player (there is
+		// no way for us to detect that at compile time, sadly).
+		// If we only have Blueprint sub-classes, they can't use the legacy player except via
+		// lazy-creation on the FWidgetAnimationHandle, so we don't need to create one right away.
+		UClass* CurrentClass = UserWidget->GetClass();
+		while (CurrentClass && CurrentClass != UUserWidget::StaticClass())
+		{
+			if (EnumHasAnyFlags(CurrentClass->GetClassFlags(), CLASS_Native))
+			{
+				return true;
+			}
+
+			CurrentClass = CurrentClass->GetSuperClass();
+		}
+
+		return false;
+	}
+	else
+	{
+		return Requirement == ERequiresLegacyPlayer::Yes;
+	}
+}
+
+void FWidgetAnimationState::Tick(float InDeltaSeconds)
+{
+	using namespace UE::MovieScene;
+
+	UUserWidget* UserWidget = WeakUserWidget.Get();
+	TSharedPtr<FSharedPlaybackState> SharedPlaybackState = WeakPlaybackState.Pin();
+	if (!UserWidget || !SharedPlaybackState)
+	{
+#if !NO_LOGGING
+		if (UE::UMG::GVarAnimationDetailedLogging)
+		{
+			UE_LOG(LogUMG, Log, TEXT("Animation: Ticking '%s' [%d] aborted, invalid state."),
+					*GetNameSafe(Animation), StateIndex);
+		}
+#endif
+		return;
+	}
+
+	const EMovieScenePlayerStatus::Type PreviousPlaybackStatus = PlaybackManager.GetPlaybackStatus();
+	if (PreviousPlaybackStatus == EMovieScenePlayerStatus::Playing)
+	{
+		// Update root transform in case it has changed.
+		if (const FMovieSceneSequenceHierarchy* Hierarchy = SharedPlaybackState->GetHierarchy())
+		{
+			PlaybackManager.SetPlaybackTimeTransform(Hierarchy->GetRootTransform());
+		}
+
+		// Get the evaluation contexts for this tick.
+		FMovieScenePlaybackManager::FContexts TickContexts;
+		PlaybackManager.Update(InDeltaSeconds, TickContexts);
+
+		const EMovieScenePlayerStatus::Type NextPlaybackStatus = PlaybackManager.GetPlaybackStatus();
+		const bool bNeedsFinalUpdate = (NextPlaybackStatus == EMovieScenePlayerStatus::Stopped);
+
+		// Queue up the evalutions as needed. If we finished playback, setup our OnStopped callback
+		// to tear things down.
+		TSharedPtr<FMovieSceneEntitySystemRunner> EntitySystemRunner = SharedPlaybackState->GetRunner();
+		for (int32 Index = 0; Index < TickContexts.Num(); ++Index)
+		{
+			FMovieSceneContext& TickContext(TickContexts[Index]);
+
+			FSimpleDelegate OnFlushDelegate;
+			ERunnerUpdateFlags UpdateFlags = ERunnerUpdateFlags::None;
+			if (bNeedsFinalUpdate && Index == TickContexts.Num() - 1)
+			{
+				// WARNING: widget animation states are held inside a TSparseArray by UUserWidget. As this array
+				// gets re-allocated, the states also get re-allocated. So we don't want to capture the "this"
+				// pointer in a callback that could be invoked when that pointer has moved. Instead, we use a
+				// static method as the callback, with our handle as an argument. This also handles the case of
+				// a callback invoked on a state that has been removed.
+				FWidgetAnimationHandle ThisHandle = GetAnimationHandle();
+				OnFlushDelegate = FSimpleDelegate::CreateStatic(&FWidgetAnimationState::OnStopped, ThisHandle);
+				UpdateFlags = ERunnerUpdateFlags::Flush;
+				bIsStopping = true;
+			}
+			EntitySystemRunner->QueueUpdate(
+					TickContext, SharedPlaybackState->GetRootInstanceHandle(), 
+					MoveTemp(OnFlushDelegate),
+					UpdateFlags);
+
+#if !NO_LOGGING
+			if (UE::UMG::GVarAnimationDetailedLogging)
+			{
+				UE_LOG(LogUMG, Log, TEXT("Animation: Ticking '%s' [%d], Time=%s, Status=%s, Direction=%s (update %d)"),
+						*GetNameSafe(Animation), StateIndex, 
+						*LexToString(TickContext.GetTime()),
+						*UEnum::GetValueAsString(TickContext.GetStatus()),
+						((TickContext.GetDirection() == EPlayDirection::Forwards) ? TEXT("Forwards") : TEXT("Backwards")),
+						Index);
+			}
+#endif
+		}
+
+#if !NO_LOGGING
+		if (UE::UMG::GVarAnimationDetailedLogging)
+		{
+			UE_LOG(LogUMG, Log, TEXT("Animation: Ticking done '%s' [%d], NumLoopsCompleted=%d, PlaybackSpeed=%f"),
+					*GetNameSafe(Animation), StateIndex, 
+					PlaybackManager.GetNumLoopsCompleted(),
+					PlaybackManager.GetPlayRate());
+		}
+#endif
+	}
+	else
+	{
+#if !NO_LOGGING
+		if (UE::UMG::GVarAnimationDetailedLogging)
+		{
+			UE_LOG(LogUMG, Log, TEXT("Animation: Ticking '%s' [%d] skipped, state not playing."),
+					*GetNameSafe(Animation), StateIndex);
+		}
+#endif
+	}
+}
+
+void FWidgetAnimationState::OnBegunPlay(FWidgetAnimationHandle InHandle)
+{
+	if (FWidgetAnimationState* State = InHandle.GetAnimationState())
+	{
+		State->OnBegunPlay();
+	}
+}
+
+void FWidgetAnimationState::OnStopped(FWidgetAnimationHandle InHandle)
+{
+	if (FWidgetAnimationState* State = InHandle.GetAnimationState())
+	{
+		State->OnStopped();
+	}
+}
+
+void FWidgetAnimationState::OnBegunPlay()
+{
+	using namespace UE::MovieScene;
+
+	bIsBeginningPlay = false;
+
+#if !NO_LOGGING
+	if (UE::UMG::GVarAnimationDetailedLogging)
+	{
+		UE_LOG(LogUMG, Log, TEXT("Animation: OnBegunPlay '%s' [%d]"),
+				*GetNameSafe(Animation), StateIndex);
+	}
+#endif
+}
+
+void FWidgetAnimationState::OnStopped()
+{
+	using namespace UE::MovieScene;
+
+	TSharedPtr<FSharedPlaybackState> SharedPlaybackState = WeakPlaybackState.Pin();
+	if (SharedPlaybackState)
+	{
+		TSharedPtr<FMovieSceneEntitySystemRunner> EntitySystemRunner = SharedPlaybackState->GetRunner();
+		EntitySystemRunner->QueueFinalUpdate(SharedPlaybackState->GetRootInstanceHandle());
+
+		// Even if our request to Finish the instance was queued, we can wait until the next flush for those effects to be seen
+		// This will most likely happen immediately anyway since the runner will keep looping until its queue is empty,
+		// And we are already inside an active evaluation
+	}
+
+	if (SharedPlaybackState && bRestoreState)
+	{
+		SharedPlaybackState->GetPreAnimatedState().RestorePreAnimatedState();
+	}
+
+	bIsStopping = false;
+
+#if !NO_LOGGING
+	if (UE::UMG::GVarAnimationDetailedLogging)
+	{
+		PRAGMA_DISABLE_DEPRECATION_WARNINGS
+		UE_LOG(LogUMG, Log, TEXT("Animation: OnStopped '%s' [%d]%s%s"),
+				*GetNameSafe(Animation), StateIndex,
+				OnWidgetAnimationFinishedEvent.IsBound() ? TEXT(", OnWidgetAnimationFinishedEvent bound") : TEXT(""),
+				(LegacyPlayer && LegacyPlayer->OnSequenceFinishedPlaying().IsBound()) ? TEXT(", OnSequenceFinishedPlaying bound") : TEXT(""));
+		PRAGMA_ENABLE_DEPRECATION_WARNINGS
+	}
+#endif
+
+	UUserWidget* UserWidget = WeakUserWidget.Get();
+	if (UserWidget)
+	{
+		if (UE::UMG::GVarAnimationMarkers && Animation)
+		{
+			CSV_EVENT_GLOBAL(TEXT("Stop Animation [%s::%s]"), *GetNameSafe(UserWidget), *GetNameSafe(Animation));
+		}
+
+		UserWidget->BroadcastAnimationFinishedPlaying(*this);
+	}
+
+	OnWidgetAnimationFinishedEvent.Broadcast(*this);
+
+	if (LegacyPlayer)
+	{
+		LegacyPlayer->BroadcastSequenceFinishedPlaying();
+	}
+}
+
+UUMGSequencePlayer* FWidgetAnimationState::GetOrCreateLegacyPlayer()
+{
+	if (!LegacyPlayer)
+	{
+		if (UUserWidget* UserWidget = WeakUserWidget.Get())
+		{
+			LegacyPlayer = NewObject<UUMGSequencePlayer>(UserWidget);
+			LegacyPlayer->InitSequencePlayer(*this);
+		}
+	}
+	return LegacyPlayer;
+}
+
+UUMGSequencePlayer* FWidgetAnimationState::GetLegacyPlayer() const
+{
+	return LegacyPlayer;
+}
+
+FWidgetAnimationHandle FWidgetAnimationState::GetAnimationHandle() const
+{
+	return FWidgetAnimationHandle(WeakUserWidget.Get(), StateIndex, SerialNumber);
+}
+
+bool FWidgetAnimationState::IsPlayingForward() const
+{
+	return PlaybackManager.IsPlayingForward();
+}
+
+EMovieScenePlayerStatus::Type FWidgetAnimationState::GetPlaybackStatus() const
+{
+	return PlaybackManager.GetPlaybackStatus();
+}
+
+void FWidgetAnimationState::SetPlaybackStatus(EMovieScenePlayerStatus::Type InPlaybackStatus)
+{
+	PlaybackManager.SetPlaybackStatus(InPlaybackStatus);
+}
+
+FQualifiedFrameTime FWidgetAnimationState::GetCurrentTime() const
+{
+	const FFrameTime CurrentPosition = PlaybackManager.GetCurrentTime();
+	return FQualifiedFrameTime(CurrentPosition, PlaybackManager.GetDisplayRate());
+}
+
+void FWidgetAnimationState::SetCurrentTime(float InTime)
+{
+	const FFrameTime JumpFrameTime = PlaybackManager.GetDisplayRate().AsFrameTime(InTime);
+	PlaybackManager.SetCurrentTime(JumpFrameTime);
+}
+
+void FWidgetAnimationState::SetNumLoopsToPlay(int32 InNumLoopsToPlay)
+{
+	// For UMG animations, we treat a whole ping-pong as a loop.
+	if (PlayMode == EUMGSequencePlayMode::PingPong)
+	{
+		PlaybackManager.SetNumLoopsToPlay(2 * InNumLoopsToPlay);
+	}
+	else
+	{
+		PlaybackManager.SetNumLoopsToPlay(InNumLoopsToPlay);
+	}
+}
+
+void FWidgetAnimationState::SetPlaybackSpeed(float InPlaybackSpeed)
+{
+	PlaybackManager.SetPlayRate(InPlaybackSpeed);
+}
+
+void FWidgetAnimationState::Play(const FWidgetAnimationStatePlayParams& PlayParams)
+{
+	using namespace UE::MovieScene;
+
+	TSharedPtr<FSharedPlaybackState> SharedPlaybackState = WeakPlaybackState.Pin();
+	if (!ensure(SharedPlaybackState))
+	{
+		return;
+	}
+
+	UUserWidget* UserWidget = WeakUserWidget.Get();
+
+	if (UE::UMG::GVarAnimationMarkers && Animation && UserWidget)
+	{
+		CSV_EVENT_GLOBAL(TEXT("Play Animation [%s::%s]"), *GetNameSafe(UserWidget), *GetNameSafe(Animation));
+	}
+
+	bRestoreState = PlayParams.bRestoreState;
+	if (bRestoreState)
+	{
+		SharedPlaybackState->GetPreAnimatedState().EnableGlobalPreAnimatedStateCapture();
+	}
+
+	const FFrameRate DisplayRate = PlaybackManager.GetDisplayRate();
+
+	PlaybackManager.SetStartOffset(FFrameTime(0));
+	if (PlayParams.StartOffset.IsSet())
+	{
+		const FFrameTime StartOffsetTime = DisplayRate.AsFrameTime(PlayParams.StartOffset.GetValue());
+		PlaybackManager.SetStartOffset(StartOffsetTime);
+	}
+
+	PlaybackManager.SetEndOffset(FFrameTime(0));
+	if (PlayParams.EndOffset.IsSet())
+	{
+		const FFrameTime EndOffsetTime = DisplayRate.AsFrameTime(PlayParams.EndOffset.GetValue());
+		PlaybackManager.SetEndOffsetAsTime(EndOffsetTime);
+	}
+
+	PlayMode = PlayParams.PlayMode;
+	if (PlayMode == EUMGSequencePlayMode::Reverse)
+	{
+		// When playing in reverse, subtract the start time from the end.
+		const FFrameTime StartFrameTime = DisplayRate.AsFrameTime(PlayParams.StartAtTime);
+		const FFrameTime ReverseStartFrameTime = PlaybackManager.GetEffectiveEndTime() - StartFrameTime;
+		PlaybackManager.SetCurrentTime(ReverseStartFrameTime);
+	}
+	else
+	{
+		const FFrameTime StartFrameTime = DisplayRate.AsFrameTime(PlayParams.StartAtTime);
+		PlaybackManager.SetCurrentTime(StartFrameTime);
+	}
+
+	PlaybackManager.SetPlayDirection(EPlayDirection::Forwards);
+	PlaybackManager.SetPingPongPlayback(false);
+	switch (PlayMode)
+	{
+		case EUMGSequencePlayMode::Forward:
+			break;
+		case EUMGSequencePlayMode::PingPong:
+			PlaybackManager.SetPingPongPlayback(true);
+			break;
+		case EUMGSequencePlayMode::Reverse:
+			PlaybackManager.SetPlayDirection(EPlayDirection::Backwards);
+			break;
+	}
+
+	// For UMG animations, a whole ping-pong is a single loop.
+	PlaybackManager.SetNumLoopsToPlay(
+			(PlayMode == EUMGSequencePlayMode::PingPong) ?
+				2 * PlayParams.NumLoopsToPlay :
+				PlayParams.NumLoopsToPlay);
+	PlaybackManager.ResetNumLoopsCompleted();
+
+	PlaybackManager.SetPlayRate(PlayParams.PlaybackSpeed);
+	PlaybackManager.SetPlaybackStatus(EMovieScenePlayerStatus::Playing);
+
+	// Setup time warping.
+	if (const FMovieSceneSequenceHierarchy* Hierarchy = SharedPlaybackState->GetHierarchy())
+	{
+		PlaybackManager.SetPlaybackTimeTransform(Hierarchy->GetRootTransform());
+		PlaybackManager.SetTransformPlaybackTime(true);
+	}
+	else
+	{
+		PlaybackManager.SetTransformPlaybackTime(false);
+	}
+
+	// TODO: we shouldn't have to queue an update right away but we preserve the old behavior for now.
+	TSharedPtr<FMovieSceneEntitySystemRunner> Runner = SharedPlaybackState->GetRunner();
+	if (Runner)
+	{
+		// See WARNING in Tick() method.
+		FWidgetAnimationHandle ThisHandle = GetAnimationHandle();
+		const FMovieSceneContext FirstContext = PlaybackManager.UpdateAtCurrentTime();
+		Runner->QueueUpdate(
+				FirstContext, 
+				SharedPlaybackState->GetRootInstanceHandle(),
+				FSimpleDelegate::CreateStatic(&FWidgetAnimationState::OnBegunPlay, ThisHandle),
+				ERunnerUpdateFlags::Flush);
+	}
+
+#if !NO_LOGGING
+	if (UE::UMG::GVarAnimationDetailedLogging)
+	{
+		UE_LOG(LogUMG, Log, TEXT("Animation: Play '%s' [%d], StartTime=%s, NumLoopsToPlay=%d, PlaybackSpeed=%f, PlayMode=%s%s"),
+				*GetNameSafe(Animation), StateIndex,
+				*LexToString(PlaybackManager.GetCurrentTime()),
+				PlaybackManager.GetNumLoopsToPlay(),
+				PlaybackManager.GetPlayRate(),
+				PlaybackManager.GetPlayDirection() == EPlayDirection::Forwards ? TEXT("Forwards") : TEXT("Backwards"),
+				PlaybackManager.IsPingPongPlayback() ? TEXT(", PingPong") : TEXT(""));
+	}
+#endif
+}
+
+void FWidgetAnimationState::Stop()
+{
+	using namespace UE::MovieScene;
+
+	if (PlaybackManager.GetPlaybackStatus() == EMovieScenePlayerStatus::Stopped)
+	{
+		return;
+	}
+
+	PlaybackManager.SetPlaybackStatus(EMovieScenePlayerStatus::Stopped);
+
+	// TODO: we should use PlaybackManager.GetEffectiveStartTime(), but the old behavior always used frame zero.
+	PlaybackManager.SetCurrentTime(0);
+
+	if (TSharedPtr<FSharedPlaybackState> SharedPlaybackState = WeakPlaybackState.Pin())
+	{
+		TSharedPtr<FMovieSceneEntitySystemRunner> Runner = SharedPlaybackState->GetRunner();
+
+		// See WARNING in Tick() method.
+		FWidgetAnimationHandle ThisHandle = GetAnimationHandle();
+		const FMovieSceneContext ReturnToStartContext = PlaybackManager.UpdateAtCurrentTime();
+		Runner->QueueUpdate(
+				ReturnToStartContext,
+				SharedPlaybackState->GetRootInstanceHandle(),
+				FSimpleDelegate::CreateStatic(&FWidgetAnimationState::OnStopped, ThisHandle),
+				ERunnerUpdateFlags::Flush);
+
+		bIsStopping = true;
+	}
+	else
+	{
+		OnStopped();
+	}
+
+#if !NO_LOGGING
+	if (UE::UMG::GVarAnimationDetailedLogging)
+	{
+		UE_LOG(LogUMG, Log, TEXT("Animation: Stop '%s' [%d]"),
+				*GetNameSafe(Animation), StateIndex);
+	}
+#endif
+}
+
+void FWidgetAnimationState::Pause()
+{
+	using namespace UE::MovieScene;
+
+	// Should be Paused but old behavior was to set to Stopped.
+	PlaybackManager.SetPlaybackStatus(EMovieScenePlayerStatus::Stopped);
+
+	// Evaluate the sequence at its current time, with a status of 'Stopped' to ensure that animated state 
+	// pauses correctly. (ie. audio sounds should stop/pause)
+	if (TSharedPtr<FSharedPlaybackState> SharedPlaybackState = WeakPlaybackState.Pin())
+	{
+		TSharedPtr<FMovieSceneEntitySystemRunner> Runner = SharedPlaybackState->GetRunner();
+
+		const FMovieSceneContext PauseContext = PlaybackManager.UpdateAtCurrentTime();
+		Runner->QueueUpdate(
+				PauseContext,
+				SharedPlaybackState->GetRootInstanceHandle(),
+				ERunnerUpdateFlags::Flush);
+	}
+
+#if !NO_LOGGING
+	if (UE::UMG::GVarAnimationDetailedLogging)
+	{
+		UE_LOG(LogUMG, Log, TEXT("Animation: Pause '%s' [%d], PauseTime=%s"),
+				*GetNameSafe(Animation), StateIndex,
+				*LexToString(PlaybackManager.GetCurrentTime()));
+	}
+#endif
+}
+
+void FWidgetAnimationState::Reverse()
+{
+	PlaybackManager.ReversePlayDirection();
+}
+
+void FWidgetAnimationState::RemoveEvaluationData()
+{
+	using namespace UE::MovieScene;
+
+	if (TSharedPtr<const FSharedPlaybackState> SharedPlaybackState = WeakPlaybackState.Pin())
+	{
+		UMovieSceneEntitySystemLinker* Linker = SharedPlaybackState->GetLinker();
+		FRootInstanceHandle RootInstanceHandle = SharedPlaybackState->GetRootInstanceHandle();
+		FSequenceInstance& RootInstance = Linker->GetInstanceRegistry()->MutateInstance(RootInstanceHandle);
+
+		TSharedRef<FMovieSceneEntitySystemRunner> Runner = Linker->GetRunner();
+		if (Runner->IsCurrentlyEvaluating())
+		{
+			Runner->FlushOutstanding();
+		}
+
+		RootInstance.Ledger.UnlinkEverything(Linker);
+		RootInstance.InvalidateCachedData();
+	}
+
+	//PRAGMA_DISABLE_DEPRECATION_WARNINGS
+	//RootTemplateInstance.ResetDirectorInstances();
+	//PRAGMA_ENABLE_DEPRECATION_WARNINGS
+
+#if !NO_LOGGING
+	if (UE::UMG::GVarAnimationDetailedLogging)
+	{
+		UE_LOG(LogUMG, Log, TEXT("Animation: RemoveEvaluationData '%s' [%d]"),
+				*GetNameSafe(Animation), StateIndex);
+	}
+#endif
+}
+
+void FWidgetAnimationState::TearDown()
+{
+	using namespace UE::MovieScene;
+
+	FRootInstanceHandle RootInstanceHandle;
+	UMovieSceneEntitySystemLinker* EntitySystemLinker = nullptr;
+
+	if (TSharedPtr<const FSharedPlaybackState> SharedPlaybackState = WeakPlaybackState.Pin())
+	{
+		RootInstanceHandle = SharedPlaybackState->GetRootInstanceHandle();
+		EntitySystemLinker = SharedPlaybackState->GetLinker();
+	}
+
+	if (LegacyPlayer)
+	{
+		PRAGMA_DISABLE_DEPRECATION_WARNINGS
+		FMovieSceneRootEvaluationTemplateInstance& EvaluationTemplate = LegacyPlayer->GetEvaluationTemplate();
+		EvaluationTemplate.TearDown();
+		PRAGMA_ENABLE_DEPRECATION_WARNINGS
+	}
+	LegacyPlayer = nullptr;
+
+	WeakPlaybackState.Reset();
+
+	if (EntitySystemLinker && RootInstanceHandle.IsValid())
+	{
+		EntitySystemLinker->DestroyInstanceImmediately(RootInstanceHandle);
+	}
+
+#if !NO_LOGGING
+	if (UE::UMG::GVarAnimationDetailedLogging)
+	{
+		UE_LOG(LogUMG, Log, TEXT("Animation: TearDown '%s' [%d]"),
+				*GetNameSafe(Animation), StateIndex);
+	}
+#endif
+}
+
+bool FWidgetAnimationState::IsValid() const
+{
+	return WeakPlaybackState.IsValid();
+}
+
+void FWidgetAnimationState::AddReferencedObjects(FReferenceCollector& Collector)
+{
+	Collector.AddReferencedObject(Animation);
+	Collector.AddReferencedObject(LegacyPlayer);
+}
+
diff --git a/Engine/Source/Runtime/UMG/Private/UserWidget.cpp b/Engine/Source/Runtime/UMG/Private/UserWidget.cpp
--- a/Engine/Source/Runtime/UMG/Private/UserWidget.cpp
+++ b/Engine/Source/Runtime/UMG/Private/UserWidget.cpp
@@ -42,13 +42,6 @@
 
 #define LOCTEXT_NAMESPACE "UMG"
 
-TAutoConsoleVariable<bool> CVarUserWidgetUseParallelAnimation(
-	TEXT("Widget.UseParallelAnimation"),
-	true,
-	TEXT("Use multi-threaded evaluation for widget animations."),
-	ECVF_Default
-);
-
 namespace UE::UMG::Private
 {
 	static thread_local uint32 InitializingFromWidgetTree = 0;
@@ -80,6 +73,11 @@ FPaintContext::FPaintContext()
 {
 }
 
+UUMGSequencePlayer* UUserWidgetFunctionLibrary::Conv_UMGSequencePlayer(const FWidgetAnimationHandle& WidgetAnimationHandle)
+{
+	return WidgetAnimationHandle.GetSequencePlayer();
+}
+
 /////////////////////////////////////////////////////
 // UUserWidget
 UUserWidget::UUserWidget(const FObjectInitializer& ObjectInitializer)
@@ -485,18 +483,35 @@ UWorld* UUserWidget::GetWorld() const
 	return nullptr;
 }
 
-UUMGSequencePlayer* UUserWidget::GetSequencePlayer(const UWidgetAnimation* InAnimation) const
+const FWidgetAnimationState* UUserWidget::GetAnimationState(const UWidgetAnimation* InAnimation) const
 {
-	TObjectPtr<UUMGSequencePlayer> const* FoundPlayer = ActiveSequencePlayers.FindByPredicate(
-		[&](const UUMGSequencePlayer* Player)
+	const FWidgetAnimationState* FoundState = nullptr;
+	for (auto It = ActiveAnimations.CreateConstIterator(); It; ++It)
 	{
-		return Player->GetAnimation() == InAnimation && !Player->IsStopping();
-	});
+		if (It->GetAnimation() == InAnimation)
+		{
+			FoundState = &(*It);
+			break;
+		}
+	};
+	return FoundState;
+}
 
-	return FoundPlayer ? *FoundPlayer : nullptr;
+FWidgetAnimationState* UUserWidget::GetAnimationState(const UWidgetAnimation* InAnimation)
+{
+	FWidgetAnimationState* FoundState = nullptr;
+	for (auto It = ActiveAnimations.CreateIterator(); It; ++It)
+	{
+		if (It->GetAnimation() == InAnimation)
+		{
+			FoundState = &(*It);
+			break;
+		}
+	};
+	return FoundState;
 }
 
-UUMGSequencePlayer* UUserWidget::GetOrAddSequencePlayer(UWidgetAnimation* InAnimation)
+FWidgetAnimationState* UUserWidget::GetOrAddAnimationState(UWidgetAnimation* InAnimation)
 {
 	if (InAnimation && !bStoppingAllAnimations)
 	{
@@ -509,31 +524,27 @@ UUMGSequencePlayer* UUserWidget::GetOrAddSequencePlayer(UWidgetAnimation* InAnim
 		// If this widget is currently offscreen it would very well not be being tracked by the tick manager
 		AnimationTickManager->AddWidget(this);
 
-		// @todo UMG sequencer - Restart animations which have had Play called on them?
-		UUMGSequencePlayer* FoundPlayer = nullptr;
-		for (UUMGSequencePlayer* Player : ActiveSequencePlayers)
+		FWidgetAnimationState* FoundState = GetAnimationState(InAnimation);
+		if (!FoundState)
 		{
-			// We need to make sure we haven't stopped the animation, otherwise it'll get canceled on the next frame.
-			if (Player->GetAnimation() == InAnimation
-			 && !StoppedSequencePlayers.Contains(Player) && !Player->IsStopping())
-			{
-				FoundPlayer = Player;
-				break;
-			}
-		}
+			// Create a new state and initialize it.
+			int32 NewStateIndex = ActiveAnimations.Emplace();
+			FWidgetAnimationState& NewState = ActiveAnimations[NewStateIndex];
+			NewState.Initialize(InAnimation, this, NewStateIndex, NextSerialNumber++);
 
-		if (!FoundPlayer)
+			return &NewState;
+		}
+		else if (!FoundState->IsValid())
 		{
-			UUMGSequencePlayer* NewPlayer = NewObject<UUMGSequencePlayer>(this, NAME_None, RF_Transient);
-			ActiveSequencePlayers.Add(NewPlayer);
-
-			NewPlayer->InitSequencePlayer(*InAnimation, *this);
-
-			return NewPlayer;
+			// We found a state that just recently finished playing and was torn-down, but not yet cleared
+			// from ActiveAnimations. Let's reinitialize it.
+			const int32 StateIndex = FoundState->StateIndex;
+			const uint32 StateSerial = FoundState->SerialNumber;
+			FoundState->Initialize(InAnimation, this, StateIndex, StateSerial);
 		}
 		else
 		{
-			return FoundPlayer;
+			return FoundState;
 		}
 	}
 
@@ -594,61 +605,36 @@ void UUserWidget::ExecuteQueuedAnimationTransitions()
 
 void UUserWidget::ConditionalTearDownAnimations()
 {
-	for (auto It = ActiveSequencePlayers.CreateIterator(); It; ++It)
+	for (auto It = ActiveAnimations.CreateIterator(); It; ++It)
 	{
-		UUMGSequencePlayer* Player = *It;
-		if (!Player)
+		FWidgetAnimationState& State(*It);
+		if (!State.IsValid())
 		{
 			It.RemoveCurrent();
 		}
-		else if (!Player->IsStopping())
+		else if (!State.IsStopping())
 		{
-			Player->TearDown();
-			It.RemoveCurrent();
+			State.TearDown();
 		}
 	}
-
-	for (UUMGSequencePlayer* Player : StoppedSequencePlayers)
-	{
-		Player->TearDown();
-	}
-	StoppedSequencePlayers.Empty();
 }
 
 void UUserWidget::TearDownAnimations()
 {
-	for (UUMGSequencePlayer* Player : ActiveSequencePlayers)
-	{
-		if (Player)
-		{
-			Player->TearDown();
-		}
-	}
-
-	for (UUMGSequencePlayer* Player : StoppedSequencePlayers)
+	for (auto It = ActiveAnimations.CreateIterator(); It; ++It)
 	{
-		// These null checks should not be necessary since StoppedSequencePlayers should always keep them alive
-		//   but there are some rare cases where the GC may eliminate and null these references if 'this' is an
-		//   instance of a BP class. This is a bug in the reference collection semantics, but is too intricate
-		//   to fix in a timely manner so we have to be defensive here.
-		if (Player)
-		{
-			Player->TearDown();
-		}
+		FWidgetAnimationState& State(*It);
+		State.TearDown();
 	}
 
-	ActiveSequencePlayers.Empty();
-	StoppedSequencePlayers.Empty();
+	ActiveAnimations.Reset();
 }
 
 void UUserWidget::DisableAnimations()
 {
-	for (UUMGSequencePlayer* Player : ActiveSequencePlayers)
+	for (auto It = ActiveAnimations.CreateIterator(); It; ++It)
 	{
-		if (Player)
-		{
-			Player->RemoveEvaluationData();
-		}
+		It->RemoveEvaluationData();
 	}
 }
 
@@ -661,6 +647,11 @@ void UUserWidget::Invalidate(EInvalidateWidgetReason InvalidateReason)
 	}
 }
 
+bool UUserWidget::IsPlayingAnimation() const
+{
+	return ActiveAnimations.Num() > 0;
+}
+
 void UUserWidget::QueuePlayAnimation(UWidgetAnimation* InAnimation, float StartAtTime, int32 NumLoopsToPlay, EUMGSequencePlayMode::Type PlayMode, float PlaybackSpeed, bool bRestoreState)
 {
 	if (!InAnimation)
@@ -764,12 +755,11 @@ void UUserWidget::QueueStopAllAnimations()
 		QueuedWidgetAnimationTransition.TransitionMode = EQueuedWidgetAnimationMode::Stop;
 	}
 
-	TArray<UUMGSequencePlayer*, TInlineAllocator<8>> CurrentActivePlayers(ActiveSequencePlayers);
-	for (UUMGSequencePlayer* FoundPlayer : ActiveSequencePlayers)
+	for (FWidgetAnimationState& AnimationState : ActiveAnimations)
 	{
-		if (FoundPlayer->GetPlaybackStatus() == EMovieScenePlayerStatus::Playing)
+		if (AnimationState.GetPlaybackStatus() == EMovieScenePlayerStatus::Playing)
 		{
-			QueueStopAnimation(FoundPlayer->GetAnimation());
+			QueueStopAnimation(AnimationState.GetAnimation());
 		}
 	}
 
@@ -788,88 +778,95 @@ float UUserWidget::QueuePauseAnimation(const UWidgetAnimation* InAnimation)
 
 		UpdateCanTick();
 
-		if (UUMGSequencePlayer* FoundPlayer = GetSequencePlayer(InAnimation))
+		if (FWidgetAnimationState* FoundState = GetAnimationState(InAnimation))
 		{
-			return (float)FoundPlayer->GetCurrentTime().AsSeconds();
+			return (float)FoundState->GetCurrentTime().AsSeconds();
 		}
 	}
 
 	return 0;
 }
 
-UUMGSequencePlayer* UUserWidget::PlayAnimation(UWidgetAnimation* InAnimation, float StartAtTime, int32 NumberOfLoops, EUMGSequencePlayMode::Type PlayMode, float PlaybackSpeed, bool bRestoreState)
+FWidgetAnimationHandle UUserWidget::PlayAnimation(UWidgetAnimation* InAnimation, float StartAtTime, int32 NumberOfLoops, EUMGSequencePlayMode::Type PlayMode, float PlaybackSpeed, bool bRestoreState)
 {
 	SCOPED_NAMED_EVENT_TEXT("Widget::PlayAnimation", FColor::Emerald);
 
-	UUMGSequencePlayer* Player = GetOrAddSequencePlayer(InAnimation);
-	if (Player)
+	if (FWidgetAnimationState* AnimationState = GetOrAddAnimationState(InAnimation))
 	{
-		Player->Play(StartAtTime, NumberOfLoops, PlayMode, PlaybackSpeed, bRestoreState);
+		FWidgetAnimationStatePlayParams PlayParams;
+		PlayParams.StartAtTime = StartAtTime;
+		PlayParams.NumLoopsToPlay = NumberOfLoops;
+		PlayParams.PlayMode = PlayMode;
+		PlayParams.PlaybackSpeed = PlaybackSpeed;
+		PlayParams.bRestoreState = bRestoreState;
+
+		AnimationState->Play(PlayParams);
 
-		OnAnimationStartedPlaying(*Player);
+		BroadcastAnimationStartedPlaying(*AnimationState);
 
 		UpdateCanTick();
+
+		return FWidgetAnimationHandle(this, AnimationState->StateIndex, AnimationState->SerialNumber);
 	}
 
-	return Player;
+	return FWidgetAnimationHandle();
 }
 
-UUMGSequencePlayer* UUserWidget::PlayAnimationTimeRange(UWidgetAnimation* InAnimation, float StartAtTime, float EndAtTime, int32 NumberOfLoops, EUMGSequencePlayMode::Type PlayMode, float PlaybackSpeed, bool bRestoreState)
+FWidgetAnimationHandle UUserWidget::PlayAnimationTimeRange(UWidgetAnimation* InAnimation, float StartAtTime, float EndAtTime, int32 NumberOfLoops, EUMGSequencePlayMode::Type PlayMode, float PlaybackSpeed, bool bRestoreState)
 {
 	SCOPED_NAMED_EVENT_TEXT("Widget::PlayAnimationTimeRange", FColor::Emerald);
 
-	UUMGSequencePlayer* Player = GetOrAddSequencePlayer(InAnimation);
-	if (Player)
+	if (FWidgetAnimationState* AnimationState = GetOrAddAnimationState(InAnimation))
 	{
-		Player->PlayTo(StartAtTime, EndAtTime, NumberOfLoops, PlayMode, PlaybackSpeed, bRestoreState);
+		FWidgetAnimationStatePlayParams PlayParams;
+		PlayParams.StartAtTime = StartAtTime;
+		PlayParams.EndOffset = EndAtTime;
+		PlayParams.NumLoopsToPlay = NumberOfLoops;
+		PlayParams.PlayMode = PlayMode;
+		PlayParams.PlaybackSpeed = PlaybackSpeed;
+		PlayParams.bRestoreState = bRestoreState;
+
+		AnimationState->Play(PlayParams);
 
-		OnAnimationStartedPlaying(*Player);
+		BroadcastAnimationStartedPlaying(*AnimationState);
 
 		UpdateCanTick();
+
+		return FWidgetAnimationHandle(this, AnimationState->StateIndex, AnimationState->SerialNumber);
 	}
 
-	return Player;
+	return FWidgetAnimationHandle();
 }
 
-UUMGSequencePlayer* UUserWidget::PlayAnimationForward(UWidgetAnimation* InAnimation, float PlaybackSpeed, bool bRestoreState)
+FWidgetAnimationHandle UUserWidget::PlayAnimationForward(UWidgetAnimation* InAnimation, float PlaybackSpeed, bool bRestoreState)
 {
-	// Don't create the player, only search for it.
-	UUMGSequencePlayer* Player = GetSequencePlayer(InAnimation);
+	FWidgetAnimationState* AnimationState = GetAnimationState(InAnimation);
 
-	// Just return the player if it's already playing. GetSequencePlayers should only be returning players that are NOT stopped. 
-	// However, there is the possibility that Stop() has been called on a player, but its status has not been changed to Stopped. 
-	// In that case, this check will be bypassed and PlayAnimation will be called so that the animation will play.
-	if (Player && Player->GetPlaybackStatus() == EMovieScenePlayerStatus::Playing)
+	if (AnimationState && AnimationState->GetPlaybackStatus() == EMovieScenePlayerStatus::Playing)
 	{
-		if (!Player->IsPlayingForward())
+		if (!AnimationState->IsPlayingForward())
 		{
-			// Reverse the direction we're playing the animation if we're playing it in reverse currently.
-			Player->Reverse();
+			AnimationState->Reverse();
 		}
 
-		return Player;
+		return FWidgetAnimationHandle(this, AnimationState->StateIndex, AnimationState->SerialNumber);
 	}
 
 	return PlayAnimation(InAnimation, 0.0f, 1, EUMGSequencePlayMode::Forward, PlaybackSpeed, bRestoreState);
 }
 
-UUMGSequencePlayer* UUserWidget::PlayAnimationReverse(UWidgetAnimation* InAnimation, float PlaybackSpeed, bool bRestoreState)
+FWidgetAnimationHandle UUserWidget::PlayAnimationReverse(UWidgetAnimation* InAnimation, float PlaybackSpeed, bool bRestoreState)
 {
-	// Don't create the player, only search for it.
-	UUMGSequencePlayer* Player = GetSequencePlayer(InAnimation);
+	FWidgetAnimationState* AnimationState = GetAnimationState(InAnimation);
 
-	// Just return the player if it's already playing. GetSequencePlayers should only be returning players that are NOT stopped. 
-	// However, there is the possibility that Stop() has been called on a player, but its status has not been changed to Stopped. 
-	// In that case, this check will be bypassed and PlayAnimation will be called so that the animation will play.
-	if (Player && Player->GetPlaybackStatus() == EMovieScenePlayerStatus::Playing)
+	if (AnimationState && AnimationState->GetPlaybackStatus() == EMovieScenePlayerStatus::Playing)
 	{
-		if (Player->IsPlayingForward())
+		if (AnimationState->IsPlayingForward())
 		{
-			// Reverse the direction we're playing the animation if we're playing it in forward currently.
-			Player->Reverse();
+			AnimationState->Reverse();
 		}
 
-		return Player;
+		return FWidgetAnimationHandle(this, AnimationState->StateIndex, AnimationState->SerialNumber);
 	}
 
 	return PlayAnimation(InAnimation, 0.0f, 1, EUMGSequencePlayMode::Reverse, PlaybackSpeed, bRestoreState);
@@ -879,10 +876,9 @@ void UUserWidget::StopAnimation(const UWidgetAnimation* InAnimation)
 {
 	if (InAnimation)
 	{
-		// @todo UMG sequencer - Restart animations which have had Play called on them?
-		if (UUMGSequencePlayer* FoundPlayer = GetSequencePlayer(InAnimation))
+		if (FWidgetAnimationState* FoundState = GetAnimationState(InAnimation))
 		{
-			FoundPlayer->Stop();
+			FoundState->Stop();
 
 			UpdateCanTick();
 		}
@@ -892,16 +888,15 @@ void UUserWidget::StopAnimation(const UWidgetAnimation* InAnimation)
 void UUserWidget::StopAllAnimations()
 {
 	bStoppingAllAnimations = true;
-	
-	// Stopping players modifies ActiveSequencePlayers, work on a copy aprray
-	TArray<UUMGSequencePlayer*, TInlineAllocator<8>> CurrentActivePlayers(ActiveSequencePlayers);
-	for (UUMGSequencePlayer* FoundPlayer : CurrentActivePlayers)
+
+	for (FWidgetAnimationState& AnimationState : ActiveAnimations)
 	{
-		if (FoundPlayer->GetPlaybackStatus() == EMovieScenePlayerStatus::Playing)
+		if (AnimationState.GetPlaybackStatus() == EMovieScenePlayerStatus::Playing)
 		{
-			FoundPlayer->Stop();
+			AnimationState.Stop();
 		}
 	}
+
 	bStoppingAllAnimations = false;
 
 	UpdateCanTick();
@@ -912,10 +907,10 @@ float UUserWidget::PauseAnimation(const UWidgetAnimation* InAnimation)
 	if (InAnimation)
 	{
 		// @todo UMG sequencer - Restart animations which have had Play called on them?
-		if (UUMGSequencePlayer* FoundPlayer = GetSequencePlayer(InAnimation))
+		if (FWidgetAnimationState* FoundState = GetAnimationState(InAnimation))
 		{
-			FoundPlayer->Pause();
-			return (float)FoundPlayer->GetCurrentTime().AsSeconds();
+			FoundState->Pause();
+			return (float)FoundState->GetCurrentTime().AsSeconds();
 		}
 	}
 
@@ -926,9 +921,9 @@ float UUserWidget::GetAnimationCurrentTime(const UWidgetAnimation* InAnimation)
 {
 	if (InAnimation)
 	{
-		if (UUMGSequencePlayer* FoundPlayer = GetSequencePlayer(InAnimation))
+		if (const FWidgetAnimationState* FoundState = GetAnimationState(InAnimation))
 		{
-			return (float)FoundPlayer->GetCurrentTime().AsSeconds();
+			return (float)FoundState->GetCurrentTime().AsSeconds();
 		}
 	}
 
@@ -939,9 +934,9 @@ void UUserWidget::SetAnimationCurrentTime(const UWidgetAnimation* InAnimation, f
 {
 	if (InAnimation)
 	{
-		if (UUMGSequencePlayer* FoundPlayer = GetSequencePlayer(InAnimation))
+		if (FWidgetAnimationState* FoundState = GetAnimationState(InAnimation))
 		{
-			FoundPlayer->SetCurrentTime(InTime);
+			FoundState->SetCurrentTime(InTime);
 		}
 	}
 }
@@ -950,9 +945,9 @@ bool UUserWidget::IsAnimationPlaying(const UWidgetAnimation* InAnimation) const
 {
 	if (InAnimation)
 	{
-		if (UUMGSequencePlayer* FoundPlayer = GetSequencePlayer(InAnimation))
+		if (const FWidgetAnimationState* FoundState = GetAnimationState(InAnimation))
 		{
-			return FoundPlayer->GetPlaybackStatus() == EMovieScenePlayerStatus::Playing;
+			return FoundState->GetPlaybackStatus() == EMovieScenePlayerStatus::Playing;
 		}
 	}
 
@@ -961,72 +956,81 @@ bool UUserWidget::IsAnimationPlaying(const UWidgetAnimation* InAnimation) const
 
 bool UUserWidget::IsAnyAnimationPlaying() const
 {
-	return ActiveSequencePlayers.Num() > 0;
+	return ActiveAnimations.Num() > 0;
 }
 
 void UUserWidget::SetNumLoopsToPlay(const UWidgetAnimation* InAnimation, int32 InNumLoopsToPlay)
 {
-	if (UUMGSequencePlayer* FoundPlayer = GetSequencePlayer(InAnimation))
+	if (FWidgetAnimationState* FoundState = GetAnimationState(InAnimation))
 	{
-		FoundPlayer->SetNumLoopsToPlay(InNumLoopsToPlay);
+		FoundState->SetNumLoopsToPlay(InNumLoopsToPlay);
 	}
 }
 
 void UUserWidget::SetPlaybackSpeed(const UWidgetAnimation* InAnimation, float PlaybackSpeed)
 {
-	if (UUMGSequencePlayer* FoundPlayer = GetSequencePlayer(InAnimation))
+	if (FWidgetAnimationState* FoundState = GetAnimationState(InAnimation))
 	{
-		FoundPlayer->SetPlaybackSpeed(PlaybackSpeed);
+		FoundState->SetPlaybackSpeed(PlaybackSpeed);
 	}
 }
 
 void UUserWidget::ReverseAnimation(const UWidgetAnimation* InAnimation)
 {
-	if (UUMGSequencePlayer* FoundPlayer = GetSequencePlayer(InAnimation))
+	if (FWidgetAnimationState* FoundState = GetAnimationState(InAnimation))
 	{
-		FoundPlayer->Reverse();
+		FoundState->Reverse();
 	}
 }
 
-void UUserWidget::OnAnimationStartedPlaying(UUMGSequencePlayer& Player)
+void UUserWidget::BroadcastAnimationStartedPlaying(FWidgetAnimationState& State)
 {
-	OnAnimationStarted(Player.GetAnimation());
+	OnAnimationStartedPlayingEvent.Broadcast(State);
 
-	BroadcastAnimationStateChange(Player, EWidgetAnimationEvent::Started);
+	OnAnimationStarted(State.GetAnimation());
+
+	BroadcastAnimationStateChange(State, EWidgetAnimationEvent::Started);
+
+	if (UUMGSequencePlayer* LegacyPlayer = State.GetLegacyPlayer())
+	{
+		PRAGMA_DISABLE_DEPRECATION_WARNINGS
+		OnAnimationStartedPlaying(*LegacyPlayer);
+		PRAGMA_ENABLE_DEPRECATION_WARNINGS
+	}
 }
 
 bool UUserWidget::IsAnimationPlayingForward(const UWidgetAnimation* InAnimation)
 {
 	if (InAnimation)
 	{
-		TObjectPtr<UUMGSequencePlayer>* FoundPlayer = ActiveSequencePlayers.FindByPredicate([&](const UUMGSequencePlayer* Player) { return Player->GetAnimation() == InAnimation; });
-
-		if (FoundPlayer)
+		if (FWidgetAnimationState* FoundState = GetAnimationState(InAnimation))
 		{
-			return (*FoundPlayer)->IsPlayingForward();
+			return FoundState->IsPlayingForward();
 		}
 	}
 
 	return true;
 }
 
-void UUserWidget::OnAnimationFinishedPlaying(UUMGSequencePlayer& Player)
+void UUserWidget::BroadcastAnimationFinishedPlaying(FWidgetAnimationState& State)
 {
-	// This event is called directly by the sequence player when the animation finishes.
+	OnAnimationFinishedPlayingEvent.Broadcast(State);
 
-	OnAnimationFinished(Player.GetAnimation());
+	OnAnimationFinished(State.GetAnimation());
 
-	BroadcastAnimationStateChange(Player, EWidgetAnimationEvent::Finished);
+	BroadcastAnimationStateChange(State, EWidgetAnimationEvent::Finished);
 
-	if ( Player.GetPlaybackStatus() == EMovieScenePlayerStatus::Stopped )
+	if (UUMGSequencePlayer* LegacyPlayer = State.GetLegacyPlayer())
 	{
-		//ensureAlways(!StoppedSequencePlayers.Contains(&Player));
-		StoppedSequencePlayers.AddUnique(&Player);
+		PRAGMA_DISABLE_DEPRECATION_WARNINGS
+		OnAnimationFinishedPlaying(*LegacyPlayer);
+		PRAGMA_ENABLE_DEPRECATION_WARNINGS
+	}
 
-		if (AnimationTickManager)
-		{
-			AnimationTickManager->AddLatentAction(FMovieSceneSequenceLatentActionDelegate::CreateUObject(this, &UUserWidget::ClearStoppedSequencePlayers));
-		}
+	if (State.GetPlaybackStatus() == EMovieScenePlayerStatus::Stopped && AnimationTickManager)
+	{
+		AnimationTickManager->AddLatentAction(
+				FMovieSceneSequenceLatentActionDelegate::CreateUObject(this, &UUserWidget::ClearStoppedAnimationStates));
 	}
 
 	UpdateCanTick();
@@ -1034,8 +1038,20 @@ void UUserWidget::OnAnimationFinishedPlaying(UUMGSequencePlayer& Player)
 
 void UUserWidget::BroadcastAnimationStateChange(const UUMGSequencePlayer& Player, EWidgetAnimationEvent AnimationEvent)
 {
-	const UWidgetAnimation* Animation = Player.GetAnimation();
+	PRAGMA_DISABLE_DEPRECATION_WARNINGS
+
+	BroadcastAnimationStateChange(Player.GetAnimation(), Player.GetUserTag(), AnimationEvent);
+
+	PRAGMA_ENABLE_DEPRECATION_WARNINGS
+}
+
+void UUserWidget::BroadcastAnimationStateChange(const FWidgetAnimationState& State, EWidgetAnimationEvent AnimationEvent)
+{
+	BroadcastAnimationStateChange(State.GetAnimation(), State.GetUserTag(), AnimationEvent);
+}
 
+void UUserWidget::BroadcastAnimationStateChange(const UWidgetAnimation* Animation, FName UserTag, EWidgetAnimationEvent AnimationEvent)
+{
 	// Make a temporary copy of the animation callbacks so that everyone gets a callback
 	// even if they're removed as a result of other calls, we don't want order to matter here.
 	TArray<FAnimationEventBinding> TempAnimationCallbacks = AnimationCallbacks;
@@ -1044,7 +1060,7 @@ void UUserWidget::BroadcastAnimationStateChange(const UUMGSequencePlayer& Player
 	{
 		if (Binding.Animation == Animation && Binding.AnimationEvent == AnimationEvent)
 		{
-			if (Binding.UserTag == NAME_None || Binding.UserTag == Player.GetUserTag())
+			if (Binding.UserTag == NAME_None || Binding.UserTag == UserTag)
 			{
 				Binding.Delegate.ExecuteIfBound();
 			}
@@ -1880,13 +1896,6 @@ void UUserWidget::NativeTick(const FGeometry& MyGeometry, float InDeltaTime)
 				AnimationTickManager->OnWidgetTicked(this);
 			}
 
-			if (!CVarUserWidgetUseParallelAnimation.GetValueOnGameThread())
-			{
-				TickActionsAndAnimation(InDeltaTime);
-				PostTickActionsAndAnimation(InDeltaTime);
-			}
-			// else: the TickManager object will tick all animations at once.
-
 			UWorld* World = GetWorld();
 			if (World)
 			{
@@ -1913,17 +1922,12 @@ void UUserWidget::TickActionsAndAnimation(float InDeltaTime)
 	// Update active movie scenes, none will be removed here, but new
 	// ones can be added during the tick, if a player ends and triggers
 	// starting another animation
-	for (int32 Index = 0; Index < ActiveSequencePlayers.Num(); Index++)
+	for (FWidgetAnimationState& State : ActiveAnimations)
 	{
-		UUMGSequencePlayer* Player = ActiveSequencePlayers[Index];
-		Player->Tick(InDeltaTime);
+		State.Tick(InDeltaTime);
 	}
 }
 
-void UUserWidget::PostTickActionsAndAnimation(float InDeltaTime)
-{
-}
-
 void UUserWidget::FlushAnimations()
 {
 	UUMGSequenceTickManager::Get(this)->ForceFlush();
@@ -2112,7 +2116,7 @@ void UUserWidget::UpdateCanTick()
 			bCanTick |= !WidgetBPClass || WidgetBPClass->ClassRequiresNativeTick();
 			bCanTick |= bHasScriptImplementedTick;
 			bCanTick |= World->GetLatentActionManager().GetNumActionsForObject(this) != 0;
-			bCanTick |= ActiveSequencePlayers.Num() > 0;
+			bCanTick |= ActiveAnimations.Num() > 0;
 			bCanTick |= QueuedWidgetAnimationTransitions.Num() > 0;
 
 			if (!bCanTick && bAreExtensionsConstructed)
@@ -2558,17 +2562,33 @@ UUserWidget* UUserWidget::CreateInstanceInternal(UObject* Outer, TSubclassOf<UUs
 	return NewWidget;
 }
 
-
-void UUserWidget::ClearStoppedSequencePlayers()
+void UUserWidget::AddReferencedObjects(UObject* InThis, FReferenceCollector& Collector)
 {
-	// after all players have ticked, remove and tear down stopped players
-	for (UUMGSequencePlayer* StoppedPlayer : StoppedSequencePlayers)
+	Super::AddReferencedObjects(InThis, Collector);
+
+	UUserWidget* TypedThis = CastChecked<UUserWidget>(InThis);
+
+	for (auto It = TypedThis->ActiveAnimations.CreateIterator(); It; ++It)
 	{
-		ActiveSequencePlayers.RemoveSwap(StoppedPlayer);
-		StoppedPlayer->TearDown();
+		It->AddReferencedObjects(Collector);
 	}
+}
 
-	StoppedSequencePlayers.Empty();
+void UUserWidget::ClearStoppedAnimationStates()
+{
+	for (auto It = ActiveAnimations.CreateIterator(); It; ++It)
+	{
+		FWidgetAnimationState& State(*It);
+		if (!State.IsValid())
+		{
+			It.RemoveCurrent();
+		}
+		else if ((State.GetPlaybackStatus() == EMovieScenePlayerStatus::Stopped) && !State.IsStopping())
+		{
+			State.TearDown();
+			It.RemoveCurrent();
+		}
+	}
 }
 
 void UUserWidget::UpdatePlayerContextIfInvalid(const FLocalPlayerContext& ParentPlayerContext)
diff --git a/Engine/Source/Runtime/UMG/Public/Animation/UMGSequencePlayer.h b/Engine/Source/Runtime/UMG/Public/Animation/UMGSequencePlayer.h
--- a/Engine/Source/Runtime/UMG/Public/Animation/UMGSequencePlayer.h
+++ b/Engine/Source/Runtime/UMG/Public/Animation/UMGSequencePlayer.h
@@ -8,6 +8,7 @@
 #include "Blueprint/UserWidget.h"
 #include "IMovieScenePlayer.h"
 #include "Animation/UMGSequenceTickManager.h"
+#include "Animation/WidgetAnimationHandle.h"
 #include "Evaluation/MovieSceneEvaluationTemplateInstance.h"
 #include "Misc/QualifiedFrameTime.h"
 #include "UMGSequencePlayer.generated.h"
@@ -17,166 +18,130 @@ class UWidgetAnimation;
 UCLASS(Transient, BlueprintType, MinimalAPI)
 class UUMGSequencePlayer : public UObject, public IMovieScenePlayer
 {
-	GENERATED_UCLASS_BODY()
+	GENERATED_BODY()
 
 public:
+
+	UUMGSequencePlayer(const FObjectInitializer& ObjectInitializer);
+
+	UMG_API void InitSequencePlayer(FWidgetAnimationState& InState);
+
+	UE_DEPRECATED(5.6, "Please use FWidgetAnimationHandle and FWidgetAnimationState")
 	UMG_API void InitSequencePlayer(UWidgetAnimation& InAnimation, UUserWidget& InUserWidget);
 
 	/** Updates the running movie */
-	UMG_API void Tick( float DeltaTime );
+	UE_DEPRECATED(5.6, "Please use FWidgetAnimationHandle and FWidgetAnimationState")
+	UMG_API void Tick(float DeltaTime);
 
 	/** Begins playing or restarts an animation */
+	UE_DEPRECATED(5.6, "Please use FWidgetAnimationHandle and FWidgetAnimationState")
 	UMG_API void Play(float StartAtTime, int32 InNumLoopsToPlay, EUMGSequencePlayMode::Type InPlayMode, float InPlaybackSpeed, bool bRestoreState);
 
 	/** Begins playing or restarts an animation  and plays to the specified end time */
+	UE_DEPRECATED(5.6, "Please use FWidgetAnimationHandle and FWidgetAnimationState")
 	UMG_API void PlayTo(float StartAtTime, float EndAtTime, int32 InNumLoopsToPlay, EUMGSequencePlayMode::Type InPlayMode, float InPlaybackSpeed, bool bRestoreState);
 
 	/** Stops a running animation and resets time */
+	UE_DEPRECATED(5.6, "Please use FWidgetAnimationHandle and FWidgetAnimationState")
 	UMG_API void Stop();
 
 	/** Pauses a running animation */
+	UE_DEPRECATED(5.6, "Please use FWidgetAnimationHandle and FWidgetAnimationState")
 	UMG_API void Pause();
 
 	/** Reverses a running animation */
+	UE_DEPRECATED(5.6, "Please use FWidgetAnimationHandle and FWidgetAnimationState")
 	UMG_API void Reverse();
 
-	void SetCurrentTime(float InTime) { TimeCursorPosition = UnwarpedPosition = AnimationResolution.AsFrameTime(InTime); }
-	FQualifiedFrameTime GetCurrentTime() const { return FQualifiedFrameTime(TimeCursorPosition, AnimationResolution); }
+	UE_DEPRECATED(5.6, "Please use FWidgetAnimationHandle and FWidgetAnimationState")
+	UMG_API void SetCurrentTime(float InTime);
+
+	UE_DEPRECATED(5.6, "Please use FWidgetAnimationHandle and FWidgetAnimationState")
+	UMG_API FQualifiedFrameTime GetCurrentTime() const;
 
 	/** @return The current animation being played */
-	const UWidgetAnimation* GetAnimation() const { return Animation; }
+	UE_DEPRECATED(5.6, "Please use FWidgetAnimationHandle and FWidgetAnimationState")
+	UMG_API const UWidgetAnimation* GetAnimation() const;
 
 	/** @return */
-	UFUNCTION(BlueprintCallable, Category="Animation")
-	FName GetUserTag() const { return UserTag; }
+	UFUNCTION(BlueprintCallable, Category="Animation", meta=(DeprecatedFunction))
+	UMG_API FName GetUserTag() const;
 
-	UFUNCTION(BlueprintCallable, Category = "Animation")
-	void SetUserTag(FName InUserTag)
-	{
-		UserTag = InUserTag;
-	}
+	UFUNCTION(BlueprintCallable, Category = "Animation", meta=(DeprecatedFunction))
+	UMG_API void SetUserTag(FName InUserTag);
 
 	/** Sets the number of loops to play */
+	UE_DEPRECATED(5.6, "Please use FWidgetAnimationHandle and FWidgetAnimationState")
 	UMG_API void SetNumLoopsToPlay(int32 InNumLoopsToPlay);
 
 	/** Sets the animation playback rate */
+	UE_DEPRECATED(5.6, "Please use FWidgetAnimationHandle and FWidgetAnimationState")
 	UMG_API void SetPlaybackSpeed(float PlaybackSpeed);
 
 	/** Gets the current time position in the player (in seconds). */
-	bool IsPlayingForward() const { return bIsPlayingForward; }
+	UE_DEPRECATED(5.6, "Please use FWidgetAnimationHandle and FWidgetAnimationState")
+	UMG_API bool IsPlayingForward() const;
 
 	/** Check whether this player is currently being stopped */
-	bool IsStopping() const { return bIsStopping; }
+	UE_DEPRECATED(5.6, "Please use FWidgetAnimationHandle and FWidgetAnimationState")
+	UMG_API bool IsStopping() const;
 
-	/** IMovieScenePlayer interface */
-	virtual FMovieSceneRootEvaluationTemplateInstance& GetEvaluationTemplate() override { return RootTemplateInstance; }
-	UMG_API virtual UMovieSceneEntitySystemLinker* ConstructEntitySystemLinker() override;
-	virtual UObject* AsUObject() override { return this; }
-	virtual void SetViewportSettings(const TMap<FViewportClient*, EMovieSceneViewportParams>& ViewportParamsMap) override {}
-	virtual void GetViewportSettings(TMap<FViewportClient*, EMovieSceneViewportParams>& ViewportParamsMap) const override {}
-	UMG_API virtual EMovieScenePlayerStatus::Type GetPlaybackStatus() const override;
-	UMG_API virtual UObject* GetPlaybackContext() const override;
-	UMG_API virtual TArray<UObject*> GetEventContexts() const override;
-	UMG_API virtual void SetPlaybackStatus(EMovieScenePlayerStatus::Type InPlaybackStatus) override;
-	UMG_API virtual void PopulateUpdateFlags(UE::MovieScene::ESequenceInstanceUpdateFlags& OutFlags) override;
+public:
 
-	/** UObject interface */
-	UMG_API virtual void BeginDestroy() override;
+	UE_DEPRECATED(5.6, "Please use FWidgetAnimationHandle and FWidgetAnimationState")
+	UMG_API virtual FMovieSceneRootEvaluationTemplateInstance& GetEvaluationTemplate() override;
 
-	/** Disable this sequence player by removing any of its animation data from the entity manager */
-	UMG_API void RemoveEvaluationData();
+	UE_DEPRECATED(5.6, "Please use FWidgetAnimationHandle and FWidgetAnimationState")
+	UMG_API virtual UMovieSceneEntitySystemLinker* ConstructEntitySystemLinker() override;
 
-	UMG_API void TearDown();
+	UE_DEPRECATED(5.6, "Please use FWidgetAnimationHandle and FWidgetAnimationState")
+	UMG_API virtual UObject* AsUObject() override;
 
-	DECLARE_EVENT_OneParam(UUMGSequencePlayer, FOnSequenceFinishedPlaying, UUMGSequencePlayer&);
-	FOnSequenceFinishedPlaying& OnSequenceFinishedPlaying() { return OnSequenceFinishedPlayingEvent; }
+	UE_DEPRECATED(5.6, "Please use FWidgetAnimationHandle and FWidgetAnimationState")
+	UMG_API virtual EMovieScenePlayerStatus::Type GetPlaybackStatus() const override;
 
-private:
-	/** Internal play function with a verbose parameter set */
-	UMG_API void PlayInternal(double StartAtTime, double EndAtTime, int32 InNumLoopsToPlay, EUMGSequencePlayMode::Type InPlayMode, float InPlaybackSpeed, bool bRestoreState);
+	UE_DEPRECATED(5.6, "Please use FWidgetAnimationHandle and FWidgetAnimationState")
+	UMG_API virtual void SetPlaybackStatus(EMovieScenePlayerStatus::Type InPlaybackStatus) override;
 
-	/** Internal update function */
-	UMG_API void UpdateInternal(FFrameTime LastTimePosition, FFrameTime NextTimePosition, bool bHasJumped);
+	UE_DEPRECATED(5.6, "Please use FWidgetAnimationHandle and FWidgetAnimationState")
+	UMG_API virtual IMovieScenePlaybackClient* GetPlaybackClient() override;
 
-	UMG_API bool NeedsQueueLatentAction() const;
-	UMG_API void QueueLatentAction(FMovieSceneSequenceLatentActionDelegate Delegate);
-	UMG_API void ApplyLatentActions();
+	UE_DEPRECATED(5.6, "Please use FWidgetAnimationHandle and FWidgetAnimationState")
+	UMG_API virtual FMovieSceneSpawnRegister& GetSpawnRegister() override;
 
-	UMG_API void HandleLatentStop();
+	UE_DEPRECATED(5.6, "Please use FWidgetAnimationHandle and FWidgetAnimationState")
+	UMG_API virtual UObject* GetPlaybackContext() const override;
 
-	/** Animation being played */
-	UPROPERTY()
-	TObjectPtr<UWidgetAnimation> Animation;
+	UE_DEPRECATED(5.6, "Please use FWidgetAnimationHandle and FWidgetAnimationState")
+	UMG_API virtual void InitializeRootInstance(TSharedRef<UE::MovieScene::FSharedPlaybackState> NewSharedPlaybackState) override;
 
-	/** The user widget this sequence is animating */
-	TWeakObjectPtr<UUserWidget> UserWidget;
+public:
 
-	UPROPERTY()
-	FMovieSceneRootEvaluationTemplateInstance RootTemplateInstance;
+	UMG_API virtual void BeginDestroy() override;
 
-	TSharedPtr<FMovieSceneEntitySystemRunner> SynchronousRunner;
+	/** Disable this sequence player by removing any of its animation data from the entity manager */
+	UE_DEPRECATED(5.6, "Please use FWidgetAnimationHandle and FWidgetAnimationState")
+	UMG_API void RemoveEvaluationData();
 
-	/** The resolution at which all FFrameNumbers are stored */
-	FFrameRate AnimationResolution;
+	UE_DEPRECATED(5.6, "Please use FWidgetAnimationHandle and FWidgetAnimationState")
+	UMG_API void TearDown();
 
-	/** Start frame for the sequence playback range */
-	FFrameNumber AbsolutePlaybackStart;
+	DECLARE_EVENT_OneParam(UUMGSequencePlayer, FOnSequenceFinishedPlaying, UUMGSequencePlayer&);
 
-	/** The current time cursor position within the sequence, between 0 and Duration */
-	FFrameTime TimeCursorPosition;
+	UE_DEPRECATED(5.6, "Please use FWidgetAnimationHandle and FWidgetAnimationState")
+	UMG_API FOnSequenceFinishedPlaying& OnSequenceFinishedPlaying();
 
-	/** Unwarped position. Only different from TimeCursorPosition if there is a global playback transform. */
-	FFrameTime UnwarpedPosition;
+public:
 
-	/** The duration of the sequence */
-	int32 Duration;
+	// Internal API.
+	void BroadcastSequenceFinishedPlaying();
 
-	/** Time at which to end the animation after looping */
-	FFrameTime EndTime;
+private:
 
-	/** A compensation for delta-time we were not able to evaluate with due to ongoing budgeted evaluation */
-	float BlockedDeltaTimeCompensation;
+	FWidgetAnimationHandle WidgetAnimationHandle;
 
-	/** Status of the player (e.g play, stopped) */
-	EMovieScenePlayerStatus::Type PlayerStatus;
+	FMovieSceneRootEvaluationTemplateInstance RootTemplateInstance;
 
-	/** Delegate to call when a sequence has finished playing */
 	FOnSequenceFinishedPlaying OnSequenceFinishedPlayingEvent;
-
-	/** The number of times to loop the animation playback */
-	int32 NumLoopsToPlay;
-
-	/** The number of loops completed since the last call to Play() */
-	int32 NumLoopsCompleted;
-
-	/** The speed at which the animation should be played */
-	float PlaybackSpeed;
-
-	/** Whether to restore pre-animated state */
-	bool bRestoreState;
-
-	/** The current playback mode. */
-	EUMGSequencePlayMode::Type PlayMode;
-
-	/**
-	 * The 'state' tag the user may want to use to track what the animation is for.  It's very common in UI to use the
-	 * same animation for intro / outro, so this allows you to tag what the animation is currently doing so that you can
-	 * have some events just get called back when the animation finishes the outtro, to say, remove the UI then.
-	 */
-	FName UserTag;
-
-	/** True if the animation is playing forward, otherwise false and it's playing in reverse. */
-	bool bIsPlayingForward;
-
-	/** Set to true when this player is in the process of being stopped */
-	bool bIsStopping : 1;
-
-	/** Set to true when this player has just started playing, but hasn't finished evaluating yet */
-	bool bIsBeginningPlay : 1;
-
-	/** Set to true if we need to run the finishing logic in post-evaluation */
-	bool bCompleteOnPostEvaluation : 1;
-
-	/** List of latent action delegates. Only used when multi-threaded animation evaluation is disabled */
-	TArray<FMovieSceneSequenceLatentActionDelegate> LatentActions;
 };
+
diff --git a/Engine/Source/Runtime/UMG/Public/Animation/WidgetAnimationHandle.h b/Engine/Source/Runtime/UMG/Public/Animation/WidgetAnimationHandle.h
new file mode 100644
--- /dev/null
+++ b/Engine/Source/Runtime/UMG/Public/Animation/WidgetAnimationHandle.h
@@ -0,0 +1,88 @@
+// Copyright Epic Games, Inc. All Rights Reserved.
+
+#pragma once
+
+#include "CoreTypes.h"
+#include "Kismet/BlueprintFunctionLibrary.h"
+#include "UObject/WeakObjectPtr.h"
+
+#include "WidgetAnimationHandle.generated.h"
+
+class UUMGSequencePlayer;
+class UUserWidget;
+struct FWidgetAnimationState;
+
+/**
+ * Handle to an ongoing or finished widget animation.
+ */
+USTRUCT(BlueprintType)
+struct FWidgetAnimationHandle
+{
+	GENERATED_BODY()
+
+public:
+
+	/** Creates an invalid handle. */
+	UMG_API FWidgetAnimationHandle();
+
+	/** Gets the animation state. */
+	UMG_API FWidgetAnimationState* GetAnimationState() const;
+
+public:
+
+	// For backwards compatibility in C++ code.
+	UMG_API UUMGSequencePlayer* GetSequencePlayer() const;
+	inline operator UUMGSequencePlayer*() const { return GetSequencePlayer(); }
+
+public:
+
+	/**
+	 * Returns whether this handle is valid.
+	 * A valid handle may still return a null animation state if the animation has finished playing.
+	 */
+	UMG_API bool IsValid() const;
+
+	/** Gets the user tag associated with the running animation. */
+	UMG_API FName GetUserTag() const;
+
+	/** Sets the user tag associated with the running animation. */
+	UMG_API void SetUserTag(FName InUserTag);
+
+private:
+
+	/** Creates a handle for the given running animation on the given widget. */
+	UMG_API FWidgetAnimationHandle(UUserWidget* InUserWidget, int32 InStateIndex, uint32 InStateSerial);
+
+private:
+
+	/** The widget this handle relates to. */
+	UPROPERTY()
+	TWeakObjectPtr<UUserWidget> WeakUserWidget;
+
+	/** The animation state index. */
+	UPROPERTY()
+	int32 StateIndex = INDEX_NONE;
+
+	/** The animation state serial. */
+	UPROPERTY()
+	uint32 StateSerial = 0;
+
+	// Only UUserWidget and FWidgetAnimationState can create handles.
+	friend class UUserWidget;
+	friend struct FWidgetAnimationState;
+};
+
+UCLASS()
+class UMG_API UWidgetAnimationHandleFunctionLibrary : public UBlueprintFunctionLibrary
+{
+	GENERATED_BODY()
+
+public:
+
+	UFUNCTION(BlueprintPure, Category = "UMG")
+	static FName GetUserTag(const FWidgetAnimationHandle& Target);
+
+	UFUNCTION(BlueprintCallable, Category = "UMG")
+	static void SetUserTag(UPARAM(ref) FWidgetAnimationHandle& Target, FName InUserTag);
+};
+
diff --git a/Engine/Source/Runtime/UMG/Public/Animation/WidgetAnimationPlayCallbackProxy.h b/Engine/Source/Runtime/UMG/Public/Animation/WidgetAnimationPlayCallbackProxy.h
--- a/Engine/Source/Runtime/UMG/Public/Animation/WidgetAnimationPlayCallbackProxy.h
+++ b/Engine/Source/Runtime/UMG/Public/Animation/WidgetAnimationPlayCallbackProxy.h
@@ -2,9 +2,16 @@
 
 #pragma once
 
+#include "Animation/WidgetAnimationHandle.h"
 #include "Blueprint/UserWidget.h"
+
 #include "WidgetAnimationPlayCallbackProxy.generated.h"
 
+class UUMGSequencePlayer;
+class UUserWidget;
+class UWidgetAnimation;
+struct FWidgetAnimationState;
+
 DECLARE_DYNAMIC_MULTICAST_DELEGATE(FWidgetAnimationResult);
 
 UCLASS(MinimalAPI)
@@ -16,17 +23,75 @@ class UWidgetAnimationPlayCallbackProxy : public UObject
 	UPROPERTY(BlueprintAssignable)
 	FWidgetAnimationResult Finished;
 
-	UFUNCTION(BlueprintCallable, meta = (BlueprintInternalUseOnly = "true", DisplayName = "Play Animation with Finished event", ShortToolTip = "Play Animation and trigger event on Finished",  ToolTip="Play Animation on widget and trigger Finish event when the animation is done."), Category = "User Interface|Animation")
-	static UWidgetAnimationPlayCallbackProxy* CreatePlayAnimationProxyObject(class UUMGSequencePlayer*& Result, class UUserWidget* Widget, class UWidgetAnimation* InAnimation, float StartAtTime = 0.0f, int32 NumLoopsToPlay = 1, EUMGSequencePlayMode::Type PlayMode = EUMGSequencePlayMode::Forward, float PlaybackSpeed = 1.0f);
+	UFUNCTION(BlueprintCallable , Category = "User Interface|Animation", meta = (
+				DeprecatedFunction,
+				DeprecationMessage = "Use NewPlayAnimationProxyObject",
+				BlueprintInternalUseOnly = "true",
+				DisplayName = "Play Animation with Finished event (legacy)",
+				ShortToolTip = "Play Animation and trigger event on Finished (legacy version using deprecated UMG Sequencer Player)",
+				ToolTip="Play Animation on widget and trigger Finish event when the animation is done (legacy version using deprecated UMG Sequencer Player)."))
+	static UWidgetAnimationPlayCallbackProxy* CreatePlayAnimationProxyObject(
+			UUMGSequencePlayer*& Result,
+			UUserWidget* Widget,
+			UWidgetAnimation* InAnimation,
+			float StartAtTime = 0.0f,
+			int32 NumLoopsToPlay = 1,
+			EUMGSequencePlayMode::Type PlayMode = EUMGSequencePlayMode::Forward,
+			float PlaybackSpeed = 1.0f);
+
+	UFUNCTION(BlueprintCallable , Category = "User Interface|Animation", meta = (
+				BlueprintInternalUseOnly = "true",
+				DisplayName = "Play Animation with Finished event",
+				ShortToolTip = "Play Animation and trigger event on Finished",
+				ToolTip="Play Animation on widget and trigger Finish event when the animation is done."))
+	static UWidgetAnimationPlayCallbackProxy* NewPlayAnimationProxyObject(
+			FWidgetAnimationHandle& Result,
+			UUserWidget* Widget,
+			UWidgetAnimation* InAnimation,
+			float StartAtTime = 0.0f,
+			int32 NumLoopsToPlay = 1,
+			EUMGSequencePlayMode::Type PlayMode = EUMGSequencePlayMode::Forward,
+			float PlaybackSpeed = 1.0f);
 
-	UFUNCTION(BlueprintCallable, meta = (BlueprintInternalUseOnly = "true", DisplayName = "Play Animation Time Range with Finished event", ShortToolTip = "Play Animation Time Range and trigger event on Finished", ToolTip = "Play Animation Time Range on widget and trigger Finish event when the animation is done."), Category = "User Interface|Animation")
-	static UWidgetAnimationPlayCallbackProxy* CreatePlayAnimationTimeRangeProxyObject(class UUMGSequencePlayer*& Result, class UUserWidget* Widget, class UWidgetAnimation* InAnimation, float StartAtTime = 0.0f, float EndAtTime = 0.0f, int32 NumLoopsToPlay = 1, EUMGSequencePlayMode::Type PlayMode = EUMGSequencePlayMode::Forward, float PlaybackSpeed = 1.0f);
+	UFUNCTION(BlueprintCallable, Category = "User Interface|Animation", meta = (
+				DeprecatedFunction, 
+				DeprecationMessage = "Use NewPlayAnimationTimeRangeProxyObject",
+				BlueprintInternalUseOnly = "true",
+				DisplayName = "Play Animation Time Range with Finished event (legacy)",
+				ShortToolTip = "Play Animation Time Range and trigger event on Finished (legacy version using deprecated UMG Sequencer Player)",
+				ToolTip = "Play Animation Time Range on widget and trigger Finish event when the animation is done (legacy version using deprecated UMG Sequencer Player)."))
+	static UWidgetAnimationPlayCallbackProxy* CreatePlayAnimationTimeRangeProxyObject(
+			UUMGSequencePlayer*& Result,
+			UUserWidget* Widget,
+			UWidgetAnimation* InAnimation,
+			float StartAtTime = 0.0f,
+			float EndAtTime = 0.0f,
+			int32 NumLoopsToPlay = 1,
+			EUMGSequencePlayMode::Type PlayMode = EUMGSequencePlayMode::Forward,
+			float PlaybackSpeed = 1.0f);
+
+	UFUNCTION(BlueprintCallable, Category = "User Interface|Animation", meta = (
+				BlueprintInternalUseOnly = "true",
+				DisplayName = "Play Animation Time Range with Finished event",
+				ShortToolTip = "Play Animation Time Range and trigger event on Finished",
+				ToolTip = "Play Animation Time Range on widget and trigger Finish event when the animation is done."))
+	static UWidgetAnimationPlayCallbackProxy* NewPlayAnimationTimeRangeProxyObject(
+			FWidgetAnimationHandle& Result,
+			UUserWidget* Widget,
+			UWidgetAnimation* InAnimation,
+			float StartAtTime = 0.0f,
+			float EndAtTime = 0.0f,
+			int32 NumLoopsToPlay = 1,
+			EUMGSequencePlayMode::Type PlayMode = EUMGSequencePlayMode::Forward,
+			float PlaybackSpeed = 1.0f);
 
 private:
-	class UUMGSequencePlayer* ExecutePlayAnimation(class UUserWidget* Widget, class UWidgetAnimation* InAnimation, float StartAtTime, int32 NumLoopsToPlay, EUMGSequencePlayMode::Type PlayMode, float PlaybackSpeed);
-	class UUMGSequencePlayer* ExecutePlayAnimationTimeRange(class UUserWidget* Widget, class UWidgetAnimation* InAnimation, float StartAtTime, float EndAtTime, int32 NumLoopsToPlay, EUMGSequencePlayMode::Type PlayMode, float PlaybackSpeed);
-	void OnSequenceFinished(class UUMGSequencePlayer& Player);
+
+	void ExecutePlayAnimation(UUserWidget* Widget, UWidgetAnimation* InAnimation, float StartAtTime, int32 NumLoopsToPlay, EUMGSequencePlayMode::Type PlayMode, float PlaybackSpeed);
+	void ExecutePlayAnimationTimeRange(UUserWidget* Widget, UWidgetAnimation* InAnimation, float StartAtTime, float EndAtTime, int32 NumLoopsToPlay, EUMGSequencePlayMode::Type PlayMode, float PlaybackSpeed);
+	void OnSequenceFinished(FWidgetAnimationState& State);
 	bool OnAnimationFinished(float DeltaTime);
 
+	FWidgetAnimationHandle WidgetAnimationHandle;
 	FDelegateHandle OnFinishedHandle;
 };
diff --git a/Engine/Source/Runtime/UMG/Public/Animation/WidgetAnimationState.h b/Engine/Source/Runtime/UMG/Public/Animation/WidgetAnimationState.h
new file mode 100644
--- /dev/null
+++ b/Engine/Source/Runtime/UMG/Public/Animation/WidgetAnimationState.h
@@ -0,0 +1,197 @@
+// Copyright Epic Games, Inc. All Rights Reserved.
+
+#pragma once
+
+#include "Blueprint/UMGSequencePlayMode.h"
+#include "Animation/UMGSequenceTickManager.h"
+#include "CoreMinimal.h"
+#include "Evaluation/MovieScenePlayback.h"
+#include "Evaluation/MovieScenePlaybackManager.h"
+#include "EntitySystem/MovieSceneSharedPlaybackState.h"
+#include "Misc/QualifiedFrameTime.h"
+#include "MovieSceneFwd.h"
+#include "UObject/ObjectMacros.h"
+
+class UUserWidget;
+class UWidgetAnimation;
+class UUMGSequencePlayer;
+struct FWidgetAnimationHandle;
+struct FWidgetAnimationState;
+
+DECLARE_MULTICAST_DELEGATE_OneParam(FOnWidgetAnimationEvent, FWidgetAnimationState&);
+
+struct FWidgetAnimationStatePlayParams
+{
+	/** The time to start playing at. */
+	double StartAtTime = 0;
+	/** An optional start offset to apply to the animation. */
+	TOptional<double> StartOffset;
+	/** An optional end offset to apply to the animation. */
+	TOptional<double> EndOffset;
+	/** The number of loops to play before finishing the animation. */
+	int32 NumLoopsToPlay = 1;
+	/** The play-rate of the animation. */
+	float PlaybackSpeed = 1.f;
+	/** The play mode of the animation. */
+	EUMGSequencePlayMode::Type PlayMode = EUMGSequencePlayMode::Forward;
+	/** Whether to restore pre-animated state after the animation has finished. */
+	bool bRestoreState = false;
+};
+
+/**
+ * A class that runs an animation on a widget.
+ */
+struct FWidgetAnimationState
+{
+public:
+
+	/** Builds a new widget animation state. */
+	FWidgetAnimationState();
+
+	/** Initializes the widget animation state. */
+	void Initialize(UWidgetAnimation* InAnimation, UUserWidget* InUserWidget, int32 InIndex, int32 InSerialNumber);
+
+	/** Gets the user widget this state is animating. */
+	UUserWidget* GetUserWidget() const { return WeakUserWidget.Get(); }
+
+	/** Gets the current animation being played */
+	const UWidgetAnimation* GetAnimation() const { return Animation; }
+
+	/** Gets the playback state of the animation being played. */
+	TSharedPtr<UE::MovieScene::FSharedPlaybackState> GetSharedPlaybackState() const { return WeakPlaybackState.Pin(); }
+
+	/** Gets the user tag for the animation. */
+	FName GetUserTag() const { return UserTag; }
+
+	/** Sets the user tag for the animation. */
+	void SetUserTag(FName InUserTag) { UserTag = InUserTag; }
+
+	/** Returns whether this state is currently stopping. */
+	bool IsStopping() const { return bIsStopping; }
+
+	/** Allows registering a callback for when the animation has finished playing. */
+	FOnWidgetAnimationEvent& GetOnWidgetAnimationFinished() { return OnWidgetAnimationFinishedEvent; }
+
+public:
+
+	/** Gets a legacy player object for backwards compatibility, lazily created it if needed. */
+	UMG_API UUMGSequencePlayer* GetOrCreateLegacyPlayer();
+
+	/** Gets a legacy player object for backwards compatibility, returning null if none exists yet. */
+	UMG_API UUMGSequencePlayer* GetLegacyPlayer() const;
+
+	/** Gets the widget animation handle for the animation being played. */
+	UMG_API FWidgetAnimationHandle GetAnimationHandle() const;
+
+	/** Gets whether the animation is playing forwards or backwards. */
+	UMG_API bool IsPlayingForward() const;
+
+	/** Gets the playback status of the animation. */
+	UMG_API EMovieScenePlayerStatus::Type GetPlaybackStatus() const;
+
+	/** Sets the playback status of the animation. */
+	UMG_API void SetPlaybackStatus(EMovieScenePlayerStatus::Type InPlaybackStatus);
+
+	/** Gets the current time of the animation. */
+	UMG_API FQualifiedFrameTime GetCurrentTime() const;
+
+	/** Sets the current time of the animation. */
+	UMG_API void SetCurrentTime(float InTime);
+
+	/** Sets the number of loops to play. */
+	UMG_API void SetNumLoopsToPlay(int32 InNumLoopsToPlay);
+
+	/** Sets the animation playback rate. */
+	UMG_API void SetPlaybackSpeed(float PlaybackSpeed);
+
+	/** Update the animation. */
+	UMG_API void Tick(float InDeltaSeconds);
+
+	/** Begins playing or restarts an animation. */
+	UMG_API void Play(const FWidgetAnimationStatePlayParams& PlayParams);
+
+	/** Stops a running animation and resets time. */
+	UMG_API void Stop();
+
+	/** Pauses a running animation. */
+	UMG_API void Pause();
+
+	/** Reverses a running animation. */
+	UMG_API void Reverse();
+
+	/** Disable this sequence player by removing any of its animation data from the entity manager. */
+	UMG_API void RemoveEvaluationData();
+
+	/** Tears down the animation. */
+	UMG_API void TearDown();
+
+	/** Returns whether this animation is valid. An uninitialized or torn-down state is not valid. */
+	UMG_API bool IsValid() const;
+
+	/** Collect objects for the GC. */
+	UMG_API void AddReferencedObjects(FReferenceCollector& Collector);
+
+private:
+
+	bool NeedsLegacyPlayer() const;
+
+	void OnBegunPlay();
+	void OnStopped();
+
+	static void OnBegunPlay(FWidgetAnimationHandle InHandle);
+	static void OnStopped(FWidgetAnimationHandle InHandle);
+
+private:
+
+	/** Animation being played */
+	TObjectPtr<UWidgetAnimation> Animation;
+
+	/** Legacy sequence player for backwards compatibility */
+	TObjectPtr<UUMGSequencePlayer> LegacyPlayer;
+
+	/** The user widget this sequence is animating */
+	TWeakObjectPtr<UUserWidget> WeakUserWidget;
+
+	/** The index of this state. */
+	int32 StateIndex;
+
+	/** The serial number of this state. */
+	int32 SerialNumber;
+
+	// Access to index and serial number.
+	friend class UUserWidget;
+	friend struct FWidgetAnimationHandle;
+
+private:
+
+	/** Shared playback state for the animation. */
+	TWeakPtr<UE::MovieScene::FSharedPlaybackState> WeakPlaybackState;
+
+	/** Playback manager for the animation. */
+	FMovieScenePlaybackManager PlaybackManager;
+
+	/** The current playback mode. */
+	EUMGSequencePlayMode::Type PlayMode;
+
+	/**
+	 * The 'state' tag the user may want to use to track what the animation is for.
+	 *
+	 * It's very common in UI to use the same animation for intro / outro, so this allows you 
+	 * to tag what the animation is currently doing so that you can have some events just get 
+	 * called back when the animation finishes the outtro, to say, remove the UI then.
+	 */
+	FName UserTag;
+
+	/** Whether to restore pre-animated state. */
+	bool bRestoreState : 1;
+
+	/** Whether we are in the process of starting play. */
+	bool bIsBeginningPlay : 1;
+
+	/** Whether we are in the process of stopping. */
+	bool bIsStopping : 1;
+
+	/** Callback for when the animation has finished playing. */
+	FOnWidgetAnimationEvent OnWidgetAnimationFinishedEvent;
+};
+
diff --git a/Engine/Source/Runtime/UMG/Public/Blueprint/UMGSequencePlayMode.h b/Engine/Source/Runtime/UMG/Public/Blueprint/UMGSequencePlayMode.h
new file mode 100644
--- /dev/null
+++ b/Engine/Source/Runtime/UMG/Public/Blueprint/UMGSequencePlayMode.h
@@ -0,0 +1,25 @@
+// Copyright Epic Games, Inc. All Rights Reserved.
+
+#pragma once
+
+#include "CoreTypes.h"
+
+#include "UMGSequencePlayMode.generated.h"
+
+class UUserWidget;
+
+/** Describes playback modes for UMG sequences. */
+UENUM(BlueprintType)
+namespace EUMGSequencePlayMode
+{
+	enum Type : int
+	{
+		/** Animation plays and loops from the beginning to the end. */
+		Forward,
+		/** Animation plays and loops from the end to the beginning. */
+		Reverse,
+		/** Animation plays from the beginning to the end and then from the end to the beginning. */
+		PingPong,
+	};
+}
+
diff --git a/Engine/Source/Runtime/UMG/Public/Blueprint/UserWidget.h b/Engine/Source/Runtime/UMG/Public/Blueprint/UserWidget.h
--- a/Engine/Source/Runtime/UMG/Public/Blueprint/UserWidget.h
+++ b/Engine/Source/Runtime/UMG/Public/Blueprint/UserWidget.h
@@ -5,6 +5,7 @@
 #if UE_ENABLE_INCLUDE_ORDER_DEPRECATED_IN_5_4
 #include "CoreMinimal.h"
 #endif //UE_ENABLE_INCLUDE_ORDER_DEPRECATED_IN_5_4
+#include "Blueprint/UMGSequencePlayMode.h"
 #include "Blueprint/WidgetChild.h"
 #include "UObject/ObjectMacros.h"
 #include "UObject/ObjectSaveContext.h"
@@ -14,6 +15,7 @@
 #include "Input/Events.h"
 #include "Input/Reply.h"
 #include "Widgets/SWidget.h"
+#include "Kismet/BlueprintFunctionLibrary.h"
 #include "Layout/Margin.h"
 #include "Components/SlateWrapperTypes.h"
 #include "Components/Widget.h"
@@ -27,6 +29,8 @@
 #include "EngineStats.h"
 #include "SlateGlobals.h"
 #include "Animation/WidgetAnimationEvents.h"
+#include "Animation/WidgetAnimationHandle.h"
+#include "Animation/WidgetAnimationState.h"
 #include "UserWidget.generated.h"
 
 class Error;
@@ -60,6 +64,17 @@ enum class EQueuedWidgetAnimationMode : uint8
 	None,
 };
 
+/** Describes whether the widget needs a legacy UMG sequence player to be created. */
+enum class ERequiresLegacyPlayer : uint8
+{
+	/** Will create a legacy player if the widget class has native C++ children. */
+	AutoDetect,
+	/** Force creating a legacy player. */
+	Yes,
+	/** Declares that no legacy player is required. */
+	No
+};
+
 /**
  * Struct that maintains state of currently queued animation transtions to be evaluated next frame.
  */
@@ -227,23 +242,6 @@ public:
 	TObjectPtr<UWidget> Content;
 };
 
-class UUMGSequencePlayer;
-
-/** Describes playback modes for UMG sequences. */
-UENUM(BlueprintType)
-namespace EUMGSequencePlayMode
-{
-	enum Type : int
-	{
-		/** Animation plays and loops from the beginning to the end. */
-		Forward,
-		/** Animation plays and loops from the end to the beginning. */
-		Reverse,
-		/** Animation plays from the beginning to the end and then from the end to the beginning. */
-		PingPong,
-	};
-}
-
 #if WITH_EDITORONLY_DATA
 
 UENUM()
@@ -258,6 +256,16 @@ enum class EDesignPreviewSizeMode : uint8
 
 #endif
 
+UCLASS()
+class UUserWidgetFunctionLibrary : public UBlueprintFunctionLibrary
+{
+	GENERATED_BODY()
+
+public:
+	UFUNCTION(BlueprintPure, Category = "UMG", meta = (BlueprintAutocast, DeprecatedFunction, DeprecatedMessage="UMGSequencePlayer is deprecated, please use FWidgetAnimationState"))
+	static UUMGSequencePlayer* Conv_UMGSequencePlayer(const FWidgetAnimationHandle& WidgetAnimationHandle);
+};
+
 DECLARE_DYNAMIC_MULTICAST_DELEGATE(FOnConstructEvent);
 
 DECLARE_DYNAMIC_DELEGATE( FOnInputAction );
@@ -303,6 +311,9 @@ protected:
 	/** The function is implemented only in nativized widgets (automatically converted from BP to c++) */
 	virtual void InitializeNativeClassData() {}
 
+	/** Returns whether this widget needs a legacy player. */
+	virtual ERequiresLegacyPlayer GetLegacyPlayerRequirement() const { return ERequiresLegacyPlayer::AutoDetect; }
+
 	UMG_API void InitializeNamedSlots();
 
 public:
@@ -945,17 +956,39 @@ protected:
 	UMG_API virtual void OnAnimationFinished_Implementation(const UWidgetAnimation* Animation);
 
 	/** Broadcast any events based on a state transition for the sequence player, started, finished...etc. */
+	UE_DEPRECATED(5.6, "UUMGSequencePlayer is deprecated")
 	UMG_API void BroadcastAnimationStateChange(const UUMGSequencePlayer& Player, EWidgetAnimationEvent AnimationEvent);
 
+	UMG_API void BroadcastAnimationStateChange(const FWidgetAnimationState& State, EWidgetAnimationEvent AnimationEvent);
+
+private:
+
+	UMG_API void BroadcastAnimationStateChange(const UWidgetAnimation* Animation, FName UserTag, EWidgetAnimationEvent AnimationEvent);
+
+	UMG_API void BroadcastAnimationStartedPlaying(FWidgetAnimationState& State);
+	UMG_API void BroadcastAnimationFinishedPlaying(FWidgetAnimationState& State);
+
 protected:
 
 	/** Called when a sequence player is finished playing an animation */
-	UMG_API virtual void OnAnimationStartedPlaying(UUMGSequencePlayer& Player);
+	UMG_API FOnWidgetAnimationEvent& OnAnimationStartedPlaying() { return OnAnimationStartedPlayingEvent; }
 
 	/** Called when a sequence player is finished playing an animation */
-	UMG_API virtual void OnAnimationFinishedPlaying(UUMGSequencePlayer& Player);
+	UMG_API FOnWidgetAnimationEvent& OnAnimationFinishedPlaying() { return OnAnimationFinishedPlayingEvent; }
+
+	UE_DEPRECATED(5.6, "UUMGSequencePlayer is deprecated")
+	UMG_API virtual void OnAnimationStartedPlaying(UUMGSequencePlayer& Player) {}
+
+	UE_DEPRECATED(5.6, "UUMGSequencePlayer is deprecated")
+	UMG_API virtual void OnAnimationFinishedPlaying(UUMGSequencePlayer& Player) {}
+
+private:
+
+	FOnWidgetAnimationEvent OnAnimationStartedPlayingEvent;
+	FOnWidgetAnimationEvent OnAnimationFinishedPlayingEvent;
 
 public:
+
 	UE_DEPRECATED(5.2, "Direct access to ColorAndOpacity is deprecated. Please use the getter or setter.")
 	/** The color and opacity of this widget.  Tints all child widgets. */
 	UPROPERTY(EditAnywhere, BlueprintReadWrite, Getter, Setter, BlueprintSetter = "SetColorAndOpacity", Category = "Appearance")
@@ -1156,7 +1189,7 @@ public:
 	 * @param bRestoreState Restores widgets to their pre-animated state when the animation stops
 	 */
 	UFUNCTION(BlueprintCallable, BlueprintCosmetic, Category = "User Interface|Animation")
-	UMG_API UUMGSequencePlayer* PlayAnimation(UWidgetAnimation* InAnimation, float StartAtTime = 0.0f, int32 NumLoopsToPlay = 1, EUMGSequencePlayMode::Type PlayMode = EUMGSequencePlayMode::Forward, float PlaybackSpeed = 1.0f, bool bRestoreState = false);
+	UMG_API FWidgetAnimationHandle PlayAnimation(UWidgetAnimation* InAnimation, float StartAtTime = 0.0f, int32 NumLoopsToPlay = 1, EUMGSequencePlayMode::Type PlayMode = EUMGSequencePlayMode::Forward, float PlaybackSpeed = 1.0f, bool bRestoreState = false);
 
 	/**
 	 * Plays an animation in this widget a specified number of times stopping at a specified time
@@ -1170,7 +1203,7 @@ public:
 	 * @param bRestoreState Restores widgets to their pre-animated state when the animation stops
 	 */
 	UFUNCTION(BlueprintCallable, BlueprintCosmetic, Category = "User Interface|Animation")
-	UMG_API UUMGSequencePlayer* PlayAnimationTimeRange(UWidgetAnimation* InAnimation, float StartAtTime = 0.0f, float EndAtTime = 0.0f, int32 NumLoopsToPlay = 1, EUMGSequencePlayMode::Type PlayMode = EUMGSequencePlayMode::Forward, float PlaybackSpeed = 1.0f, bool bRestoreState = false);
+	UMG_API FWidgetAnimationHandle PlayAnimationTimeRange(UWidgetAnimation* InAnimation, float StartAtTime = 0.0f, float EndAtTime = 0.0f, int32 NumLoopsToPlay = 1, EUMGSequencePlayMode::Type PlayMode = EUMGSequencePlayMode::Forward, float PlaybackSpeed = 1.0f, bool bRestoreState = false);
 
 	/**
 	 * Plays an animation on this widget relative to it's current state forward.  You should use this version in situations where
@@ -1183,7 +1216,7 @@ public:
 	 * @param bRestoreState Restores widgets to their pre-animated state when the animation stops
 	 */
 	UFUNCTION(BlueprintCallable, BlueprintCosmetic, Category = "User Interface|Animation")
-	UMG_API UUMGSequencePlayer* PlayAnimationForward(UWidgetAnimation* InAnimation, float PlaybackSpeed = 1.0f, bool bRestoreState = false);
+	UMG_API FWidgetAnimationHandle PlayAnimationForward(UWidgetAnimation* InAnimation, float PlaybackSpeed = 1.0f, bool bRestoreState = false);
 
 	/**
 	 * Plays an animation on this widget relative to it's current state in reverse.  You should use this version in situations where
@@ -1196,7 +1229,7 @@ public:
 	 * @param bRestoreState Restores widgets to their pre-animated state when the animation stops
 	 */
 	UFUNCTION(BlueprintCallable, BlueprintCosmetic, Category = "User Interface|Animation")
-	UMG_API UUMGSequencePlayer* PlayAnimationReverse(UWidgetAnimation* InAnimation, float PlaybackSpeed = 1.0f, bool bRestoreState = false);
+	UMG_API FWidgetAnimationHandle PlayAnimationReverse(UWidgetAnimation* InAnimation, float PlaybackSpeed = 1.0f, bool bRestoreState = false);
 
 	/**
 	 * Stops an already running animation in this widget
@@ -1387,7 +1420,7 @@ public:
 
 	/** Are we currently playing any animations? */
 	UFUNCTION(BlueprintCallable, Category="User Interface|Animation")
-	FORCEINLINE bool IsPlayingAnimation() const { return ActiveSequencePlayers.Num() > 0; }
+	UMG_API bool IsPlayingAnimation() const;
 
 #if WITH_EDITOR
 	//~ Begin UWidget Interface
@@ -1429,9 +1462,10 @@ public:
 	static UMG_API UUserWidget* CreateWidgetInstance(UWorld& World, TSubclassOf<UUserWidget> UserWidgetClass, FName WidgetName);
 
 private:
+
 	static UMG_API UUserWidget* CreateInstanceInternal(UObject* Outer, TSubclassOf<UUserWidget> UserWidgetClass, FName WidgetName, UWorld* World, ULocalPlayer* LocalPlayer);
 
-	UMG_API void ClearStoppedSequencePlayers();
+	void ClearStoppedAnimationStates();
 
 	// If this widget's player context is invalid, this function will set it to the provided player context, otherwise it will keep its current player context and call CheckUpdatePlayerContext on its children
 	void UpdatePlayerContextIfInvalid(const FLocalPlayerContext& ParentPlayerContext);
@@ -1442,18 +1476,28 @@ public:
 	UPROPERTY(Transient)
 	TArray<FQueuedWidgetAnimationTransition> QueuedWidgetAnimationTransitions;
 
-	/** All the sequence players currently playing */
+	UE_DEPRECATED(5.6, "Please use ActiveAnimations")
 	UPROPERTY(Transient)
 	TArray<TObjectPtr<UUMGSequencePlayer>> ActiveSequencePlayers;
-
+	
+	/** Global tick manager for running widget animations */
 	UPROPERTY(Transient)
 	TObjectPtr<UUMGSequenceTickManager> AnimationTickManager;
 
-	/** List of sequence players to cache and clean up when safe */
+	UE_DEPRECATED(5.6, "Please use ActiveAnimations")
 	UPROPERTY(Transient)
 	TArray<TObjectPtr<UUMGSequencePlayer>> StoppedSequencePlayers;
 
 private:
+
+	/** All the sequences currently playing */
+	TSparseArray<FWidgetAnimationState> ActiveAnimations;
+
+	/** Next serial number for animation handles */
+	uint32 NextSerialNumber = 0;
+
+private:
+
 	/** Stores the widgets being assigned to named slots */
 	UPROPERTY()
 	TArray<FNamedSlotBinding> NamedSlotBindings;
@@ -1514,6 +1558,9 @@ private:
 	uint8 bStoppingAllAnimations : 1;
 
 protected:
+
+	static UMG_API void AddReferencedObjects(UObject* InThis, FReferenceCollector& Collector);
+
 	UMG_API virtual TSharedRef<SWidget> RebuildWidget() override;
 	UMG_API virtual void OnWidgetRebuilt() override;
 
@@ -1582,12 +1629,15 @@ protected:
 	 * Ticks the active sequences and latent actions that have been scheduled for this Widget.
 	 */
 	UMG_API void TickActionsAndAnimation(float InDeltaTime);
-	UMG_API void PostTickActionsAndAnimation(float InDeltaTime);
+
+	UE_DEPRECATED(5.6, "This method isn't needed anymore")
+	UMG_API void PostTickActionsAndAnimation(float InDeltaTime) {}
 
 	UMG_API void RemoveObsoleteBindings(const TArray<FName>& NamedSlots);
 
-	UMG_API UUMGSequencePlayer* GetSequencePlayer(const UWidgetAnimation* InAnimation) const;
-	UMG_API UUMGSequencePlayer* GetOrAddSequencePlayer(UWidgetAnimation* InAnimation);
+	UMG_API const FWidgetAnimationState* GetAnimationState(const UWidgetAnimation* InAnimation) const;
+	UMG_API FWidgetAnimationState* GetAnimationState(const UWidgetAnimation* InAnimation);
+	UMG_API FWidgetAnimationState* GetOrAddAnimationState(UWidgetAnimation* InAnimation);
 
 	UMG_API void ExecuteQueuedAnimationTransitions();
 
@@ -1695,6 +1745,8 @@ protected:
 	 * UserWidget of state transitions.
 	 */
 	friend UUMGSequencePlayer;
+	friend FWidgetAnimationState;
+	friend FWidgetAnimationHandle;
 
 	friend UUMGSequenceTickManager;
 
