UE: Remove deprecation for now

diff --git a/Engine/Source/Runtime/UMG/Public/Animation/WidgetAnimationPlayCallbackProxy.h b/Engine/Source/Runtime/UMG/Public/Animation/WidgetAnimationPlayCallbackProxy.h
--- a/Engine/Source/Runtime/UMG/Public/Animation/WidgetAnimationPlayCallbackProxy.h
+++ b/Engine/Source/Runtime/UMG/Public/Animation/WidgetAnimationPlayCallbackProxy.h
@@ -24,8 +24,6 @@ class UWidgetAnimationPlayCallbackProxy : public UObject
 	FWidgetAnimationResult Finished;
 
 	UFUNCTION(BlueprintCallable , Category = "User Interface|Animation", meta = (
-				DeprecatedFunction,
-				DeprecationMessage = "Use NewPlayAnimationProxyObject",
 				BlueprintInternalUseOnly = "true",
 				DisplayName = "Play Animation with Finished event (legacy)",
 				ShortToolTip = "Play Animation and trigger event on Finished (legacy version using deprecated UMG Sequencer Player)",
@@ -54,8 +52,6 @@ class UWidgetAnimationPlayCallbackProxy : public UObject
 			float PlaybackSpeed = 1.0f);
 
 	UFUNCTION(BlueprintCallable, Category = "User Interface|Animation", meta = (
-				DeprecatedFunction, 
-				DeprecationMessage = "Use NewPlayAnimationTimeRangeProxyObject",
 				BlueprintInternalUseOnly = "true",
 				DisplayName = "Play Animation Time Range with Finished event (legacy)",
 				ShortToolTip = "Play Animation Time Range and trigger event on Finished (legacy version using deprecated UMG Sequencer Player)",
