UMG: add type hash for animation handles so we can use them as map keys

diff --git a/Engine/Source/Runtime/UMG/Public/Animation/WidgetAnimationHandle.h b/Engine/Source/Runtime/UMG/Public/Animation/WidgetAnimationHandle.h
--- a/Engine/Source/Runtime/UMG/Public/Animation/WidgetAnimationHandle.h
+++ b/Engine/Source/Runtime/UMG/Public/Animation/WidgetAnimationHandle.h
@@ -48,6 +48,14 @@ public:
 	/** Sets the user tag associated with the running animation. */
 	UMG_API void SetUserTag(FName InUserTag);
 
+public:
+
+	/** Type hash for using handles in maps and other containers. */
+	friend uint32 GetTypeHash(const FWidgetAnimationHandle& Handle)
+	{
+		return HashCombine(GetTypeHash(Handle.WeakUserWidget), Handle.StateIndex, Handle.StateSerial);
+	}
+
 private:
 
 	/** Creates a handle for the given running animation on the given widget. */
