UMG: fixes to animation states

- When a widget animation ends, its state object shouldn't be re-used if someone starts the same animation right away. Instead, the state object should be left alone and discarded at the end of the update, and a new state object should be created.
- If a state object never actually started playing its animation, take a shortcut when stopping it so that we don't need to create any animation data for nothing.
- Clarify animation end offset vs end time by having options for both.
- Handle degenerate looping lengths (animation loops that are zero frames in length)

diff --git a/Engine/Source/Runtime/MovieScene/Private/Evaluation/MovieScenePlaybackManager.cpp b/Engine/Source/Runtime/MovieScene/Private/Evaluation/MovieScenePlaybackManager.cpp
--- a/Engine/Source/Runtime/MovieScene/Private/Evaluation/MovieScenePlaybackManager.cpp
+++ b/Engine/Source/Runtime/MovieScene/Private/Evaluation/MovieScenePlaybackManager.cpp
@@ -116,7 +116,8 @@ void FMovieScenePlaybackManager::InternalUpdateToTick(const FFrameNumber NextTic
 	const bool bIsPlayingForward = IsPlayingForward();
 	const FFrameNumber LoopStartTick = bIsPlayingForward ? EffectiveStartTick : LastValidTick;
 	const FFrameNumber LoopLastTick = bIsPlayingForward ? LastValidTick : EffectiveStartTick;
-	const FFrameNumber LoopDurationTicks = EffectiveDurationTicks;
+	// If the start/end offsets make the duration 0, we treat each loop as one frame long.
+	const FFrameNumber LoopDurationTicks = FMath::Max(FFrameNumber(1), EffectiveDurationTicks);
 
 	// Figure out if we crossed the loop-end boundary.
 	if ((bIsPlayingForward && NextTick > LoopLastTick) ||
diff --git a/Engine/Source/Runtime/UMG/Private/Animation/WidgetAnimationState.cpp b/Engine/Source/Runtime/UMG/Private/Animation/WidgetAnimationState.cpp
--- a/Engine/Source/Runtime/UMG/Private/Animation/WidgetAnimationState.cpp
+++ b/Engine/Source/Runtime/UMG/Private/Animation/WidgetAnimationState.cpp
@@ -60,6 +60,7 @@ FWidgetAnimationState::FWidgetAnimationState()
 	bRestoreState = false;
 	bIsBeginningPlay = false;
 	bIsStopping = false;
+	bIsPendingDelete = false;
 }
 
 void FWidgetAnimationState::Initialize(UWidgetAnimation* InAnimation, UUserWidget* InUserWidget, int32 InIndex, int32 InSerialNumber)
@@ -296,6 +297,7 @@ void FWidgetAnimationState::OnStopped()
 	}
 
 	bIsStopping = false;
+	bIsPendingDelete = true;
 
 #if !NO_LOGGING
 	if (UE::UMG::GVarAnimationDetailedLogging)
@@ -429,10 +431,15 @@ void FWidgetAnimationState::Play(const FWidgetAnimationStatePlayParams& PlayPara
 	}
 
 	PlaybackManager.SetEndOffset(FFrameTime(0));
+	if (PlayParams.EndAtTime.IsSet())
+	{
+		const FFrameTime EndFrameTime = DisplayRate.AsFrameTime(PlayParams.EndAtTime.GetValue());
+		PlaybackManager.SetEndOffsetAsTime(EndFrameTime);
+	}
 	if (PlayParams.EndOffset.IsSet())
 	{
 		const FFrameTime EndOffsetTime = DisplayRate.AsFrameTime(PlayParams.EndOffset.GetValue());
-		PlaybackManager.SetEndOffsetAsTime(EndOffsetTime);
+		PlaybackManager.SetEndOffset(PlaybackManager.GetEndOffset() + EndOffsetTime);
 	}
 
 	PlayMode = PlayParams.PlayMode;
@@ -528,18 +535,28 @@ void FWidgetAnimationState::Stop()
 
 	if (TSharedPtr<FSharedPlaybackState> SharedPlaybackState = WeakPlaybackState.Pin())
 	{
-		TSharedPtr<FMovieSceneEntitySystemRunner> Runner = SharedPlaybackState->GetRunner();
-
-		// See WARNING in Tick() method.
-		FWidgetAnimationHandle ThisHandle = GetAnimationHandle();
-		const FMovieSceneContext ReturnToStartContext = PlaybackManager.UpdateAtCurrentTime();
-		Runner->QueueUpdate(
-				ReturnToStartContext,
-				SharedPlaybackState->GetRootInstanceHandle(),
-				FSimpleDelegate::CreateStatic(&FWidgetAnimationState::OnStopped, ThisHandle),
-				ERunnerUpdateFlags::Flush);
-
-		bIsStopping = true;
+		const FRootInstanceHandle& RootInstanceHandle = SharedPlaybackState->GetRootInstanceHandle();
+		UMovieSceneEntitySystemLinker* Linker = SharedPlaybackState->GetLinker();
+		const FSequenceInstance& RootInstance = Linker->GetInstanceRegistry()->GetInstance(RootInstanceHandle);
+		if (RootInstance.HasEverUpdated())
+		{
+			TSharedPtr<FMovieSceneEntitySystemRunner> Runner = SharedPlaybackState->GetRunner();
+
+			// See WARNING in Tick() method.
+			FWidgetAnimationHandle ThisHandle = GetAnimationHandle();
+			const FMovieSceneContext ReturnToStartContext = PlaybackManager.UpdateAtCurrentTime();
+			Runner->QueueUpdate(
+					ReturnToStartContext,
+					SharedPlaybackState->GetRootInstanceHandle(),
+					FSimpleDelegate::CreateStatic(&FWidgetAnimationState::OnStopped, ThisHandle),
+					ERunnerUpdateFlags::Flush);
+
+			bIsStopping = true;
+		}
+		else
+		{
+			OnStopped();
+		}
 	}
 	else
 	{
@@ -663,7 +680,7 @@ void FWidgetAnimationState::TearDown()
 
 bool FWidgetAnimationState::IsValid() const
 {
-	return WeakPlaybackState.IsValid();
+	return WeakPlaybackState.IsValid() && !bIsPendingDelete;
 }
 
 void FWidgetAnimationState::AddReferencedObjects(FReferenceCollector& Collector)
diff --git a/Engine/Source/Runtime/UMG/Private/UserWidget.cpp b/Engine/Source/Runtime/UMG/Private/UserWidget.cpp
--- a/Engine/Source/Runtime/UMG/Private/UserWidget.cpp
+++ b/Engine/Source/Runtime/UMG/Private/UserWidget.cpp
@@ -483,12 +483,12 @@ UWorld* UUserWidget::GetWorld() const
 	return nullptr;
 }
 
-const FWidgetAnimationState* UUserWidget::GetAnimationState(const UWidgetAnimation* InAnimation) const
+const FWidgetAnimationState* UUserWidget::GetAnimationState(const UWidgetAnimation* InAnimation, bool bValidOnly) const
 {
 	const FWidgetAnimationState* FoundState = nullptr;
 	for (auto It = ActiveAnimations.CreateConstIterator(); It; ++It)
 	{
-		if (It->GetAnimation() == InAnimation)
+		if (It->GetAnimation() == InAnimation && (!bValidOnly || It->IsValid()))
 		{
 			FoundState = &(*It);
 			break;
@@ -497,12 +497,12 @@ const FWidgetAnimationState* UUserWidget::GetAnimationState(const UWidgetAnimati
 	return FoundState;
 }
 
-FWidgetAnimationState* UUserWidget::GetAnimationState(const UWidgetAnimation* InAnimation)
+FWidgetAnimationState* UUserWidget::GetAnimationState(const UWidgetAnimation* InAnimation, bool bValidOnly)
 {
 	FWidgetAnimationState* FoundState = nullptr;
 	for (auto It = ActiveAnimations.CreateIterator(); It; ++It)
 	{
-		if (It->GetAnimation() == InAnimation)
+		if (It->GetAnimation() == InAnimation && (!bValidOnly || It->IsValid()))
 		{
 			FoundState = &(*It);
 			break;
@@ -524,7 +524,7 @@ FWidgetAnimationState* UUserWidget::GetOrAddAnimationState(UWidgetAnimation* InA
 		// If this widget is currently offscreen it would very well not be being tracked by the tick manager
 		AnimationTickManager->AddWidget(this);
 
-		FWidgetAnimationState* FoundState = GetAnimationState(InAnimation);
+		FWidgetAnimationState* FoundState = GetAnimationState(InAnimation, false);
 		if (!FoundState)
 		{
 			// Create a new state and initialize it.
@@ -820,7 +820,7 @@ FWidgetAnimationHandle UUserWidget::PlayAnimationTimeRange(UWidgetAnimation* InA
 	{
 		FWidgetAnimationStatePlayParams PlayParams;
 		PlayParams.StartAtTime = StartAtTime;
-		PlayParams.EndOffset = EndAtTime;
+		PlayParams.EndAtTime = EndAtTime;
 		PlayParams.NumLoopsToPlay = NumberOfLoops;
 		PlayParams.PlayMode = PlayMode;
 		PlayParams.PlaybackSpeed = PlaybackSpeed;
diff --git a/Engine/Source/Runtime/UMG/Public/Animation/WidgetAnimationState.h b/Engine/Source/Runtime/UMG/Public/Animation/WidgetAnimationState.h
--- a/Engine/Source/Runtime/UMG/Public/Animation/WidgetAnimationState.h
+++ b/Engine/Source/Runtime/UMG/Public/Animation/WidgetAnimationState.h
@@ -22,8 +22,10 @@ DECLARE_MULTICAST_DELEGATE_OneParam(FOnWidgetAnimationEvent, FWidgetAnimationSta
 
 struct FWidgetAnimationStatePlayParams
 {
-	/** The time to start playing at. */
+	/** The start time for the animation. */
 	double StartAtTime = 0;
+	/** The end time for the animation. */
+	TOptional<double> EndAtTime;
 	/** An optional start offset to apply to the animation. */
 	TOptional<double> StartOffset;
 	/** An optional end offset to apply to the animation. */
@@ -69,6 +71,9 @@ public:
 	/** Returns whether this state is currently stopping. */
 	bool IsStopping() const { return bIsStopping; }
 
+	/** Returns whether this state is currently pending destruction. */
+	bool IsPendingDelete() const { return bIsPendingDelete; }
+
 	/** Allows registering a callback for when the animation has finished playing. */
 	FOnWidgetAnimationEvent& GetOnWidgetAnimationFinished() { return OnWidgetAnimationFinishedEvent; }
 
@@ -191,6 +196,9 @@ private:
 	/** Whether we are in the process of stopping. */
 	bool bIsStopping : 1;
 
+	/** Whether we have played, stopped, and now waiting to be deleted. */
+	bool bIsPendingDelete : 1;
+
 	/** Callback for when the animation has finished playing. */
 	FOnWidgetAnimationEvent OnWidgetAnimationFinishedEvent;
 };
diff --git a/Engine/Source/Runtime/UMG/Public/Blueprint/UserWidget.h b/Engine/Source/Runtime/UMG/Public/Blueprint/UserWidget.h
--- a/Engine/Source/Runtime/UMG/Public/Blueprint/UserWidget.h
+++ b/Engine/Source/Runtime/UMG/Public/Blueprint/UserWidget.h
@@ -1635,8 +1635,8 @@ protected:
 
 	UMG_API void RemoveObsoleteBindings(const TArray<FName>& NamedSlots);
 
-	UMG_API const FWidgetAnimationState* GetAnimationState(const UWidgetAnimation* InAnimation) const;
-	UMG_API FWidgetAnimationState* GetAnimationState(const UWidgetAnimation* InAnimation);
+	UMG_API const FWidgetAnimationState* GetAnimationState(const UWidgetAnimation* InAnimation, bool bValidOnly = true) const;
+	UMG_API FWidgetAnimationState* GetAnimationState(const UWidgetAnimation* InAnimation, bool bValidOnly = true);
 	UMG_API FWidgetAnimationState* GetOrAddAnimationState(UWidgetAnimation* InAnimation);
 
 	UMG_API void ExecuteQueuedAnimationTransitions();
