UMG: fix possible memory corruption after a widget animation ends

When a widget animation ends, an "on finished" callback may start another animation. This might reallocate the ActiveAnimations array inside UUserWidget and invalidate the current animation state that is running this callback. Any further action on that object can result in accessing corrupted memory.

diff --git a/Engine/Source/Runtime/UMG/Private/Animation/WidgetAnimationState.cpp b/Engine/Source/Runtime/UMG/Private/Animation/WidgetAnimationState.cpp
--- a/Engine/Source/Runtime/UMG/Private/Animation/WidgetAnimationState.cpp
+++ b/Engine/Source/Runtime/UMG/Private/Animation/WidgetAnimationState.cpp
@@ -311,6 +311,8 @@ void FWidgetAnimationState::OnStopped()
 	}
 #endif
 
+	FWidgetAnimationHandle ThisHandle = GetAnimationHandle();
+
 	UUserWidget* UserWidget = WeakUserWidget.Get();
 	if (UserWidget)
 	{
@@ -322,11 +324,18 @@ void FWidgetAnimationState::OnStopped()
 		UserWidget->BroadcastAnimationFinishedPlaying(*this);
 	}
 
-	OnWidgetAnimationFinishedEvent.Broadcast(*this);
-
-	if (LegacyPlayer)
+	// WARNING
+	// Don't access any more member fields after this... the call to BroadcastAnimationFinishedPlaying
+	// may have started new animations, which might have reallocated our parent widget's ActiveAnimations
+	// array, and rendering us invalid.
+	if (FWidgetAnimationState* NewThis = ThisHandle.GetAnimationState())
 	{
-		LegacyPlayer->BroadcastSequenceFinishedPlaying();
+		NewThis->OnWidgetAnimationFinishedEvent.Broadcast(*NewThis);
+
+		if (NewThis->LegacyPlayer)
+		{
+			NewThis->LegacyPlayer->BroadcastSequenceFinishedPlaying();
+		}
 	}
 }
 
