Replace TObjectPtr by TWeakObjectPtr to ensure UWidgetInteractionComponent does not prevent garbage collecting external widget components

diff --git a/Engine/Source/Runtime/UMG/Private/Components/WidgetInteractionComponent.cpp b/Engine/Source/Runtime/UMG/Private/Components/WidgetInteractionComponent.cpp
--- a/Engine/Source/Runtime/UMG/Private/Components/WidgetInteractionComponent.cpp
+++ b/Engine/Source/Runtime/UMG/Private/Components/WidgetInteractionComponent.cpp
@@ -313,30 +313,31 @@ FWidgetPath UWidgetInteractionComponent::DetermineWidgetUnderPointer()
 	bIsHoveredWidgetFocusable = false;
 	bIsHoveredWidgetHitTestVisible = false;
 
-	UWidgetComponent* OldHoveredWidget = HoveredWidgetComponent;
-
-	HoveredWidgetComponent = nullptr;
+	UWidgetComponent* OldHoveredWidgetComponent = GetHoveredWidgetComponent();
+	WeakHoveredWidgetComponent.Reset();
 
 	FWidgetTraceResult TraceResult = PerformTrace();
 	LastHitResult = TraceResult.HitResult;
-	HoveredWidgetComponent = TraceResult.HitWidgetComponent;
+	WeakHoveredWidgetComponent = TraceResult.HitWidgetComponent;
 	LastLocalHitLocation = LocalHitLocation;
 	LocalHitLocation = TraceResult.bWasHit
 		? TraceResult.LocalHitLocation
 		: LastLocalHitLocation;
 	WidgetPathUnderPointer = TraceResult.HitWidgetPath;
 
+	UWidgetComponent* NewHoveredWidgetComponent = GetHoveredWidgetComponent();
+
 #if ENABLE_DRAW_DEBUG
 	if ( bShowDebug )
 	{
-		if ( HoveredWidgetComponent )
+		if ( NewHoveredWidgetComponent )
 		{
 			UKismetSystemLibrary::DrawDebugSphere(this, LastHitResult.ImpactPoint, 2.5f, 12, DebugColor, 0, DebugSphereLineThickness);
 		}
 
 		if ( InteractionSource == EWidgetInteractionSource::World || InteractionSource == EWidgetInteractionSource::Custom )
 		{
-			if ( HoveredWidgetComponent )
+			if ( NewHoveredWidgetComponent )
 			{
 				UKismetSystemLibrary::DrawDebugLine(this, LastHitResult.TraceStart, LastHitResult.ImpactPoint, DebugColor, 0, DebugLineThickness);
 			}
@@ -348,9 +349,9 @@ FWidgetPath UWidgetInteractionComponent::DetermineWidgetUnderPointer()
 	}
 #endif // ENABLE_DRAW_DEBUG
 
-	if ( HoveredWidgetComponent )
+	if ( NewHoveredWidgetComponent )
 	{
-		HoveredWidgetComponent->RequestRedraw();
+		NewHoveredWidgetComponent->RequestRedraw();
 	}
 
 	if ( WidgetPathUnderPointer.IsValid() )
@@ -379,14 +380,14 @@ FWidgetPath UWidgetInteractionComponent::DetermineWidgetUnderPointer()
 		}
 	}
 
-	if ( HoveredWidgetComponent != OldHoveredWidget )
+	if ( NewHoveredWidgetComponent != OldHoveredWidgetComponent )
 	{
-		if ( OldHoveredWidget )
+		if ( OldHoveredWidgetComponent )
 		{
-			OldHoveredWidget->RequestRedraw();
+			OldHoveredWidgetComponent->RequestRedraw();
 		}
 
-		OnHoveredWidgetChanged.Broadcast( HoveredWidgetComponent, OldHoveredWidget );
+		OnHoveredWidgetChanged.Broadcast( NewHoveredWidgetComponent, OldHoveredWidgetComponent );
 	}
 
 	return WidgetPathUnderPointer;
@@ -419,7 +420,7 @@ void UWidgetInteractionComponent::SimulatePointerMovement()
 	
 	if (WidgetPathUnderFinger.IsValid())
 	{
-		check(HoveredWidgetComponent);
+		check(WeakHoveredWidgetComponent.IsValid());
 		LastWidgetPath = WidgetPathUnderFinger;
 		
 		FSlateApplication::Get().RoutePointerMoveEvent(WidgetPathUnderFinger, PointerEvent, false);
@@ -708,7 +709,7 @@ void UWidgetInteractionComponent::ScrollWheel(float ScrollDelta)
 
 UWidgetComponent* UWidgetInteractionComponent::GetHoveredWidgetComponent() const
 {
-	return HoveredWidgetComponent;
+	return WeakHoveredWidgetComponent.Get();
 }
 
 bool UWidgetInteractionComponent::IsOverInteractableWidget() const
diff --git a/Engine/Source/Runtime/UMG/Public/Components/WidgetInteractionComponent.h b/Engine/Source/Runtime/UMG/Public/Components/WidgetInteractionComponent.h
--- a/Engine/Source/Runtime/UMG/Public/Components/WidgetInteractionComponent.h
+++ b/Engine/Source/Runtime/UMG/Public/Components/WidgetInteractionComponent.h
@@ -335,9 +335,14 @@ protected:
 	UPROPERTY(Transient)
 	FVector2D LastLocalHitLocation;
 
+	/** DEPRECATED - Use WeakHoveredWidgetComponent instead */
+	UE_DEPRECATED(5.6, "Property deprecated, please use WeakHoveredWidgetComponent instead.")
+	UPROPERTY(Transient)
+	TObjectPtr<UWidgetComponent> HoveredWidgetComponent = nullptr;
+
 	/** The widget component we're currently hovering over. */
 	UPROPERTY(Transient)
-	TObjectPtr<UWidgetComponent> HoveredWidgetComponent;
+	TWeakObjectPtr<UWidgetComponent> WeakHoveredWidgetComponent;
 
 	/** The last hit result we used. */
 	UPROPERTY(Transient)
