Sequencer: potential fix for a memory corruption of UMG animation states

diff --git a/Engine/Source/Runtime/UMG/Private/UserWidget.cpp b/Engine/Source/Runtime/UMG/Private/UserWidget.cpp
--- a/Engine/Source/Runtime/UMG/Private/UserWidget.cpp
+++ b/Engine/Source/Runtime/UMG/Private/UserWidget.cpp
@@ -1006,20 +1006,28 @@ bool UUserWidget::IsAnimationPlayingForward(const UWidgetAnimation* InAnimation)
 
 void UUserWidget::BroadcastAnimationFinishedPlaying(FWidgetAnimationState& State)
 {
+	const FName UserTag = State.GetUserTag();
+	const UWidgetAnimation* Animation = State.GetAnimation();
+	const EMovieScenePlayerStatus::Type PlayerStatus = State.GetPlaybackStatus();
+	UUMGSequencePlayer* LegacyPlayer = State.GetLegacyPlayer();
+
 	OnAnimationFinishedPlayingEvent.Broadcast(State);
 
-	OnAnimationFinished(State.GetAnimation());
+	// WARNING: do not use State after this point. OnAnimationFinishedPlayingEvent may have triggered new animations
+	// and reallocated the array of animation states.
 
-	BroadcastAnimationStateChange(State, EWidgetAnimationEvent::Finished);
+	OnAnimationFinished(Animation);
 
-	if (UUMGSequencePlayer* LegacyPlayer = State.GetLegacyPlayer())
+	BroadcastAnimationStateChange(Animation, UserTag, EWidgetAnimationEvent::Finished);
+
+	if (LegacyPlayer)
 	{
 		PRAGMA_DISABLE_DEPRECATION_WARNINGS
 		OnAnimationFinishedPlaying(*LegacyPlayer);
 		PRAGMA_ENABLE_DEPRECATION_WARNINGS
 	}
 
-	if (State.GetPlaybackStatus() == EMovieScenePlayerStatus::Stopped && AnimationTickManager)
+	if (PlayerStatus == EMovieScenePlayerStatus::Stopped && AnimationTickManager)
 	{
 		AnimationTickManager->AddLatentAction(
 				FMovieSceneSequenceLatentActionDelegate::CreateUObject(this, &UUserWidget::ClearStoppedAnimationStates));
