ListView: update invalidation to support changing entry spacing at runtime

diff --git a/Engine/Plugins/Runtime/CommonUI/Source/CommonUI/Private/CommonListView.cpp b/Engine/Plugins/Runtime/CommonUI/Source/CommonUI/Private/CommonListView.cpp
--- a/Engine/Plugins/Runtime/CommonUI/Source/CommonUI/Private/CommonListView.cpp
+++ b/Engine/Plugins/Runtime/CommonUI/Source/CommonUI/Private/CommonListView.cpp
@@ -19,8 +19,8 @@ UCommonListView::UCommonListView(const FObjectInitializer& ObjectInitializer)
 
 void UCommonListView::SetEntrySpacing(float InEntrySpacing)
 {
-	InitHorizontalEntrySpacing(InEntrySpacing);
-	InitVerticalEntrySpacing(InEntrySpacing);
+	SetHorizontalEntrySpacing(InEntrySpacing);
+	SetVerticalEntrySpacing(InEntrySpacing);
 }
 
 #if WITH_EDITOR
diff --git a/Engine/Source/Runtime/Slate/Public/Framework/SlateDelegates.h b/Engine/Source/Runtime/Slate/Public/Framework/SlateDelegates.h
--- a/Engine/Source/Runtime/Slate/Public/Framework/SlateDelegates.h
+++ b/Engine/Source/Runtime/Slate/Public/Framework/SlateDelegates.h
@@ -162,6 +162,12 @@ public:
 		/** param: The owning widget */
 		const TSharedRef< class STableViewBase >& );
 
+	/** A delegate to be invoked to update an existing item (or its containing widget) when the owning view is invalidated. */
+	DECLARE_DELEGATE_OneParam(
+		FOnRefreshRow,
+		/** param: An item to update */
+		ArgumentType);
+
 	/** A delegate to be invoked when an item has come into view after it was requested to come into view. */
 	DECLARE_DELEGATE_TwoParams (
 		FOnItemScrolledIntoView,
diff --git a/Engine/Source/Runtime/Slate/Public/Widgets/Views/SListView.h b/Engine/Source/Runtime/Slate/Public/Widgets/Views/SListView.h
--- a/Engine/Source/Runtime/Slate/Public/Widgets/Views/SListView.h
+++ b/Engine/Source/Runtime/Slate/Public/Widgets/Views/SListView.h
@@ -72,6 +72,7 @@ public:
 	using TItemSet          = TSet< TObjectPtrWrapTypeOf<ItemType>, typename TListTypeTraits< TObjectPtrWrapTypeOf<ItemType> >::SetKeyFuncs >;
 
 	using FOnGenerateRow            = typename TSlateDelegates< ItemType >::FOnGenerateRow;
+	using FOnRefreshRow				= typename TSlateDelegates< ItemType >::FOnRefreshRow;
 	using FOnItemScrolledIntoView   = typename TSlateDelegates< ItemType >::FOnItemScrolledIntoView;
 	using FOnSelectionChanged       = typename TSlateDelegates< NullableItemType >::FOnSelectionChanged;
 	using FIsSelectableOrNavigable	= typename TSlateDelegates< ItemType >::FIsSelectableOrNavigable;
@@ -89,6 +90,7 @@ public:
 		: _ListViewStyle(&FAppStyle::Get().GetWidgetStyle<FTableViewStyle>("ListView"))
 		, _OnGenerateRow()
 		, _OnGeneratePinnedRow()
+		, _OnRefreshRow()
 		, _OnEntryInitialized()
 		, _OnRowReleased()
 		, _MaxPinnedItems(6)
@@ -129,6 +131,8 @@ public:
 
 		SLATE_EVENT( FOnGenerateRow, OnGeneratePinnedRow )
 		
+		SLATE_EVENT( FOnRefreshRow, OnRefreshRow )
+
 		SLATE_EVENT( FOnEntryInitialized, OnEntryInitialized )
 
 		SLATE_EVENT( FOnWidgetToBeRemoved, OnRowReleased )
@@ -224,6 +228,7 @@ public:
 
 		this->OnGenerateRow = InArgs._OnGenerateRow;
 		this->OnGeneratePinnedRow = InArgs._OnGeneratePinnedRow;
+		this->OnRefreshRow = InArgs._OnRefreshRow;
 		this->OnEntryInitialized = InArgs._OnEntryInitialized;
 		this->OnRowReleased = InArgs._OnRowReleased;
 		this->OnItemScrolledIntoView = InArgs._OnItemScrolledIntoView;
@@ -1522,6 +1527,13 @@ public:
 			// Make a new widget for it.
 			WidgetForItem = this->GenerateNewWidget(CurItem);
 		}
+		else
+		{
+			if (OnRefreshRow.IsBound())
+			{
+				OnRefreshRow.Execute(CurItem);
+			}
+		}
 
 		// It is useful to know the item's index that the widget was generated from.
 		// Helps with even/odd coloring
@@ -1746,6 +1758,11 @@ public:
 		SetItemsSource(TUniquePtr<UE::Slate::ItemsSource::IItemsSource<ItemType>>());
 	}
 
+	void BindToRefreshRow(const FOnRefreshRow& Binding)
+	{
+		OnRefreshRow = Binding;
+	}
+
 public:
 
 	UE_DEPRECATED(5.2, "SetListItemsSource is deprecated. Please use the correct SetItemsSource implementation.")
@@ -2595,6 +2612,9 @@ protected:
 	/** Delegate to be invoked when the list needs to generate a new pinned widget from a data item. */
 	FOnGenerateRow OnGeneratePinnedRow;
 
+	/** Delegate to be invoked when the list needs to update an existing widget for an item. */
+	FOnRefreshRow OnRefreshRow;
+
 	/** Assign this to get more diagnostics from the list view. */
 	FOnItemToString_Debug OnItemToString_Debug;
 
diff --git a/Engine/Source/Runtime/Slate/Public/Widgets/Views/STileView.h b/Engine/Source/Runtime/Slate/Public/Widgets/Views/STileView.h
--- a/Engine/Source/Runtime/Slate/Public/Widgets/Views/STileView.h
+++ b/Engine/Source/Runtime/Slate/Public/Widgets/Views/STileView.h
@@ -26,6 +26,7 @@ public:
 	typedef typename TListTypeTraits< ItemType >::NullableType NullableItemType;
 
 	typedef typename TSlateDelegates< ItemType >::FOnGenerateRow FOnGenerateRow;
+	typedef typename TSlateDelegates< ItemType >::FOnRefreshRow FOnRefreshRow;
 	typedef typename TSlateDelegates< ItemType >::FOnItemScrolledIntoView FOnItemScrolledIntoView;
 	typedef typename TSlateDelegates< ItemType >::FOnMouseButtonClick FOnMouseButtonClick;
 	typedef typename TSlateDelegates< ItemType >::FOnMouseButtonDoubleClick FOnMouseButtonDoubleClick;
@@ -40,6 +41,7 @@ public:
 
 	SLATE_BEGIN_ARGS(STileView<ItemType>)
 		: _OnGenerateTile()
+		, _OnRefreshTile()
 		, _OnTileReleased()
 		, _ItemHeight(128)
 		, _ItemWidth(128)
@@ -74,6 +76,8 @@ public:
 
 		SLATE_EVENT( FOnGenerateRow, OnGenerateTile )
 
+		SLATE_EVENT( FOnRefreshRow, OnRefreshTile )
+
 		SLATE_EVENT( FOnWidgetToBeRemoved, OnTileReleased )
 
 		SLATE_EVENT( FOnTableViewScrolled, OnTileViewScrolled )
@@ -153,6 +157,7 @@ public:
 		this->Clipping = InArgs._Clipping;
 
 		this->OnGenerateRow = InArgs._OnGenerateTile;
+		this->OnRefreshRow = InArgs._OnRefreshTile;
 		this->OnRowReleased = InArgs._OnTileReleased;
 		this->OnItemScrolledIntoView = InArgs._OnItemScrolledIntoView;
 		
diff --git a/Engine/Source/Runtime/Slate/Public/Widgets/Views/STreeView.h b/Engine/Source/Runtime/Slate/Public/Widgets/Views/STreeView.h
--- a/Engine/Source/Runtime/Slate/Public/Widgets/Views/STreeView.h
+++ b/Engine/Source/Runtime/Slate/Public/Widgets/Views/STreeView.h
@@ -99,6 +99,7 @@ public:
 
 	using FOnGetChildren            = typename TSlateDelegates< ItemType >::FOnGetChildren;
 	using FOnGenerateRow            = typename TSlateDelegates< ItemType >::FOnGenerateRow;
+	using FOnRefreshRow				= typename TSlateDelegates< ItemType >::FOnRefreshRow;
 	using FOnSetExpansionRecursive  = typename TSlateDelegates< ItemType >::FOnSetExpansionRecursive;
 	using FOnItemScrolledIntoView   = typename TSlateDelegates< ItemType >::FOnItemScrolledIntoView;
 	using FOnSelectionChanged       = typename TSlateDelegates< NullableItemType >::FOnSelectionChanged;
@@ -117,6 +118,7 @@ public:
 		: _TreeViewStyle(&FAppStyle::Get().GetWidgetStyle<FTableViewStyle>("TreeView"))
 		, _OnGenerateRow()
 		, _OnGeneratePinnedRow()
+		, _OnRefreshRow()
 		, _OnGetChildren()
 		, _OnSetExpansionRecursive()
 		, _MaxPinnedItems(6) // Having more than the max amount of items leads to the extra items in the middle being collapsed into ellipses, and the last item is fully shown
@@ -155,6 +157,8 @@ public:
 
 		SLATE_EVENT( FOnGenerateRow, OnGeneratePinnedRow )
 
+		SLATE_EVENT( FOnRefreshRow, OnRefreshRow )
+
 		SLATE_EVENT( FOnWidgetToBeRemoved, OnRowReleased )
 
 		SLATE_EVENT( FOnTableViewScrolled, OnTreeViewScrolled )
@@ -246,6 +250,7 @@ public:
 
 		this->OnGenerateRow = InArgs._OnGenerateRow;
 		this->OnGeneratePinnedRow = InArgs._OnGeneratePinnedRow;
+		this->OnRefreshRow = InArgs._OnRefreshRow;
 		this->OnRowReleased = InArgs._OnRowReleased;
 		this->OnItemScrolledIntoView = InArgs._OnItemScrolledIntoView;
 		this->OnGetChildren = InArgs._OnGetChildren;
diff --git a/Engine/Source/Runtime/UMG/Private/Components/ListView.cpp b/Engine/Source/Runtime/UMG/Private/Components/ListView.cpp
--- a/Engine/Source/Runtime/UMG/Private/Components/ListView.cpp
+++ b/Engine/Source/Runtime/UMG/Private/Components/ListView.cpp
@@ -373,16 +373,32 @@ bool UListView::OnIsSelectableOrNavigableInternal(UObject* FirstSelectedItem)
 	return BP_OnIsItemSelectableOrNavigable.IsBound() ? BP_OnIsItemSelectableOrNavigable.Execute(FirstSelectedItem) : ITypedUMGListView::OnIsSelectableOrNavigableInternal(FirstSelectedItem);
 }
 
+void UListView::SetVerticalEntrySpacing(float InVerticalEntrySpacing)
+{
+	if (VerticalEntrySpacing != InVerticalEntrySpacing)
+	{
+		VerticalEntrySpacing = InVerticalEntrySpacing;
+		RequestRefresh();
+	}
+}
+
+void UListView::SetHorizontalEntrySpacing(float InHorizontalEntrySpacing)
+{
+	if (HorizontalEntrySpacing != InHorizontalEntrySpacing)
+	{
+		HorizontalEntrySpacing = InHorizontalEntrySpacing;
+		RequestRefresh();
+	}
+}
+
 void UListView::InitHorizontalEntrySpacing(float InHorizontalEntrySpacing)
 {
-	ensureMsgf(!MyListView.IsValid(), TEXT("The SWidget is already created."));
-	HorizontalEntrySpacing = InHorizontalEntrySpacing;
+	SetHorizontalEntrySpacing(InHorizontalEntrySpacing);
 }
 
 void UListView::InitVerticalEntrySpacing(float InVerticalEntrySpacing)
 {
-	ensureMsgf(!MyListView.IsValid(), TEXT("The SWidget is already created."));
-	VerticalEntrySpacing = InVerticalEntrySpacing;
+	SetVerticalEntrySpacing(InVerticalEntrySpacing);
 }
 
 void UListView::PostLoad()
diff --git a/Engine/Source/Runtime/UMG/Public/Components/ListView.h b/Engine/Source/Runtime/UMG/Public/Components/ListView.h
--- a/Engine/Source/Runtime/UMG/Public/Components/ListView.h
+++ b/Engine/Source/Runtime/UMG/Public/Components/ListView.h
@@ -240,10 +240,10 @@ protected:
 #endif  // WITH_EDITOR
 
 private:
-	UPROPERTY(EditAnywhere, BlueprintReadOnly, Getter, Category = ListView, meta = (ClampMin = 0, AllowPrivateAccess = "true"))
+	UPROPERTY(EditAnywhere, BlueprintReadWrite, Getter, Setter, Category = ListView, meta = (ClampMin = 0, AllowPrivateAccess = "true"))
 	float HorizontalEntrySpacing = 0.f;
 
-	UPROPERTY(EditAnywhere, BlueprintReadOnly, Getter, Category = ListView, meta = (ClampMin = 0, AllowPrivateAccess = "true"))
+	UPROPERTY(EditAnywhere, BlueprintReadWrite, Getter, Setter, Category = ListView, meta = (ClampMin = 0, AllowPrivateAccess = "true"))
 	float VerticalEntrySpacing = 0.f;
 	
 	UPROPERTY(EditAnywhere, BlueprintReadWrite, Getter, Setter, Category = ListView, meta = (AllowPrivateAccess = "true", DisplayAfter="ScrollBarStyle"))
@@ -262,8 +262,17 @@ public:
 	UFUNCTION(Category = ListView)
 	FMargin GetScrollBarPadding() const { return ScrollBarPadding; }
 
+	/** Set the vertical spacing between entries. */
+	UMG_API void SetVerticalEntrySpacing(float InVerticalEntrySpacing);
+
+	/** Set the horizontal spacing between entries. */
+	UMG_API void SetHorizontalEntrySpacing(float InHorizontalEntrySpacing);
+
 protected:
+	UE_DEPRECATED(5.6, "InitHorizontalEntrySpacing is deprecated, please use SetHorizontalEntrySpacing instead.")
 	UMG_API void InitHorizontalEntrySpacing(float InHorizontalEntrySpacing);
+
+	UE_DEPRECATED(5.6, "InitVerticalEntrySpacing is deprecated, please use SetVerticalEntrySpacing instead.")
 	UMG_API void InitVerticalEntrySpacing(float InVerticalEntrySpacing);
 
 private:
diff --git a/Engine/Source/Runtime/UMG/Public/Components/ListViewBase.h b/Engine/Source/Runtime/UMG/Public/Components/ListViewBase.h
--- a/Engine/Source/Runtime/UMG/Public/Components/ListViewBase.h
+++ b/Engine/Source/Runtime/UMG/Public/Components/ListViewBase.h
@@ -252,7 +252,7 @@ protected:
 		const FListViewConstructArgs& Args = FListViewConstructArgs())
 	{
 		static_assert(TIsDerivedFrom<ListViewT<ItemType>, SListView<ItemType>>::IsDerived, "ConstructListView can only construct instances of SListView classes");
-		return SNew(ListViewT<ItemType>)
+		TSharedRef<ListViewT<ItemType>> ListView = SNew(ListViewT<ItemType>)
 			.HandleGamepadEvents(true)
 			.ListItemsSource(&ListItems)
 			.IsFocusable(Args.bAllowFocus)
@@ -276,6 +276,9 @@ protected:
 			.OnFinishedScrolling_UObject(Implementer, &UListViewBaseT::HandleFinishedScrolling)
 			.OnMouseButtonClick_UObject(Implementer, &UListViewBaseT::HandleItemClicked)
 			.OnMouseButtonDoubleClick_UObject(Implementer, &UListViewBaseT::HandleItemDoubleClicked);
+		ListView->BindToRefreshRow(TSlateDelegates< ItemType >::FOnRefreshRow::CreateUObject(Implementer, &UListViewBaseT::HandleRefreshRow));
+
+		return ListView;
 	}
 
 	struct FTileViewConstructArgs : public FListViewConstructArgs
@@ -294,7 +297,7 @@ protected:
 		const FTileViewConstructArgs& Args = FTileViewConstructArgs())
 	{
 		static_assert(TIsDerivedFrom<TileViewT<ItemType>, STileView<ItemType>>::IsDerived, "ConstructTileView can only construct instances of STileView classes");
-		return SNew(TileViewT<ItemType>)
+		TSharedRef<TileViewT<ItemType>> TileView = SNew(TileViewT<ItemType>)
 			.HandleGamepadEvents(true)
 			.ListItemsSource(&ListItems)
 			.IsFocusable(Args.bAllowFocus)
@@ -317,6 +320,9 @@ protected:
 			.OnFinishedScrolling_UObject(Implementer, &UListViewBaseT::HandleFinishedScrolling)
 			.OnMouseButtonClick_UObject(Implementer, &UListViewBaseT::HandleItemClicked)
 			.OnMouseButtonDoubleClick_UObject(Implementer, &UListViewBaseT::HandleItemDoubleClicked);
+		TileView->BindToRefreshRow(TSlateDelegates< ItemType >::FOnRefreshRow::CreateUObject(Implementer, &UListViewBaseT::HandleRefreshRow));
+
+		return TileView;
 	}
 
 	struct FTreeViewConstructArgs
@@ -335,7 +341,7 @@ protected:
 		const FTreeViewConstructArgs& Args = FTreeViewConstructArgs())
 	{
 		static_assert(TIsDerivedFrom<TreeViewT<ItemType>, STreeView<ItemType>>::IsDerived, "ConstructTreeView can only construct instances of STreeView classes");
-		return SNew(TreeViewT<ItemType>)
+		TSharedRef<TreeViewT<ItemType>> TreeView = SNew(TreeViewT<ItemType>)
 			.HandleGamepadEvents(true)
 			.TreeItemsSource(&ListItems)
 			.ClearSelectionOnClick(Args.bClearSelectionOnClick)
@@ -355,6 +361,9 @@ protected:
 			.OnMouseButtonDoubleClick_UObject(Implementer, &UListViewBaseT::HandleItemDoubleClicked)
 			.OnGetChildren_UObject(Implementer, &UListViewBaseT::HandleGetChildren)
 			.OnExpansionChanged_UObject(Implementer, &UListViewBaseT::HandleExpansionChanged);
+		TreeView->BindToRefreshRow(TSlateDelegates< ItemType >::FOnRefreshRow::CreateUObject(Implementer, &UListViewBaseT::HandleRefreshRow));
+
+		return TreeView;
 	}
 
 protected:
@@ -410,6 +419,16 @@ private:
 		return StaticCastSharedPtr<SObjectTableRow<ItemType>>(CachedWidget).ToSharedRef();
 	}
 
+	void HandleRefreshRow(ItemType Item)
+	{
+		UUserWidget* RowWidget = GetEntryWidgetFromItem(Item);
+		if (RowWidget)
+		{
+			const FMargin DefaultPadding = RowWidget->GetClass()->GetDefaultObject<UUserWidget>()->GetPadding();
+			RowWidget->SetPadding(DefaultPadding + GetDesiredEntryPadding(Item));
+		}
+	}
+
 	void HandleItemClicked(ItemType Item)
 	{
 		OnItemClickedInternal(Item);
