Adding support for wrapping a UMG UI Component owner with a Slate Widget

diff --git a/Engine/Source/Editor/UMGEditor/Private/UIComponentWidgetBlueprintExtension.cpp b/Engine/Source/Editor/UMGEditor/Private/UIComponentWidgetBlueprintExtension.cpp
--- a/Engine/Source/Editor/UMGEditor/Private/UIComponentWidgetBlueprintExtension.cpp
+++ b/Engine/Source/Editor/UMGEditor/Private/UIComponentWidgetBlueprintExtension.cpp
@@ -107,6 +107,7 @@ UUIComponentUserWidgetExtension* UUIComponentWidgetBlueprintExtension::GetOrCrea
 	{
 		UserWidgetExtension = PreviewWidget->AddExtension<UUIComponentUserWidgetExtension>();
 		UserWidgetExtension->InitializeContainer(DuplicateContainer(UserWidgetExtension));
+		UserWidgetExtension->Initialize();
 	}
 	else
 	{
diff --git a/Engine/Source/Runtime/UMG/Private/Components/Widget.cpp b/Engine/Source/Runtime/UMG/Private/Components/Widget.cpp
--- a/Engine/Source/Runtime/UMG/Private/Components/Widget.cpp
+++ b/Engine/Source/Runtime/UMG/Private/Components/Widget.cpp
@@ -39,6 +39,9 @@
 #include "Components/NamedSlotInterface.h"
 #include "ProfilingDebugging/AssetMetadataTrace.h"
 #include "HAL/LowLevelMemStats.h"
+#include "Extensions/UIComponentContainer.h"
+#include "Extensions/UIComponent.h"
+#include "Extensions/UIComponentUserWidgetExtension.h"
 #include "Types/InvisibleToWidgetReflectorMetaData.h"
 
 #include UE_INLINE_GENERATED_CPP_BY_NAME(Widget)
@@ -210,6 +213,7 @@ PRAGMA_DISABLE_DEPRECATION_WARNINGS
 
 	bShouldBroadcastState = true;
 	bWidgetStateInitialized = false;
+	bWrappedByComponent = false;
 
 PRAGMA_ENABLE_DEPRECATION_WARNINGS
 
@@ -1017,6 +1021,34 @@ TSharedRef<SWidget> UWidget::TakeWidget_Private(ConstructMethodType ConstructMet
 			PublicWidget = SafeGCWidget;
 		}
 	}
+	if(bWrappedByComponent)
+	{
+		// Now, if the Widget instance is wrapped by a component, add it to the hierarchy
+		const UWidgetTree* WidgetTree = Cast<UWidgetTree>(GetOuter());
+		if (const UUserWidget* UserWidget = WidgetTree ? Cast<UUserWidget>(WidgetTree->GetOuter()) : nullptr)
+		{
+			if (UUIComponentUserWidgetExtension* Extension = UserWidget->GetExtension<UUIComponentUserWidgetExtension>())
+			{
+				for (UUIComponent* Component : Extension->GetComponentsFor(this))
+				{
+					if (bNewlyCreated)
+					{
+						TSharedPtr<SWidget> SafeComponentWidget = Component->RebuildWidgetWithContent(PublicWidget.ToSharedRef());
+						// If the Widget is the same of null, we do nothing 
+						if (SafeComponentWidget != PublicWidget)
+						{
+							ComponentWrapperWidget = SafeComponentWidget;
+							PublicWidget = SafeComponentWidget;
+						}
+					}
+					else if (ComponentWrapperWidget.IsValid())
+					{
+						PublicWidget = ComponentWrapperWidget.Pin();
+					}
+				}			
+			}
+		}
+	}
 
 #if WITH_EDITOR
 	if (IsDesignTime())
@@ -1067,6 +1099,11 @@ TSharedPtr<SWidget> UWidget::GetCachedWidget() const
 	}
 #endif
 
+	if (ComponentWrapperWidget.IsValid())
+	{
+		return ComponentWrapperWidget.Pin();
+	}
+	
 	if (MyGCWidget.IsValid())
 	{
 		return MyGCWidget.Pin();
diff --git a/Engine/Source/Runtime/UMG/Private/Extensions/UIComponent.cpp b/Engine/Source/Runtime/UMG/Private/Extensions/UIComponent.cpp
--- a/Engine/Source/Runtime/UMG/Private/Extensions/UIComponent.cpp
+++ b/Engine/Source/Runtime/UMG/Private/Extensions/UIComponent.cpp
@@ -12,6 +12,7 @@ void UUIComponent::Initialize(UWidget* Target)
 	ensure(Target);
 	ensure(Owner.IsExplicitlyNull());
 	Owner = Target;
+	Owner->bWrappedByComponent = true;
 	OnInitialize();
 }
 
@@ -56,6 +57,11 @@ void UUIComponent::OnDestruct()
 
 }
 
+TSharedRef<SWidget> UUIComponent::RebuildWidgetWithContent(TSharedRef<SWidget> Content)
+{
+	return Content;
+}
+
 void UUIComponent::FFieldNotificationClassDescriptor::ForEachField(const UClass* Class, TFunctionRef<bool(::UE::FieldNotification::FFieldId FieldId)> Callback) const
 {
 	if (const UBlueprintGeneratedClass* BPClass = Cast<const UBlueprintGeneratedClass>(Class))
diff --git a/Engine/Source/Runtime/UMG/Private/Extensions/UIComponentUserWidgetExtension.cpp b/Engine/Source/Runtime/UMG/Private/Extensions/UIComponentUserWidgetExtension.cpp
--- a/Engine/Source/Runtime/UMG/Private/Extensions/UIComponentUserWidgetExtension.cpp
+++ b/Engine/Source/Runtime/UMG/Private/Extensions/UIComponentUserWidgetExtension.cpp
@@ -8,13 +8,17 @@
 #include "Extensions/UIComponent.h"
 #include "Extensions/UIComponentContainer.h"
 
+void UUIComponentUserWidgetExtension::Initialize()
+{
+	InitializeComponents();
+}
+
 void UUIComponentUserWidgetExtension::PreConstruct(bool bIsDesignTime)
 {
 	Super::PreConstruct(bIsDesignTime);
 
 	if (ComponentContainer)
-	{		
-		InitializeComponents();
+	{
 		ComponentContainer->ForEachComponent([bIsDesignTime](UUIComponent* Component)
 			{
 				if(ensure(Component))
diff --git a/Engine/Source/Runtime/UMG/Private/Extensions/UIComponentWidgetBlueprintGeneratedClassExtension.cpp b/Engine/Source/Runtime/UMG/Private/Extensions/UIComponentWidgetBlueprintGeneratedClassExtension.cpp
--- a/Engine/Source/Runtime/UMG/Private/Extensions/UIComponentWidgetBlueprintGeneratedClassExtension.cpp
+++ b/Engine/Source/Runtime/UMG/Private/Extensions/UIComponentWidgetBlueprintGeneratedClassExtension.cpp
@@ -82,6 +82,17 @@ void UUIComponentWidgetBlueprintGeneratedClassExtension::InitializeContainer(UUI
 	ComponentContainer = InComponentContainer;
 }
 
+void UUIComponentWidgetBlueprintGeneratedClassExtension::Initialize(UUserWidget* UserWidget)
+{
+	UUIComponentUserWidgetExtension* UserWidgetExtension = UserWidget->GetExtension<UUIComponentUserWidgetExtension>();
+	if (!UserWidgetExtension)
+	{
+		UserWidgetExtension = UserWidget->AddExtension<UUIComponentUserWidgetExtension>();
+		UserWidgetExtension->InitializeContainer(DuplicateContainer(UserWidget));
+	}
+}
+
+
 void UUIComponentWidgetBlueprintGeneratedClassExtension::PreConstruct(UUserWidget* UserWidget, bool IsDesignTime)
 {
 	// This is where we can add the UserWidgetExtension and Duplicate the Component Container from the GeneratedClassExtension
@@ -92,6 +103,7 @@ void UUIComponentWidgetBlueprintGeneratedClassExtension::PreConstruct(UUserWidge
 	{
 		UserWidgetExtension = UserWidget->AddExtension<UUIComponentUserWidgetExtension>();
 		UserWidgetExtension->InitializeContainer(DuplicateContainer(UserWidget));
+		UserWidgetExtension->Initialize();
 	}
 	else
 	{
diff --git a/Engine/Source/Runtime/UMG/Public/Components/Widget.h b/Engine/Source/Runtime/UMG/Public/Components/Widget.h
--- a/Engine/Source/Runtime/UMG/Public/Components/Widget.h
+++ b/Engine/Source/Runtime/UMG/Public/Components/Widget.h
@@ -395,6 +395,11 @@ protected:
 	/** True implies widget state has been initialized. */
 	bool bWidgetStateInitialized : 1;
 
+public:
+	/** True if the SWidget should be Wrapped in a Component */
+	UPROPERTY(Transient)
+	uint8 bWrappedByComponent : 1;
+
 public:
 #if WITH_EDITORONLY_DATA
 	/** Stores the design time flag setting if the widget is hidden inside the designer */
@@ -1170,6 +1175,9 @@ protected:
 	/** The underlying SWidget. */
 	TWeakPtr<SWidget> MyWidget;
 
+	/** The underlying SWidget for the Component wrapper widget. */
+	TWeakPtr<SWidget> ComponentWrapperWidget;
+	
 	/** The underlying SWidget contained in a SObjectWidget */
 	TWeakPtr<SObjectWidget> MyGCWidget;
 
diff --git a/Engine/Source/Runtime/UMG/Public/Extensions/UIComponent.h b/Engine/Source/Runtime/UMG/Public/Extensions/UIComponent.h
--- a/Engine/Source/Runtime/UMG/Public/Extensions/UIComponent.h
+++ b/Engine/Source/Runtime/UMG/Public/Extensions/UIComponent.h
@@ -13,6 +13,7 @@
 #include "UIComponent.generated.h"
 
 class UWidget;
+class SWidget;
 
 /** 
  * This is the base class to for UI Components that can be added to any UMG Widgets
@@ -54,6 +55,13 @@ public:
 	 *  @returns The owner widget
 	*/
 	UMG_API TWeakObjectPtr<UWidget> GetOwner() const;
+	
+	/**
+	* Called when the OwnerWidget is Rebuilt to allow to build a more complex hierarchy between the owner parents and the Owner widget.
+	* @param Content The Owner Widget Content.
+	* @returns The SWidget that wraps the Content passed.   
+	*/
+	UMG_API virtual TSharedRef<SWidget> RebuildWidgetWithContent(TSharedRef<SWidget> Content);
 
 	//~ Begin INotifyFieldValueChanged Interface
 	UMG_API virtual FDelegateHandle AddFieldValueChangedDelegate(UE::FieldNotification::FFieldId InFieldId, FFieldValueChangedDelegate InNewDelegate) override final;
@@ -65,6 +73,7 @@ public:
 	//~ End INotifyFieldValueChanged Interface
 
 protected:
+	
 	UMG_API virtual void OnInitialize();
 	UMG_API virtual void OnPreConstruct(bool bIsDesignTime);
 	UMG_API virtual void OnConstruct();
diff --git a/Engine/Source/Runtime/UMG/Public/Extensions/UIComponentUserWidgetExtension.h b/Engine/Source/Runtime/UMG/Public/Extensions/UIComponentUserWidgetExtension.h
--- a/Engine/Source/Runtime/UMG/Public/Extensions/UIComponentUserWidgetExtension.h
+++ b/Engine/Source/Runtime/UMG/Public/Extensions/UIComponentUserWidgetExtension.h
@@ -23,7 +23,7 @@ class UUIComponentUserWidgetExtension : public UUserWidgetExtension // This Cont
 	GENERATED_BODY()
 
 public:
-	//virtual void Initialize() override;
+	UE_API virtual void Initialize() override;
 	UE_API virtual void PreConstruct(bool bIsDesignTime) override;
 	UE_API virtual void Construct() override;
 	UE_API virtual void Destruct() override;
diff --git a/Engine/Source/Runtime/UMG/Public/Extensions/UIComponentWidgetBlueprintGeneratedClassExtension.h b/Engine/Source/Runtime/UMG/Public/Extensions/UIComponentWidgetBlueprintGeneratedClassExtension.h
--- a/Engine/Source/Runtime/UMG/Public/Extensions/UIComponentWidgetBlueprintGeneratedClassExtension.h
+++ b/Engine/Source/Runtime/UMG/Public/Extensions/UIComponentWidgetBlueprintGeneratedClassExtension.h
@@ -22,6 +22,9 @@ class UUIComponentWidgetBlueprintGeneratedClassExtension : public UWidgetBluepri
 {
 	GENERATED_BODY()
 public:
+
+    UE_API virtual void Initialize(UUserWidget* UserWidget) override;
+
 	UE_API void InitializeContainer(UUIComponentContainer* InComponentContainer);
 
 	UE_API virtual void PreConstruct(UUserWidget* UserWidget, bool IsDesignTime) override;
