Fix a handful of cases for PSO precaching on the mobile renderer

Grabbed Kenzo's changes for translucency and shadow, and added changes for landscape, distortion, and a couple components.

diff --git a/Engine/Source/Runtime/Engine/Private/Components/MaterialBillboardComponent.cpp b/Engine/Source/Runtime/Engine/Private/Components/MaterialBillboardComponent.cpp
--- a/Engine/Source/Runtime/Engine/Private/Components/MaterialBillboardComponent.cpp
+++ b/Engine/Source/Runtime/Engine/Private/Components/MaterialBillboardComponent.cpp
@@ -414,13 +414,24 @@ void UMaterialBillboardComponent::CollectPSOPrecacheData(const FPSOPrecacheParam
 			else
 			{
 				FVertexDeclarationElementList VtxElements;
-				//Position Float3, 2x PackedNormal, 1x UVs Half2, Color, InstancingID
 				VtxElements.Add(FVertexElement(0, 0, VET_Float3, 0, sizeof(FVector3f)));
 				VtxElements.Add(FVertexElement(1, 0, VET_PackedNormal, 1, sizeof(FPackedNormal) * 2));
 				VtxElements.Add(FVertexElement(1, sizeof(FPackedNormal), VET_PackedNormal, 2, sizeof(FPackedNormal) * 2));
-				VtxElements.Add(FVertexElement(2, 0, VET_Half2, 3, 4));
-				VtxElements.Add(FVertexElement(3, 0, VET_Color, 4, sizeof(FColor)));
-				VtxElements.Add(FVertexElement(4, 0, VET_UInt, 13, sizeof(uint32), true));
+				VtxElements.Add(FVertexElement(2, 0, VET_Color, 3, sizeof(FColor)));
+				VtxElements.Add(FVertexElement(3, 0, VET_Half2, 4, 4));
+				VtxElements.Add(FVertexElement(3, 0, VET_Half2, 5, 4));
+				VtxElements.Add(FVertexElement(3, 0, VET_Half2, 6, 4));
+				VtxElements.Add(FVertexElement(3, 0, VET_Half2, 7, 4));
+				VtxElements.Add(FVertexElement(3, 0, VET_Half2, 15, 4));
+				if (IsGPUSkinPassThroughSupported(GMaxRHIShaderPlatform))
+				{
+					VtxElements.Add(FVertexElement(4, 0, VET_Float4, 14, 0));
+					VtxElements.Add(FVertexElement(5, 0, VET_UInt, 13, 0, true));
+				}
+				else
+				{
+					VtxElements.Add(FVertexElement(4, 0, VET_UInt, 13, 0, true));
+				}
 				ComponentParams.VertexFactoryDataList.Add(FPSOPrecacheVertexFactoryData(VFType, VtxElements));
 			}
 		}
diff --git a/Engine/Source/Runtime/Landscape/Private/LandscapeGrassWeightExporter.cpp b/Engine/Source/Runtime/Landscape/Private/LandscapeGrassWeightExporter.cpp
--- a/Engine/Source/Runtime/Landscape/Private/LandscapeGrassWeightExporter.cpp
+++ b/Engine/Source/Runtime/Landscape/Private/LandscapeGrassWeightExporter.cpp
@@ -394,6 +394,7 @@ IPSOCollector* CreateLandscapeGrassWeightPSOCollector(ERHIFeatureLevel::Type Fea
 	return new FLandscapeGrassWeightMeshProcessor(nullptr, FeatureLevel, nullptr, nullptr);
 }
 FRegisterPSOCollectorCreateFunction RegisterLandscapeGrassWeightPSOCollector(&CreateLandscapeGrassWeightPSOCollector, EShadingPath::Deferred, LandscapeGrassWeightMeshPassName);
+FRegisterPSOCollectorCreateFunction RegisterMobileLandscapeGrassWeightPSOCollector(&CreateLandscapeGrassWeightPSOCollector, EShadingPath::Mobile, LandscapeGrassWeightMeshPassName);
 
 FLandscapeGrassWeightExporter_RenderThread::FLandscapeGrassWeightExporter_RenderThread(const TArray<int32>& InHeightMips, bool bInReadbackToCPU)
 	: HeightMips(InHeightMips)
diff --git a/Engine/Source/Runtime/Renderer/Private/DistortionRendering.cpp b/Engine/Source/Runtime/Renderer/Private/DistortionRendering.cpp
--- a/Engine/Source/Runtime/Renderer/Private/DistortionRendering.cpp
+++ b/Engine/Source/Runtime/Renderer/Private/DistortionRendering.cpp
@@ -1196,4 +1196,4 @@ FMeshPassProcessor* CreateMobileDistortionPassProcessor(ERHIFeatureLevel::Type F
 }
 
 REGISTER_MESHPASSPROCESSOR_AND_PSOCOLLECTOR(DistortionPass, CreateDistortionPassProcessor, EShadingPath::Deferred, EMeshPass::Distortion, EMeshPassFlags::MainView);
-FRegisterPassProcessorCreateFunction RegisterMobileDistortionPass(&CreateMobileDistortionPassProcessor, EShadingPath::Mobile, EMeshPass::Distortion, EMeshPassFlags::MainView);
+REGISTER_MESHPASSPROCESSOR_AND_PSOCOLLECTOR(MobileDistortionPass, CreateMobileDistortionPassProcessor, EShadingPath::Mobile, EMeshPass::Distortion, EMeshPassFlags::MainView);
diff --git a/Engine/Source/Runtime/Renderer/Private/MobileBasePass.cpp b/Engine/Source/Runtime/Renderer/Private/MobileBasePass.cpp
--- a/Engine/Source/Runtime/Renderer/Private/MobileBasePass.cpp
+++ b/Engine/Source/Runtime/Renderer/Private/MobileBasePass.cpp
@@ -1025,7 +1025,18 @@ void FMobileBasePassMeshProcessor::CollectPSOInitializers(const FSceneTexturesCo
 		FExclusiveDepthStencil::DepthWrite_StencilWrite;
 
 	FGraphicsPipelineRenderTargetsInfo RenderTargetsInfo;
-	SetupGBufferRenderTargetInfo(SceneTexturesConfig, RenderTargetsInfo, false /*bSetupDepthStencil*/);
+	RenderTargetsInfo.NumSamples = 1;
+	if (bTranslucentBasePass)
+	{		
+		EPixelFormat SceneColorFormat = SceneTexturesConfig.ColorFormat;
+		ETextureCreateFlags SceneColorCreateFlags = SceneTexturesConfig.ColorCreateFlags;
+		AddRenderTargetInfo(SceneColorFormat, SceneColorCreateFlags, RenderTargetsInfo);
+		ExclusiveDepthStencil = FExclusiveDepthStencil::DepthRead_StencilWrite;
+	}
+	else
+	{
+		SetupGBufferRenderTargetInfo(SceneTexturesConfig, RenderTargetsInfo, false /*bSetupDepthStencil*/);
+	}
 	SetupDepthStencilInfo(PF_DepthStencil, SceneTexturesConfig.DepthCreateFlags, ERenderTargetLoadAction::ELoad,
 		ERenderTargetLoadAction::ELoad, ExclusiveDepthStencil, RenderTargetsInfo);
 	SetupMultiViewInfo(RenderTargetsInfo);
diff --git a/Engine/Source/Runtime/Renderer/Private/ShadowDepthRendering.cpp b/Engine/Source/Runtime/Renderer/Private/ShadowDepthRendering.cpp
--- a/Engine/Source/Runtime/Renderer/Private/ShadowDepthRendering.cpp
+++ b/Engine/Source/Runtime/Renderer/Private/ShadowDepthRendering.cpp
@@ -2147,6 +2147,7 @@ bool FShadowDepthPassMeshProcessor::TryAddMeshBatch(const FMeshBatch& RESTRICT M
 
 	bool bResult = true;
 	if (bShouldCastShadow
+		&& GetShadowQuality() != 0
 		&& ShouldIncludeDomainInMeshPass(Material.GetMaterialDomain())
 		&& ShouldIncludeMaterialInDefaultOpaquePass(Material)
 		&& EnumHasAnyFlags(MeshSelectionMask, MeshBatch.VertexFactory->SupportsGPUScene(FeatureLevel) ? EShadowMeshSelection::VSM : EShadowMeshSelection::SM))
diff --git a/Engine/Source/Runtime/UMG/Private/Components/WidgetComponent.cpp b/Engine/Source/Runtime/UMG/Private/Components/WidgetComponent.cpp
--- a/Engine/Source/Runtime/UMG/Private/Components/WidgetComponent.cpp
+++ b/Engine/Source/Runtime/UMG/Private/Components/WidgetComponent.cpp
@@ -711,7 +711,7 @@ void UWidgetComponent::CollectPSOPrecacheData(const FPSOPrecacheParams& BasePrec
 		ComponentParams.Priority = EPSOPrecachePriority::High;
 		ComponentParams.MaterialInterface = MaterialInstance;
 		ComponentParams.PSOPrecacheParams = BasePrecachePSOParams;
-		
+
 		if (RHISupportsManualVertexFetch(GMaxRHIShaderPlatform))
 		{
 			ComponentParams.VertexFactoryDataList.Add(FPSOPrecacheVertexFactoryData(&FLocalVertexFactory::StaticType));
@@ -719,13 +719,24 @@ void UWidgetComponent::CollectPSOPrecacheData(const FPSOPrecacheParams& BasePrec
 		else
 		{
 			FVertexDeclarationElementList VtxElements;
-			//Position Float3, 2x PackedNormal, 1x UVs Float2, Color, InstancingID
 			VtxElements.Add(FVertexElement(0, 0, VET_Float3, 0, sizeof(FVector3f)));
 			VtxElements.Add(FVertexElement(1, 0, VET_PackedNormal, 1, sizeof(FPackedNormal) * 2));
 			VtxElements.Add(FVertexElement(1, sizeof(FPackedNormal), VET_PackedNormal, 2, sizeof(FPackedNormal) * 2));
-			VtxElements.Add(FVertexElement(2, 0, VET_Float2, 3, sizeof(FVector2f)));
-			VtxElements.Add(FVertexElement(3, 0, VET_Color, 4, sizeof(FColor)));
-			VtxElements.Add(FVertexElement(4, 0, VET_UInt, 13, sizeof(uint32), true));
+			VtxElements.Add(FVertexElement(2, 0, VET_Color, 3, sizeof(FColor)));
+			VtxElements.Add(FVertexElement(3, 0, VET_Float2, 4, sizeof(FVector2f)));
+			VtxElements.Add(FVertexElement(3, 0, VET_Float2, 5, sizeof(FVector2f)));
+			VtxElements.Add(FVertexElement(3, 0, VET_Float2, 6, sizeof(FVector2f)));
+			VtxElements.Add(FVertexElement(3, 0, VET_Float2, 7, sizeof(FVector2f)));
+			VtxElements.Add(FVertexElement(3, 0, VET_Float2, 15, sizeof(FVector2f)));
+			if (IsGPUSkinPassThroughSupported(GMaxRHIShaderPlatform))
+			{
+				VtxElements.Add(FVertexElement(4, 0, VET_Float4, 14, 0));
+				VtxElements.Add(FVertexElement(5, 0, VET_UInt, 13, 0, true));
+			}
+			else
+			{
+				VtxElements.Add(FVertexElement(4, 0, VET_UInt, 13, 0, true));
+			}
 			ComponentParams.VertexFactoryDataList.Add(FPSOPrecacheVertexFactoryData(&FLocalVertexFactory::StaticType, VtxElements));
 		}
 	}
