Fix widget SetVisibility broadcasting field value changed before updating the cached widget

#rnx

diff --git a/Engine/Source/Runtime/UMG/Private/Components/Widget.cpp b/Engine/Source/Runtime/UMG/Private/Components/Widget.cpp
--- a/Engine/Source/Runtime/UMG/Private/Components/Widget.cpp
+++ b/Engine/Source/Runtime/UMG/Private/Components/Widget.cpp
@@ -421,10 +421,10 @@ void UWidget::SetVisibility(ESlateVisibility InVisibility)
 
 void UWidget::SetVisibilityInternal(ESlateVisibility InVisibility)
 {
-	if (Visibility != InVisibility)
+	const bool bVisibilityChanged = Visibility != InVisibility;
+	if (bVisibilityChanged)
 	{
 		Visibility = InVisibility;
-		BroadcastFieldValueChanged(FFieldNotificationClassDescriptor::Visibility);
 	}
 
 	TSharedPtr<SWidget> SafeWidget = GetCachedWidget();
@@ -433,6 +433,11 @@ void UWidget::SetVisibilityInternal(ESlateVisibility InVisibility)
 		EVisibility SlateVisibility = UWidget::ConvertSerializedVisibilityToRuntime(InVisibility);
 		SafeWidget->SetVisibility(SlateVisibility);
 	}
+
+	if (bVisibilityChanged)
+	{
+		BroadcastFieldValueChanged(FFieldNotificationClassDescriptor::Visibility);
+	}
 }
 
 float UWidget::GetRenderOpacity() const
