Sequencer: re-add dropped support for Blocking Evaluation in UMG animations

The change also adds a new automated test that checks that an animation's first frame is applied immediately (synchronously) after a call to Play()

diff --git a/Engine/Source/Runtime/UMG/Private/Animation/WidgetAnimationState.cpp b/Engine/Source/Runtime/UMG/Private/Animation/WidgetAnimationState.cpp
--- a/Engine/Source/Runtime/UMG/Private/Animation/WidgetAnimationState.cpp
+++ b/Engine/Source/Runtime/UMG/Private/Animation/WidgetAnimationState.cpp
@@ -78,8 +78,15 @@ void FWidgetAnimationState::Initialize(UWidgetAnimation* InAnimation, UUserWidge
 	check(InUserWidget->AnimationTickManager);
 
 	const bool bNeedsLegacyPlayer = NeedsLegacyPlayer();
+	const bool bNeedsPrivateLinker = EnumHasAnyFlags(InAnimation->GetFlags(), EMovieSceneSequenceFlags::BlockingEvaluation);
 
 	UMovieSceneEntitySystemLinker* Linker = InUserWidget->AnimationTickManager->GetLinker();
+	if (bNeedsPrivateLinker)
+	{
+		Linker = UMovieSceneEntitySystemLinker::CreateLinker(InUserWidget->GetWorld(), UE::MovieScene::EEntitySystemLinkerRole::UMG);
+		PrivateLinker = Linker;
+	}
+
 	FInstanceRegistry* InstanceRegistry = Linker->GetInstanceRegistry();
 	FRootInstanceHandle RootInstanceHandle = InstanceRegistry->AllocateRootInstance(*InAnimation, InUserWidget);
 	{
@@ -150,6 +157,17 @@ bool FWidgetAnimationState::NeedsLegacyPlayer() const
 	}
 }
 
+void FWidgetAnimationState::FlushIfPrivateLinker()
+{
+	using namespace UE::MovieScene;
+
+	if (PrivateLinker)
+	{
+		TSharedPtr<FMovieSceneEntitySystemRunner> EntitySystemRunner = PrivateLinker->GetRunner();
+		EntitySystemRunner->Flush();
+	}
+}
+
 void FWidgetAnimationState::Tick(float InDeltaSeconds)
 {
 	using namespace UE::MovieScene;
@@ -232,6 +250,8 @@ void FWidgetAnimationState::Tick(float InDeltaSeconds)
 					PlaybackManager.GetPlayRate());
 		}
 #endif
+
+		FlushIfPrivateLinker();
 	}
 	else
 	{
@@ -284,11 +304,16 @@ void FWidgetAnimationState::OnStopped()
 	if (SharedPlaybackState)
 	{
 		TSharedPtr<FMovieSceneEntitySystemRunner> EntitySystemRunner = SharedPlaybackState->GetRunner();
-		EntitySystemRunner->QueueFinalUpdate(SharedPlaybackState->GetRootInstanceHandle());
+		const bool bNeedsFinalFlush = EntitySystemRunner->QueueFinalUpdate(SharedPlaybackState->GetRootInstanceHandle());
 
 		// Even if our request to Finish the instance was queued, we can wait until the next flush for those effects to be seen
 		// This will most likely happen immediately anyway since the runner will keep looping until its queue is empty,
 		// And we are already inside an active evaluation
+
+		if (bNeedsFinalFlush)
+		{
+			FlushIfPrivateLinker();
+		}
 	}
 
 	if (SharedPlaybackState && bRestoreState)
@@ -526,6 +551,8 @@ void FWidgetAnimationState::Play(const FWidgetAnimationStatePlayParams& PlayPara
 				PlaybackManager.IsPingPongPlayback() ? TEXT(", PingPong") : TEXT(""));
 	}
 #endif
+
+	FlushIfPrivateLinker();
 }
 
 void FWidgetAnimationState::Stop()
@@ -542,6 +569,8 @@ void FWidgetAnimationState::Stop()
 	// TODO: we should use PlaybackManager.GetEffectiveStartTime(), but the old behavior always used frame zero.
 	PlaybackManager.SetCurrentTime(0);
 
+	bool bFlushPrivateLinker = false;
+
 	if (TSharedPtr<FSharedPlaybackState> SharedPlaybackState = WeakPlaybackState.Pin())
 	{
 		const FRootInstanceHandle& RootInstanceHandle = SharedPlaybackState->GetRootInstanceHandle();
@@ -561,6 +590,7 @@ void FWidgetAnimationState::Stop()
 					ERunnerUpdateFlags::Flush);
 
 			bIsStopping = true;
+			bFlushPrivateLinker = true;
 		}
 		else
 		{
@@ -579,6 +609,11 @@ void FWidgetAnimationState::Stop()
 				*GetNameSafe(Animation), StateIndex);
 	}
 #endif
+
+	if (bFlushPrivateLinker)
+	{
+		FlushIfPrivateLinker();
+	}
 }
 
 void FWidgetAnimationState::Pause()
@@ -609,6 +644,8 @@ void FWidgetAnimationState::Pause()
 				*LexToString(PlaybackManager.GetCurrentTime()));
 	}
 #endif
+
+	FlushIfPrivateLinker();
 }
 
 void FWidgetAnimationState::Reverse()
@@ -699,5 +736,6 @@ void FWidgetAnimationState::AddReferencedObjects(FReferenceCollector& Collector)
 {
 	Collector.AddReferencedObject(Animation);
 	Collector.AddReferencedObject(LegacyPlayer);
+	Collector.AddReferencedObject(PrivateLinker);
 }
 
diff --git a/Engine/Source/Runtime/UMG/Public/Animation/WidgetAnimationState.h b/Engine/Source/Runtime/UMG/Public/Animation/WidgetAnimationState.h
--- a/Engine/Source/Runtime/UMG/Public/Animation/WidgetAnimationState.h
+++ b/Engine/Source/Runtime/UMG/Public/Animation/WidgetAnimationState.h
@@ -140,6 +140,8 @@ private:
 
 	bool NeedsLegacyPlayer() const;
 
+	void FlushIfPrivateLinker();
+
 	void OnBegunPlay();
 	void OnStopped();
 
@@ -171,6 +173,9 @@ private:
 
 	/** Shared playback state for the animation. */
 	TWeakPtr<UE::MovieScene::FSharedPlaybackState> WeakPlaybackState;
+	
+	/** Private linker for blocking/synchronous running. */
+	TObjectPtr<UMovieSceneEntitySystemLinker> PrivateLinker;
 
 	/** Playback manager for the animation. */
 	FMovieScenePlaybackManager PlaybackManager;
