SObjectTableRow: add support for shift and control modifiers when used with mouse click. The behaviour follows STableRow implementation.
Make single click deselect the previous selections in Multi-Selection mode. A bool parameter is exposed to designer that enables the old behaviour by default
old behaviour when listview was set to Multi-select: clicking on any item on the list marks it as selected and doesn't reset previous selection. Pressing control or shift at the same time as clicking does nothing.
new behaviour: clicking on any item on the list selects that item only and resets previous selection. Pressing control+click allows selecting multiple items at the same time, Pressing shift+click does range selection.

diff --git a/Engine/Plugins/Runtime/CommonUI/Source/CommonUI/Public/SCommonButtonTableRow.h b/Engine/Plugins/Runtime/CommonUI/Source/CommonUI/Public/SCommonButtonTableRow.h
--- a/Engine/Plugins/Runtime/CommonUI/Source/CommonUI/Public/SCommonButtonTableRow.h
+++ b/Engine/Plugins/Runtime/CommonUI/Source/CommonUI/Public/SCommonButtonTableRow.h
@@ -16,10 +16,12 @@ class SCommonButtonTableRow : public SObjectTableRow<ItemType>
 {
 public:
 	SLATE_BEGIN_ARGS(SCommonButtonTableRow<ItemType>)
-		:_bAllowDragging(true)
+		:_bAllowDragging(true),
+		_bAllowKeepPreselectedItems(true)
 	{}
 		SLATE_ARGUMENT(bool, bAllowDragging)
 		SLATE_ARGUMENT(bool, bAllowDragDrop)
+		SLATE_ARGUMENT(bool, bAllowKeepPreselectedItems)
 		SLATE_EVENT(FOnRowHovered, OnHovered)
 		SLATE_EVENT(FOnRowHovered, OnUnhovered)
 		//	Drag and Drop functionality
@@ -40,6 +42,7 @@ public:
 			typename SObjectTableRow<ItemType>::FArguments()
 			.bAllowDragging(InArgs._bAllowDragging)
 			.bAllowDragDrop(InArgs._bAllowDragDrop)
+			.bAllowKeepPreselectedItems(InArgs._bAllowKeepPreselectedItems)
 			.OnHovered(InArgs._OnHovered)
 			.OnUnhovered(InArgs._OnUnhovered)
 			.OnRowCanAcceptDrop(InArgs._OnRowCanAcceptDrop)
diff --git a/Engine/Source/Runtime/UMG/Public/Components/ListViewBase.h b/Engine/Source/Runtime/UMG/Public/Components/ListViewBase.h
--- a/Engine/Source/Runtime/UMG/Public/Components/ListViewBase.h
+++ b/Engine/Source/Runtime/UMG/Public/Components/ListViewBase.h
@@ -679,6 +679,7 @@ protected:
 				return SNew(ObjectTableRowT, OwnerTable, *WidgetObject, this)
 					.bAllowDragging(bAllowDragging)
 					.bAllowDragDrop(bAllowDragDrop)
+					.bAllowKeepPreselectedItems(bAllowKeepPreselectedItems)
 					.OnHovered_UObject(this, &UListViewBase::HandleListEntryHovered)
 					.OnUnhovered_UObject(this, &UListViewBase::HandleListEntryUnhovered)
 					.OnRowDragDetected_UObject(this, &UListViewBase::HandleListEntryDragDetected)
@@ -841,6 +842,10 @@ protected:
 	UPROPERTY(EditAnywhere, BlueprintReadOnly, Category = ListView)
 	bool bSelectItemOnNavigation = true;
 
+	/** If true, when selecting an item via mouse button, we allow pre-selected items to remain selected (applicable to "Multi" Selection Mode only). */
+	UPROPERTY(EditAnywhere, BlueprintReadOnly, Category = ListView)
+	bool bAllowKeepPreselectedItems = true;
+
 	/** Called when a row widget is released by the list (i.e. when it no longer represents a list item) */
 	UPROPERTY(BlueprintAssignable, Category = Events, meta = (DisplayName = "On Entry Released"))
 	FOnListEntryReleasedDynamic BP_OnEntryReleased;
diff --git a/Engine/Source/Runtime/UMG/Public/Slate/SObjectTableRow.h b/Engine/Source/Runtime/UMG/Public/Slate/SObjectTableRow.h
--- a/Engine/Source/Runtime/UMG/Public/Slate/SObjectTableRow.h
+++ b/Engine/Source/Runtime/UMG/Public/Slate/SObjectTableRow.h
@@ -66,6 +66,7 @@ public:
 	{}
 		SLATE_ARGUMENT(bool, bAllowDragging)
 		SLATE_ARGUMENT(bool, bAllowDragDrop)
+		SLATE_ARGUMENT(bool, bAllowKeepPreselectedItems)
 		SLATE_DEFAULT_SLOT(FArguments, Content)
 		SLATE_EVENT(FOnRowHovered, OnHovered)
 		SLATE_EVENT(FOnRowHovered, OnUnhovered)
@@ -93,6 +94,7 @@ public:
 
 			bAllowDragging = InArgs._bAllowDragging;
 			bAllowDragDrop = InArgs._bAllowDragDrop;
+			bAllowKeepPreselectedItems = InArgs._bAllowKeepPreselectedItems;
 			OnHovered = InArgs._OnHovered;
 			OnUnhovered = InArgs._OnUnhovered;
 			OnRowCanAcceptDrop = InArgs._OnRowCanAcceptDrop;
@@ -609,9 +611,26 @@ public:
 						if (MyItemPtr)
 						{
 							const ItemType& MyItem = *MyItemPtr;
-							if (!OwnerTable->Private_IsItemSelected(MyItem))
+							const bool bIsSelected = OwnerTable->Private_IsItemSelected(MyItem);
+
+							if (SelectionMode == ESelectionMode::Multi)
+							{
+								if (MouseEvent.IsShiftDown())
+								{
+									OwnerTable->Private_SelectRangeFromCurrentTo(MyItem);
+									bChangedSelectionOnMouseDown = true;
+
+								}
+								else if (MouseEvent.IsControlDown())
+								{
+									OwnerTable->Private_SetItemSelection(MyItem, !bIsSelected, true);
+									bChangedSelectionOnMouseDown = true;
+								}
+							}
+
+							if (!bIsSelected && !bChangedSelectionOnMouseDown)
 							{
-								if (SelectionMode != ESelectionMode::Multi)
+								if (SelectionMode != ESelectionMode::Multi || !bAllowKeepPreselectedItems)
 								{
 									OwnerTable->Private_ClearSelection();
 								}
@@ -665,10 +684,13 @@ public:
 							OwnerTable->Private_GetNumSelectedItems() > 1 &&
 							OwnerTable->Private_IsItemSelected(*MyItemPtr))
 						{
-							// Releasing mouse over one of the multiple selected items - leave this one as the sole selected item
-							OwnerTable->Private_ClearSelection();
-							OwnerTable->Private_SetItemSelection(*MyItemPtr, true, true);
-							bSignalSelectionChanged = true;
+							if (!MouseEvent.IsControlDown() && !MouseEvent.IsShiftDown())
+							{
+								// Releasing mouse over one of the multiple selected items - leave this one as the sole selected item
+								OwnerTable->Private_ClearSelection();
+								OwnerTable->Private_SetItemSelection(*MyItemPtr, true, true);
+								bSignalSelectionChanged = true;
+							}
 						}
 					}
 
@@ -804,6 +826,9 @@ private:
 
 	bool bProcessingSelectionTouch = false;
 
+	/** If true, when selecting an item via mouse button, we allow pre-selected items to remain selected */
+	bool bAllowKeepPreselectedItems = true;
+
 	/** Whether to allow dragging of this item */
 	bool bAllowDragging;
 
