Editor Crash Fix - Check TOptional DropZone before returning the value

diff --git a/Engine/Source/Runtime/UMG/Private/Blueprint/IUserListEntry.cpp b/Engine/Source/Runtime/UMG/Private/Blueprint/IUserListEntry.cpp
--- a/Engine/Source/Runtime/UMG/Private/Blueprint/IUserListEntry.cpp
+++ b/Engine/Source/Runtime/UMG/Private/Blueprint/IUserListEntry.cpp
@@ -51,9 +51,9 @@ void IUserListEntry::UpdateItemExpansion(UUserWidget& ListEntryWidget, bool bIsE
 	}
 }
 
-UMG_API void IUserListEntry::UpdateEntryDropIndicator(UUserWidget& ListEntryWidget, EItemDropZone DropZone)
+UMG_API void IUserListEntry::UpdateEntryDropIndicator(UUserWidget& ListEntryWidget, TOptional<EItemDropZone> DropZone)
 {
-	if (ListEntryWidget.Implements<UUserListEntry>())
+	if (ListEntryWidget.Implements<UUserListEntry>() && DropZone.IsSet())
 	{
 		Execute_BP_OnUpdateEntryDropIndicator(&ListEntryWidget, UDragDropOperation::ConvertSlateDropZoneToUMG(DropZone));
 	}
diff --git a/Engine/Source/Runtime/UMG/Public/Blueprint/IUserListEntry.h b/Engine/Source/Runtime/UMG/Public/Blueprint/IUserListEntry.h
--- a/Engine/Source/Runtime/UMG/Public/Blueprint/IUserListEntry.h
+++ b/Engine/Source/Runtime/UMG/Public/Blueprint/IUserListEntry.h
@@ -49,7 +49,7 @@ public:
 	static UMG_API void ReleaseEntry(UUserWidget& ListEntryWidget);
 	static UMG_API void UpdateItemSelection(UUserWidget& ListEntryWidget, bool bIsSelected);
 	static UMG_API void UpdateItemExpansion(UUserWidget& ListEntryWidget, bool bIsExpanded);
-	static UMG_API void UpdateEntryDropIndicator(UUserWidget& ListEntryWidget, EItemDropZone DropZone);
+	static UMG_API void UpdateEntryDropIndicator(UUserWidget& ListEntryWidget, TOptional<EItemDropZone> DropZone);
 	static UMG_API void UpdateEntryDragOverState(UUserWidget& ListEntryWidget, bool bIsDraggingOver);
 	static UMG_API void EndEntryDropOperation(UUserWidget& ListEntryWidget, bool bSuccess);
 	static UMG_API void HandleEntryDropped(UUserWidget& ListEntryWidget, UDragDropOperation* Operation);
diff --git a/Engine/Source/Runtime/UMG/Public/Slate/SObjectTableRow.h b/Engine/Source/Runtime/UMG/Public/Slate/SObjectTableRow.h
--- a/Engine/Source/Runtime/UMG/Public/Slate/SObjectTableRow.h
+++ b/Engine/Source/Runtime/UMG/Public/Slate/SObjectTableRow.h
@@ -491,7 +491,7 @@ public:
 
 			if (WidgetObject)
 			{
-				IUserListEntry::UpdateEntryDropIndicator(*WidgetObject, ItemDropZone.GetValue());
+				IUserListEntry::UpdateEntryDropIndicator(*WidgetObject, ItemDropZone);
 			}
 		}
 
