At the momento opening the virtual keyboard via a hotkey will not work, right now the input field requires either a click to open or it can open when it receives focus (but this would cause a lot of unexpected times the virtual keyboard would open).
The solution proposed is to create a custom function which calls FSlateApplication::Get().ShowVirtualKeyboard(...) by request from the BP, this function receives by parameter a boolean (to show or not the VK) and also the text widget that is going to receive the VK input.

#tests: Tested the node in consoles, it appears to work correctly; it opens and closes on demand. When the keyboard is open, it correctly absorbs the input and passes it on to the corresponding editable text.

#rnx

diff --git a/Engine/Source/Runtime/Slate/Private/Widgets/Input/SEditableText.cpp b/Engine/Source/Runtime/Slate/Private/Widgets/Input/SEditableText.cpp
--- a/Engine/Source/Runtime/Slate/Private/Widgets/Input/SEditableText.cpp
+++ b/Engine/Source/Runtime/Slate/Private/Widgets/Input/SEditableText.cpp
@@ -700,3 +700,11 @@ TOptional<FText> SEditableText::GetDefaultAccessibleText(EAccessibleType Accessi
 	return GetHintText();
 }
 #endif
+
+void SEditableText::ToggleVirtualKeyboard(bool bShow, const int32 UserIndex)
+{
+	if (EditableTextLayout.IsValid())
+	{
+		EditableTextLayout->ToggleVirtualKeyboard(bShow, UserIndex);
+	}
+}
diff --git a/Engine/Source/Runtime/Slate/Private/Widgets/Text/SlateEditableTextLayout.cpp b/Engine/Source/Runtime/Slate/Private/Widgets/Text/SlateEditableTextLayout.cpp
--- a/Engine/Source/Runtime/Slate/Private/Widgets/Text/SlateEditableTextLayout.cpp
+++ b/Engine/Source/Runtime/Slate/Private/Widgets/Text/SlateEditableTextLayout.cpp
@@ -4185,3 +4185,11 @@ void FSlateEditableTextLayout::FTextInputMethodContext::EndComposition()
 		bIsComposing = false;
 	}
 }
+
+void FSlateEditableTextLayout::ToggleVirtualKeyboard(bool bShow, const int32 UserIndex)
+{
+	if (FSlateApplication::IsInitialized() && FPlatformApplicationMisc::RequiresVirtualKeyboard())
+	{
+		FSlateApplication::Get().ShowVirtualKeyboard(bShow, UserIndex, VirtualKeyboardEntry);
+	}
+}
\ No newline at end of file
diff --git a/Engine/Source/Runtime/Slate/Public/Widgets/Input/SEditableText.h b/Engine/Source/Runtime/Slate/Public/Widgets/Input/SEditableText.h
--- a/Engine/Source/Runtime/Slate/Public/Widgets/Input/SEditableText.h
+++ b/Engine/Source/Runtime/Slate/Public/Widgets/Input/SEditableText.h
@@ -391,6 +391,9 @@ public:
 	/** Select a block of text */
 	SLATE_API void SelectText(const FTextLocation& InSelectionStart, const FTextLocation& InCursorLocation);
 
+	/* Function to spawn virtual keyboard on demand */
+	SLATE_API void ToggleVirtualKeyboard(bool bShow, const int32 UserIndex);
+
 protected:
 	//~ Begin SWidget Interface
 	SLATE_API virtual void Tick( const FGeometry& AllottedGeometry, const double InCurrentTime, const float InDeltaTime ) override;
diff --git a/Engine/Source/Runtime/Slate/Public/Widgets/Text/SlateEditableTextLayout.h b/Engine/Source/Runtime/Slate/Public/Widgets/Text/SlateEditableTextLayout.h
--- a/Engine/Source/Runtime/Slate/Public/Widgets/Text/SlateEditableTextLayout.h
+++ b/Engine/Source/Runtime/Slate/Public/Widgets/Text/SlateEditableTextLayout.h
@@ -416,6 +416,9 @@ public:
 	 */
 	SLATE_API void GetTextLine(const int32 InLineIndex, FString& OutTextLine) const;
 
+	/* Function to spawn virtual keyboard on demand */
+	SLATE_API void ToggleVirtualKeyboard(bool bShow, const int32 UserIndex);
+
 private:
 	/** Insert the given text at the current cursor position, correctly taking into account new line characters */
 	SLATE_API void InsertTextAtCursorImpl(const FString& InString);
diff --git a/Engine/Source/Runtime/UMG/Private/Components/EditableText.cpp b/Engine/Source/Runtime/UMG/Private/Components/EditableText.cpp
--- a/Engine/Source/Runtime/UMG/Private/Components/EditableText.cpp
+++ b/Engine/Source/Runtime/UMG/Private/Components/EditableText.cpp
@@ -3,6 +3,8 @@
 #include "Components/EditableText.h"
 #include "UObject/ConstructorHelpers.h"
 #include "Engine/Font.h"
+#include "Engine/LocalPlayer.h"
+#include "Framework/Application/SlateApplication.h"
 #include "Widgets/DeclarativeSyntaxSupport.h"
 #include "Widgets/Input/SEditableText.h"
 #include "Slate/SlateBrushAsset.h"
@@ -425,6 +427,18 @@ const FText UEditableText::GetPaletteCategory()
 
 #endif
 
+void UEditableText::ToggleVirtualKeyboard(bool bShow)
+{
+	if (MyEditableText.IsValid())
+	{
+		if (ULocalPlayer* LocalPlayer = GetOwningLocalPlayer())
+		{
+			const int32 UserIndex = FSlateApplication::Get().GetUserIndexForController(LocalPlayer->GetControllerId());
+			MyEditableText->ToggleVirtualKeyboard(bShow, UserIndex);
+		}
+	}
+}
+
 /////////////////////////////////////////////////////
 
 #undef LOCTEXT_NAMESPACE
diff --git a/Engine/Source/Runtime/UMG/Public/Components/EditableText.h b/Engine/Source/Runtime/UMG/Public/Components/EditableText.h
--- a/Engine/Source/Runtime/UMG/Public/Components/EditableText.h
+++ b/Engine/Source/Runtime/UMG/Public/Components/EditableText.h
@@ -244,6 +244,10 @@ public:
 	UMG_API bool GetEnableIntegratedKeyboard() const;
 	UMG_API void SetEnableIntegratedKeyboard(bool InEnableIntegratedKeyboard);
 
+	/* Function to spawn virtual keyboard on demand */
+	UFUNCTION(BlueprintCallable, Category = "VirtualKeyboard")
+	UMG_API void ToggleVirtualKeyboard(bool bShow); 
+
 public:
 	
 	//~ Begin UWidget Interface
