UMG: various fixes for widget animations:

- Fix crash caused by the runner being invalid for an unknown reason
- Fix iteration of active animations that can grow during that iteration
- Fix iteration of animated widgets that can grow during that iteration

diff --git a/Engine/Source/Runtime/UMG/Private/Animation/UMGSequenceTickManager.cpp b/Engine/Source/Runtime/UMG/Private/Animation/UMGSequenceTickManager.cpp
--- a/Engine/Source/Runtime/UMG/Private/Animation/UMGSequenceTickManager.cpp
+++ b/Engine/Source/Runtime/UMG/Private/Animation/UMGSequenceTickManager.cpp
@@ -72,6 +72,10 @@ void UUMGSequenceTickManager::AddWidget(UUserWidget* InWidget)
 	{
 		WidgetData->bIsTicking = true;
 	}
+	else if (bIsTicking)
+	{
+		PendingWeakUserWidgetData.Add(WeakWidget);
+	}
 	else
 	{
 		WeakUserWidgetData.Add(WeakWidget, FSequenceTickManagerWidgetData());
@@ -91,6 +95,10 @@ void UUMGSequenceTickManager::OnWidgetTicked(UUserWidget* InWidget)
 	{
 		WidgetData->bIsTicking = true;
 	}
+	else if (bIsTicking)
+	{
+		PendingWeakUserWidgetData.Add(InWidget);
+	}
 	else
 	{
 		WeakUserWidgetData.Add(InWidget, FSequenceTickManagerWidgetData());
@@ -220,6 +228,15 @@ void UUMGSequenceTickManager::TickWidgetAnimations(float DeltaSeconds)
 
 	ForceFlush();
 
+	if (PendingWeakUserWidgetData.Num() > 0)
+	{
+		for (TWeakObjectPtr<UUserWidget> WeakWidget : PendingWeakUserWidgetData)
+		{
+			WeakUserWidgetData.Add(WeakWidget, FSequenceTickManagerWidgetData());
+		}
+		PendingWeakUserWidgetData.Reset();
+	}
+
 	if (!Runner->IsCurrentlyEvaluating())
 	{
 		for (auto WidgetIter = WeakUserWidgetData.CreateIterator(); WidgetIter; ++WidgetIter)
diff --git a/Engine/Source/Runtime/UMG/Private/Animation/WidgetAnimationState.cpp b/Engine/Source/Runtime/UMG/Private/Animation/WidgetAnimationState.cpp
--- a/Engine/Source/Runtime/UMG/Private/Animation/WidgetAnimationState.cpp
+++ b/Engine/Source/Runtime/UMG/Private/Animation/WidgetAnimationState.cpp
@@ -186,6 +186,19 @@ void FWidgetAnimationState::Tick(float InDeltaSeconds)
 		return;
 	}
 
+	TSharedPtr<FMovieSceneEntitySystemRunner> EntitySystemRunner = SharedPlaybackState->GetRunner();
+	if (!EntitySystemRunner)
+	{
+#if !NO_LOGGING
+		if (UE::UMG::GVarAnimationDetailedLogging)
+		{
+			UE_LOG(LogUMG, Log, TEXT("Animation: Ticking '%s' [%d] aborted, shared playback state has no runner."),
+					*GetNameSafe(Animation), StateIndex);
+		}
+#endif
+		return;
+	}
+
 	const EMovieScenePlayerStatus::Type PreviousPlaybackStatus = PlaybackManager.GetPlaybackStatus();
 	if (PreviousPlaybackStatus == EMovieScenePlayerStatus::Playing)
 	{
@@ -204,7 +217,6 @@ void FWidgetAnimationState::Tick(float InDeltaSeconds)
 
 		// Queue up the evalutions as needed. If we finished playback, setup our OnStopped callback
 		// to tear things down.
-		TSharedPtr<FMovieSceneEntitySystemRunner> EntitySystemRunner = SharedPlaybackState->GetRunner();
 		for (int32 Index = 0; Index < TickContexts.Num(); ++Index)
 		{
 			FMovieSceneContext& TickContext(TickContexts[Index]);
diff --git a/Engine/Source/Runtime/UMG/Private/UserWidget.cpp b/Engine/Source/Runtime/UMG/Private/UserWidget.cpp
--- a/Engine/Source/Runtime/UMG/Private/UserWidget.cpp
+++ b/Engine/Source/Runtime/UMG/Private/UserWidget.cpp
@@ -1919,11 +1919,13 @@ void UUserWidget::TickActionsAndAnimation(float InDeltaTime)
 		return;
 	}
 
-	// Update active movie scenes, none will be removed here, but new
-	// ones can be added during the tick, if a player ends and triggers
-	// starting another animation
-	for (FWidgetAnimationState& State : ActiveAnimations)
+	// Update active widget animations. None will be removed here, but new
+	// ones can be added during the tick, if one animation ends and triggers
+	// starting another animation. So iterate with indices always compared 
+	// to the new array size, and not an iterator.
+	for (int32 Index = 0; Index < ActiveAnimations.Num(); ++Index)
 	{
+		FWidgetAnimationState& State(ActiveAnimations[Index]);
 		State.Tick(InDeltaTime);
 	}
 }
diff --git a/Engine/Source/Runtime/UMG/Public/Animation/UMGSequenceTickManager.h b/Engine/Source/Runtime/UMG/Public/Animation/UMGSequenceTickManager.h
--- a/Engine/Source/Runtime/UMG/Public/Animation/UMGSequenceTickManager.h
+++ b/Engine/Source/Runtime/UMG/Public/Animation/UMGSequenceTickManager.h
@@ -66,6 +66,9 @@ private:
 	UPROPERTY(transient)
 	TMap<TWeakObjectPtr<UUserWidget>, FSequenceTickManagerWidgetData> WeakUserWidgetData;
 
+	UPROPERTY(transient)
+	TArray<TWeakObjectPtr<UUserWidget>> PendingWeakUserWidgetData;
+
 	UPROPERTY(transient)
 	TObjectPtr<UMovieSceneEntitySystemLinker> Linker;
 
