UMG: fix iteration of sparse array ActiveAnimations

diff --git a/Engine/Source/Runtime/UMG/Private/UserWidget.cpp b/Engine/Source/Runtime/UMG/Private/UserWidget.cpp
--- a/Engine/Source/Runtime/UMG/Private/UserWidget.cpp
+++ b/Engine/Source/Runtime/UMG/Private/UserWidget.cpp
@@ -1923,10 +1923,13 @@ void UUserWidget::TickActionsAndAnimation(float InDeltaTime)
 	// ones can be added during the tick, if one animation ends and triggers
 	// starting another animation. So iterate with indices always compared 
 	// to the new array size, and not an iterator.
-	for (int32 Index = 0; Index < ActiveAnimations.Num(); ++Index)
+	for (int32 Index = 0; Index < ActiveAnimations.GetMaxIndex(); ++Index)
 	{
-		FWidgetAnimationState& State(ActiveAnimations[Index]);
-		State.Tick(InDeltaTime);
+		if (ActiveAnimations.IsAllocated(Index))
+		{
+			FWidgetAnimationState& State = ActiveAnimations[Index];
+			State.Tick(InDeltaTime);
+		}
 	}
 }
 
