UserWidget: flush queued animations before paint

diff --git a/Engine/Source/Runtime/UMG/Private/UserWidget.cpp b/Engine/Source/Runtime/UMG/Private/UserWidget.cpp
--- a/Engine/Source/Runtime/UMG/Private/UserWidget.cpp
+++ b/Engine/Source/Runtime/UMG/Private/UserWidget.cpp
@@ -663,6 +663,12 @@ void UUserWidget::QueuePlayAnimation(UWidgetAnimation* InAnimation, float StartA
 	QueuedTransition.PlaybackSpeed = PlaybackSpeed;
 	QueuedTransition.bRestoreState = bRestoreState;
 
+	if (!AnimationTickManager)
+	{
+		AnimationTickManager = UUMGSequenceTickManager::Get(this);
+		AnimationTickManager->AddWidget(this);
+	}
+
 	UpdateCanTick();
 }
 
@@ -685,6 +691,12 @@ void UUserWidget::QueuePlayAnimationTimeRange(UWidgetAnimation* InAnimation, flo
 	QueuedTransition.PlaybackSpeed = PlaybackSpeed;
 	QueuedTransition.bRestoreState = bRestoreState;
 
+	if (!AnimationTickManager)
+	{
+		AnimationTickManager = UUMGSequenceTickManager::Get(this);
+		AnimationTickManager->AddWidget(this);
+	}
+
 	UpdateCanTick();
 }
 
@@ -703,6 +715,12 @@ void UUserWidget::QueuePlayAnimationForward(UWidgetAnimation* InAnimation, float
 	QueuedTransition.PlaybackSpeed = PlaybackSpeed;
 	QueuedTransition.bRestoreState = bRestoreState;
 
+	if (!AnimationTickManager)
+	{
+		AnimationTickManager = UUMGSequenceTickManager::Get(this);
+		AnimationTickManager->AddWidget(this);
+	}
+
 	UpdateCanTick();
 }
 
@@ -721,6 +739,12 @@ void UUserWidget::QueuePlayAnimationReverse(UWidgetAnimation* InAnimation, float
 	QueuedTransition.PlaybackSpeed = PlaybackSpeed;
 	QueuedTransition.bRestoreState = bRestoreState;
 
+	if (!AnimationTickManager)
+	{
+		AnimationTickManager = UUMGSequenceTickManager::Get(this);
+		AnimationTickManager->AddWidget(this);
+	}
+
 	UpdateCanTick();
 }
 
@@ -1889,7 +1913,6 @@ void UUserWidget::NativeTick(const FGeometry& MyGeometry, float InDeltaTime)
 #endif
 		if (bTickAnimations)
 		{
-			ExecuteQueuedAnimationTransitions();
 
 			if (AnimationTickManager)
 			{
@@ -1919,6 +1942,8 @@ void UUserWidget::TickActionsAndAnimation(float InDeltaTime)
 		return;
 	}
 
+	ExecuteQueuedAnimationTransitions();
+
 	// Update active widget animations. None will be removed here, but new
 	// ones can be added during the tick, if one animation ends and triggers
 	// starting another animation. So iterate with indices always compared 
