UMG: fixes for more solid handling of re-entrancy

- Replace sparse array of active animations array of shared references.
- Remove the increasing number of hacks required to handle the case of that array being re-allocated.
- Since active animations are shared references at unknown indices, change their handle to simply be the shared pointer.
- Improve on the recently added code that handles re-entrancy in the UMG tick manager. This new version supports ticking newly added widgets without waiting for next frame.

diff --git a/Engine/Source/Runtime/UMG/Private/Animation/UMGSequenceTickManager.cpp b/Engine/Source/Runtime/UMG/Private/Animation/UMGSequenceTickManager.cpp
--- a/Engine/Source/Runtime/UMG/Private/Animation/UMGSequenceTickManager.cpp
+++ b/Engine/Source/Runtime/UMG/Private/Animation/UMGSequenceTickManager.cpp
@@ -74,7 +74,7 @@ void UUMGSequenceTickManager::AddWidget(UUserWidget* InWidget)
 	}
 	else if (bIsTicking)
 	{
-		PendingWeakUserWidgetData.Add(WeakWidget);
+		PendingUserWidgets.Add(InWidget);
 	}
 	else
 	{
@@ -97,7 +97,7 @@ void UUMGSequenceTickManager::OnWidgetTicked(UUserWidget* InWidget)
 	}
 	else if (bIsTicking)
 	{
-		PendingWeakUserWidgetData.Add(InWidget);
+		PendingUserWidgets.Add(InWidget);
 	}
 	else
 	{
@@ -162,81 +162,37 @@ void UUMGSequenceTickManager::TickWidgetAnimations(float DeltaSeconds)
 			UUserWidget* UserWidget = WidgetIter.Key().Get();
 			FSequenceTickManagerWidgetData& WidgetData = WidgetIter.Value();
 
-			WidgetData.bActionsAndAnimationTicked = false;
-
-			if (!UserWidget)
+			const bool bRemoveWidget = TickWidgetAnimations(DeltaSeconds, UserWidget, WidgetData, bIsCurrentlyEvaluating);
+			if (bRemoveWidget)
 			{
 				WidgetIter.RemoveCurrent();
 			}
-			else if (!UserWidget->IsConstructed())
-			{
-				if (!bIsCurrentlyEvaluating)
-				{
-					// Tear down any animations that are not currently being stopped
-					UserWidget->ConditionalTearDownAnimations();
-					UserWidget->UpdateCanTick();
-
-					// If there are no more animations playing, we can remove this widget altogether
-					if (!UserWidget->IsAnyAnimationPlaying())
-					{
-						// Resetting the animation tick manager is ok here because TearDownAnimations will always clear out all animations
-						UserWidget->AnimationTickManager = nullptr;
+		}
+	}
 
-						WidgetIter.RemoveCurrent();
-					}
-				}
-			}
-			else if (!WidgetData.bIsTicking)
-			{
-				// If this widget has not told us it is ticking, we disable animations for that widget.
-				// Once it ticks again, the animation will be updated naturally, and doesn't need anything re-enabling.
-				// 
-				// @todo: There is a chance that relative animations hitting this code path will resume with
-				// different relative bases due to the way the ecs data is destroyed and re-created.
-				// In order to fix this we would have to annex that data instead of destroying it.
-				if (!bIsCurrentlyEvaluating)
-				{
-					UserWidget->DisableAnimations();
+	// We may have seen new animations start on new widgets during the previous loop. To prevent the
+	// WeakUserWidgetData map from being modified during iteration, any such newly animated widgets would
+	// have been placed in PendingUserWidgets. We need to process them now until nothing new is added.
+	while (PendingUserWidgets.Num() > 0)
+	{
+		TArray<TObjectPtr<UUserWidget>> NewUserWidgets;
+		Swap(PendingUserWidgets, NewUserWidgets);
 
-					// Do not null out UUserWidget::AnimationTickManager because although we removed animation _data_
-					// the animations themselves are still playing. As such any UUMGSequencePlayers may hold a reference to this
-					// tick manager's linker, and therefore also need to keep this tick manager alive since the linker is not outered to this tick manager
+		for (TObjectPtr<UUserWidget> NewWidget : NewUserWidgets)
+		{
+			FSequenceTickManagerWidgetData NewWidgetData;
 
-					WidgetIter.RemoveCurrent();
-				}
-			}
-			else
+			const bool bRemoveWidget = TickWidgetAnimations(DeltaSeconds, NewWidget, NewWidgetData, bIsCurrentlyEvaluating);
+			if (!bRemoveWidget)
 			{
-				SCOPE_CYCLE_UOBJECT(WidgetContextScope, UserWidget);
-
-	#if WITH_EDITOR
-				const bool bTickAnimations = !UserWidget->IsDesignTime();
-	#else
-				const bool bTickAnimations = true;
-	#endif
-				if (bTickAnimations && UserWidget->IsVisible())
-				{
-					UserWidget->TickActionsAndAnimation(DeltaSeconds);
-					WidgetData.bActionsAndAnimationTicked = true;
-				}
-
-				// Assume this widget will no longer tick, until we're told otherwise by way of OnWidgetTicked
-				WidgetData.bIsTicking = false;
+				WeakUserWidgetData.Add(NewWidget, NewWidgetData);
 			}
 		}
 	}
+	ensure(PendingUserWidgets.IsEmpty());
 
 	ForceFlush();
 
-	if (PendingWeakUserWidgetData.Num() > 0)
-	{
-		for (TWeakObjectPtr<UUserWidget> WeakWidget : PendingWeakUserWidgetData)
-		{
-			WeakUserWidgetData.Add(WeakWidget, FSequenceTickManagerWidgetData());
-		}
-		PendingWeakUserWidgetData.Reset();
-	}
-
 	if (!Runner->IsCurrentlyEvaluating())
 	{
 		for (auto WidgetIter = WeakUserWidgetData.CreateIterator(); WidgetIter; ++WidgetIter)
@@ -264,6 +220,83 @@ void UUMGSequenceTickManager::TickWidgetAnimations(float DeltaSeconds)
 	WeakUserWidgetData.Shrink();
 }
 
+bool UUMGSequenceTickManager::TickWidgetAnimations(
+		float DeltaSeconds, 
+		UUserWidget* UserWidget, 
+		FSequenceTickManagerWidgetData& UserWidgetData,
+		bool bIsCurrentlyEvaluating)
+{
+	UserWidgetData.bActionsAndAnimationTicked = false;
+
+	if (!UserWidget)
+	{
+		// Remove this widget.
+		return true;
+	}
+	else if (!UserWidget->IsConstructed())
+	{
+		if (!bIsCurrentlyEvaluating)
+		{
+			// Tear down any animations that are not currently being stopped.
+			UserWidget->ConditionalTearDownAnimations();
+			UserWidget->UpdateCanTick();
+
+			// If there are no more animations playing, we can remove this widget altogether.
+			if (!UserWidget->IsAnyAnimationPlaying())
+			{
+				// Resetting the animation tick manager is ok here because TearDownAnimations will always 
+				// clear out all animations.
+				UserWidget->AnimationTickManager = nullptr;
+
+				// Remove this widget.
+				return true;
+			}
+		}
+	}
+	else if (!UserWidgetData.bIsTicking)
+	{
+		// If this widget has not told us it is ticking, we disable animations for that widget.
+		// Once it ticks again, the animation will be updated naturally, and doesn't need anything re-enabling.
+		// 
+		// @todo: There is a chance that relative animations hitting this code path will resume with
+		// different relative bases due to the way the ecs data is destroyed and re-created.
+		// In order to fix this we would have to annex that data instead of destroying it.
+		if (!bIsCurrentlyEvaluating)
+		{
+			UserWidget->DisableAnimations();
+
+			// Do not null out UUserWidget::AnimationTickManager because although we removed animation _data_
+			// the animations themselves are still playing. As such any UUMGSequencePlayers may hold a reference 
+			// to this tick manager's linker, and therefore also need to keep this tick manager alive since the 
+			// linker is not outered to this tick manager.
+
+			// Remove this widget.
+			return true;
+		}
+	}
+	else
+	{
+		SCOPE_CYCLE_UOBJECT(WidgetContextScope, UserWidget);
+
+#if WITH_EDITOR
+		const bool bTickAnimations = !UserWidget->IsDesignTime();
+#else
+		const bool bTickAnimations = true;
+#endif
+		if (bTickAnimations && UserWidget->IsVisible())
+		{
+			UserWidget->TickActionsAndAnimation(DeltaSeconds);
+			UserWidgetData.bActionsAndAnimationTicked = true;
+		}
+
+		// Assume this widget will no longer tick, until we're told otherwise by way of OnWidgetTicked.
+		UserWidgetData.bIsTicking = false;
+	}
+
+	// Don't remove this widget.
+	return false;
+}
+
 void UUMGSequenceTickManager::ForceFlush()
 {
 	Runner->Flush(UE::UMG::GAnimationBudgetMs);
diff --git a/Engine/Source/Runtime/UMG/Private/Animation/WidgetAnimationHandle.cpp b/Engine/Source/Runtime/UMG/Private/Animation/WidgetAnimationHandle.cpp
--- a/Engine/Source/Runtime/UMG/Private/Animation/WidgetAnimationHandle.cpp
+++ b/Engine/Source/Runtime/UMG/Private/Animation/WidgetAnimationHandle.cpp
@@ -6,21 +6,20 @@
 #include "Animation/WidgetAnimationState.h"
 #include "Blueprint/UserWidget.h"
 #include "EdGraph/EdGraphPin.h"
+#include "Templates/SharedPointer.h"
 
 FWidgetAnimationHandle::FWidgetAnimationHandle()
 {
 }
 
-FWidgetAnimationHandle::FWidgetAnimationHandle(UUserWidget* InUserWidget, int32 InStateIndex, uint32 InStateSerial)
-	: WeakUserWidget(InUserWidget)
-	, StateIndex(InStateIndex)
-	, StateSerial(InStateSerial)
+FWidgetAnimationHandle::FWidgetAnimationHandle(TSharedPtr<FWidgetAnimationState> InState)
+	: WeakState(InState)
 {
 }
 
 bool FWidgetAnimationHandle::IsValid() const
 {
-	return (WeakUserWidget.IsValid() && StateIndex >= 0);
+	return WeakState.IsValid();
 }
 
 UUMGSequencePlayer* FWidgetAnimationHandle::GetSequencePlayer() const
@@ -34,18 +33,12 @@ UUMGSequencePlayer* FWidgetAnimationHandle::GetSequencePlayer() const
 
 FWidgetAnimationState* FWidgetAnimationHandle::GetAnimationState() const
 {
-	if (UUserWidget* UserWidget = WeakUserWidget.Get())
-	{
-		if (UserWidget->ActiveAnimations.IsValidIndex(StateIndex))
-		{
-			FWidgetAnimationState& ActiveState = UserWidget->ActiveAnimations[StateIndex];
-			if (ActiveState.SerialNumber == StateSerial)
-			{
-				return &ActiveState;
-			}
-		}
-	}
-	return nullptr;
+	return WeakState.Pin().Get();
+}
+
+TSharedPtr<FWidgetAnimationState> FWidgetAnimationHandle::PinAnimationState() const
+{
+	return WeakState.Pin();
 }
 
 FName FWidgetAnimationHandle::GetUserTag() const
diff --git a/Engine/Source/Runtime/UMG/Private/Animation/WidgetAnimationState.cpp b/Engine/Source/Runtime/UMG/Private/Animation/WidgetAnimationState.cpp
--- a/Engine/Source/Runtime/UMG/Private/Animation/WidgetAnimationState.cpp
+++ b/Engine/Source/Runtime/UMG/Private/Animation/WidgetAnimationState.cpp
@@ -63,7 +63,7 @@ FWidgetAnimationState::FWidgetAnimationState()
 	bIsPendingDelete = false;
 }
 
-void FWidgetAnimationState::Initialize(UWidgetAnimation* InAnimation, UUserWidget* InUserWidget, int32 InIndex, int32 InSerialNumber)
+void FWidgetAnimationState::Initialize(UWidgetAnimation* InAnimation, UUserWidget* InUserWidget)
 {
 	using namespace UE::UMG;
 	using namespace UE::MovieScene;
@@ -72,8 +72,6 @@ void FWidgetAnimationState::Initialize(UWidgetAnimation* InAnimation, UUserWidge
 
 	Animation = InAnimation;
 	WeakUserWidget = InUserWidget;
-	StateIndex = InIndex;
-	SerialNumber = InSerialNumber;
 
 	check(InUserWidget->AnimationTickManager);
 
@@ -115,8 +113,8 @@ void FWidgetAnimationState::Initialize(UWidgetAnimation* InAnimation, UUserWidge
 #if !NO_LOGGING
 	if (UE::UMG::GVarAnimationDetailedLogging)
 	{
-		UE_LOG(LogUMG, Log, TEXT("Animation: Initializing '%s' [%d], instance [%d]%s"), 
-				*GetNameSafe(Animation), StateIndex, RootInstanceHandle.InstanceID,
+		UE_LOG(LogUMG, Log, TEXT("Animation: Initializing '%s', instance [%d]%s"), 
+				*GetNameSafe(Animation), RootInstanceHandle.InstanceID,
 				(LegacyPlayer ? TEXT(", with legacy player") : TEXT("")));
 	}
 #endif
@@ -179,8 +177,8 @@ void FWidgetAnimationState::Tick(float InDeltaSeconds)
 #if !NO_LOGGING
 		if (UE::UMG::GVarAnimationDetailedLogging)
 		{
-			UE_LOG(LogUMG, Log, TEXT("Animation: Ticking '%s' [%d] aborted, invalid state."),
-					*GetNameSafe(Animation), StateIndex);
+			UE_LOG(LogUMG, Log, TEXT("Animation: Ticking '%s' aborted, invalid state."),
+					*GetNameSafe(Animation));
 		}
 #endif
 		return;
@@ -192,8 +190,8 @@ void FWidgetAnimationState::Tick(float InDeltaSeconds)
 #if !NO_LOGGING
 		if (UE::UMG::GVarAnimationDetailedLogging)
 		{
-			UE_LOG(LogUMG, Log, TEXT("Animation: Ticking '%s' [%d] aborted, shared playback state has no runner."),
-					*GetNameSafe(Animation), StateIndex);
+			UE_LOG(LogUMG, Log, TEXT("Animation: Ticking '%s' aborted, shared playback state has no runner."),
+					*GetNameSafe(Animation));
 		}
 #endif
 		return;
@@ -225,13 +223,7 @@ void FWidgetAnimationState::Tick(float InDeltaSeconds)
 			ERunnerUpdateFlags UpdateFlags = ERunnerUpdateFlags::None;
 			if (bNeedsFinalUpdate && Index == TickContexts.Num() - 1)
 			{
-				// WARNING: widget animation states are held inside a TSparseArray by UUserWidget. As this array
-				// gets re-allocated, the states also get re-allocated. So we don't want to capture the "this"
-				// pointer in a callback that could be invoked when that pointer has moved. Instead, we use a
-				// static method as the callback, with our handle as an argument. This also handles the case of
-				// a callback invoked on a state that has been removed.
-				FWidgetAnimationHandle ThisHandle = GetAnimationHandle();
-				OnFlushDelegate = FSimpleDelegate::CreateStatic(&FWidgetAnimationState::OnStopped, ThisHandle);
+				OnFlushDelegate = FSimpleDelegate::CreateSP(this, &FWidgetAnimationState::OnStopped);
 				UpdateFlags = ERunnerUpdateFlags::Flush;
 				bIsStopping = true;
 			}
@@ -243,8 +235,8 @@ void FWidgetAnimationState::Tick(float InDeltaSeconds)
 #if !NO_LOGGING
 			if (UE::UMG::GVarAnimationDetailedLogging)
 			{
-				UE_LOG(LogUMG, Log, TEXT("Animation: Ticking '%s' [%d], Time=%s, Status=%s, Direction=%s (update %d)"),
-						*GetNameSafe(Animation), StateIndex, 
+				UE_LOG(LogUMG, Log, TEXT("Animation: Ticking '%s', Time=%s, Status=%s, Direction=%s (update %d)"),
+						*GetNameSafe(Animation),
 						*LexToString(TickContext.GetTime()),
 						*UEnum::GetValueAsString(TickContext.GetStatus()),
 						((TickContext.GetDirection() == EPlayDirection::Forwards) ? TEXT("Forwards") : TEXT("Backwards")),
@@ -256,8 +248,8 @@ void FWidgetAnimationState::Tick(float InDeltaSeconds)
 #if !NO_LOGGING
 		if (UE::UMG::GVarAnimationDetailedLogging)
 		{
-			UE_LOG(LogUMG, Log, TEXT("Animation: Ticking done '%s' [%d], NumLoopsCompleted=%d, PlaybackSpeed=%f"),
-					*GetNameSafe(Animation), StateIndex, 
+			UE_LOG(LogUMG, Log, TEXT("Animation: Ticking done '%s', NumLoopsCompleted=%d, PlaybackSpeed=%f"),
+					*GetNameSafe(Animation),
 					PlaybackManager.GetNumLoopsCompleted(),
 					PlaybackManager.GetPlayRate());
 		}
@@ -270,29 +262,13 @@ void FWidgetAnimationState::Tick(float InDeltaSeconds)
 #if !NO_LOGGING
 		if (UE::UMG::GVarAnimationDetailedLogging)
 		{
-			UE_LOG(LogUMG, Log, TEXT("Animation: Ticking '%s' [%d] skipped, state not playing."),
-					*GetNameSafe(Animation), StateIndex);
+			UE_LOG(LogUMG, Log, TEXT("Animation: Ticking '%s' skipped, state not playing."),
+					*GetNameSafe(Animation));
 		}
 #endif
 	}
 }
 
-void FWidgetAnimationState::OnBegunPlay(FWidgetAnimationHandle InHandle)
-{
-	if (FWidgetAnimationState* State = InHandle.GetAnimationState())
-	{
-		State->OnBegunPlay();
-	}
-}
-
-void FWidgetAnimationState::OnStopped(FWidgetAnimationHandle InHandle)
-{
-	if (FWidgetAnimationState* State = InHandle.GetAnimationState())
-	{
-		State->OnStopped();
-	}
-}
-
 void FWidgetAnimationState::OnBegunPlay()
 {
 	using namespace UE::MovieScene;
@@ -302,8 +278,8 @@ void FWidgetAnimationState::OnBegunPlay()
 #if !NO_LOGGING
 	if (UE::UMG::GVarAnimationDetailedLogging)
 	{
-		UE_LOG(LogUMG, Log, TEXT("Animation: OnBegunPlay '%s' [%d]"),
-				*GetNameSafe(Animation), StateIndex);
+		UE_LOG(LogUMG, Log, TEXT("Animation: OnBegunPlay '%s'"),
+				*GetNameSafe(Animation));
 	}
 #endif
 }
@@ -340,16 +316,14 @@ void FWidgetAnimationState::OnStopped()
 	if (UE::UMG::GVarAnimationDetailedLogging)
 	{
 		PRAGMA_DISABLE_DEPRECATION_WARNINGS
-		UE_LOG(LogUMG, Log, TEXT("Animation: OnStopped '%s' [%d]%s%s"),
-				*GetNameSafe(Animation), StateIndex,
+		UE_LOG(LogUMG, Log, TEXT("Animation: OnStopped '%s' %s%s"),
+				*GetNameSafe(Animation),
 				OnWidgetAnimationFinishedEvent.IsBound() ? TEXT(", OnWidgetAnimationFinishedEvent bound") : TEXT(""),
 				(LegacyPlayer && LegacyPlayer->OnSequenceFinishedPlaying().IsBound()) ? TEXT(", OnSequenceFinishedPlaying bound") : TEXT(""));
 		PRAGMA_ENABLE_DEPRECATION_WARNINGS
 	}
 #endif
 
-	FWidgetAnimationHandle ThisHandle = GetAnimationHandle();
-
 	UUserWidget* UserWidget = WeakUserWidget.Get();
 	if (UserWidget)
 	{
@@ -361,18 +335,11 @@ void FWidgetAnimationState::OnStopped()
 		UserWidget->BroadcastAnimationFinishedPlaying(*this);
 	}
 
-	// WARNING
-	// Don't access any more member fields after this... the call to BroadcastAnimationFinishedPlaying
-	// may have started new animations, which might have reallocated our parent widget's ActiveAnimations
-	// array, and rendering us invalid.
-	if (FWidgetAnimationState* NewThis = ThisHandle.GetAnimationState())
-	{
-		NewThis->OnWidgetAnimationFinishedEvent.Broadcast(*NewThis);
+	OnWidgetAnimationFinishedEvent.Broadcast(*this);
 
-		if (NewThis->LegacyPlayer)
-		{
-			NewThis->LegacyPlayer->BroadcastSequenceFinishedPlaying();
-		}
+	if (LegacyPlayer)
+	{
+		LegacyPlayer->BroadcastSequenceFinishedPlaying();
 	}
 }
 
@@ -396,7 +363,8 @@ UUMGSequencePlayer* FWidgetAnimationState::GetLegacyPlayer() const
 
 FWidgetAnimationHandle FWidgetAnimationState::GetAnimationHandle() const
 {
-	return FWidgetAnimationHandle(WeakUserWidget.Get(), StateIndex, SerialNumber);
+	FWidgetAnimationState* MutableThis = const_cast<FWidgetAnimationState*>(this);
+	return FWidgetAnimationHandle(SharedThis(MutableThis));
 }
 
 bool FWidgetAnimationState::IsPlayingForward() const
@@ -541,21 +509,19 @@ void FWidgetAnimationState::Play(const FWidgetAnimationStatePlayParams& PlayPara
 	TSharedPtr<FMovieSceneEntitySystemRunner> Runner = SharedPlaybackState->GetRunner();
 	if (Runner)
 	{
-		// See WARNING in Tick() method.
-		FWidgetAnimationHandle ThisHandle = GetAnimationHandle();
 		const FMovieSceneContext FirstContext = PlaybackManager.UpdateAtCurrentTime();
 		Runner->QueueUpdate(
 				FirstContext, 
 				SharedPlaybackState->GetRootInstanceHandle(),
-				FSimpleDelegate::CreateStatic(&FWidgetAnimationState::OnBegunPlay, ThisHandle),
+				FSimpleDelegate::CreateSP(this, &FWidgetAnimationState::OnBegunPlay),
 				ERunnerUpdateFlags::Flush);
 	}
 
 #if !NO_LOGGING
 	if (UE::UMG::GVarAnimationDetailedLogging)
 	{
-		UE_LOG(LogUMG, Log, TEXT("Animation: Play '%s' [%d], StartTime=%s, NumLoopsToPlay=%d, PlaybackSpeed=%f, PlayMode=%s%s"),
-				*GetNameSafe(Animation), StateIndex,
+		UE_LOG(LogUMG, Log, TEXT("Animation: Play '%s', StartTime=%s, NumLoopsToPlay=%d, PlaybackSpeed=%f, PlayMode=%s%s"),
+				*GetNameSafe(Animation),
 				*LexToString(PlaybackManager.GetCurrentTime()),
 				PlaybackManager.GetNumLoopsToPlay(),
 				PlaybackManager.GetPlayRate(),
@@ -592,13 +558,11 @@ void FWidgetAnimationState::Stop()
 		{
 			TSharedPtr<FMovieSceneEntitySystemRunner> Runner = SharedPlaybackState->GetRunner();
 
-			// See WARNING in Tick() method.
-			FWidgetAnimationHandle ThisHandle = GetAnimationHandle();
 			const FMovieSceneContext ReturnToStartContext = PlaybackManager.UpdateAtCurrentTime();
 			Runner->QueueUpdate(
 					ReturnToStartContext,
 					SharedPlaybackState->GetRootInstanceHandle(),
-					FSimpleDelegate::CreateStatic(&FWidgetAnimationState::OnStopped, ThisHandle),
+					FSimpleDelegate::CreateSP(this, &FWidgetAnimationState::OnStopped),
 					ERunnerUpdateFlags::Flush);
 
 			bIsStopping = true;
@@ -617,8 +581,8 @@ void FWidgetAnimationState::Stop()
 #if !NO_LOGGING
 	if (UE::UMG::GVarAnimationDetailedLogging)
 	{
-		UE_LOG(LogUMG, Log, TEXT("Animation: Stop '%s' [%d]"),
-				*GetNameSafe(Animation), StateIndex);
+		UE_LOG(LogUMG, Log, TEXT("Animation: Stop '%s'"),
+				*GetNameSafe(Animation));
 	}
 #endif
 
@@ -651,8 +615,8 @@ void FWidgetAnimationState::Pause()
 #if !NO_LOGGING
 	if (UE::UMG::GVarAnimationDetailedLogging)
 	{
-		UE_LOG(LogUMG, Log, TEXT("Animation: Pause '%s' [%d], PauseTime=%s"),
-				*GetNameSafe(Animation), StateIndex,
+		UE_LOG(LogUMG, Log, TEXT("Animation: Pause '%s', PauseTime=%s"),
+				*GetNameSafe(Animation),
 				*LexToString(PlaybackManager.GetCurrentTime()));
 	}
 #endif
@@ -695,8 +659,8 @@ void FWidgetAnimationState::RemoveEvaluationData()
 #if !NO_LOGGING
 	if (UE::UMG::GVarAnimationDetailedLogging)
 	{
-		UE_LOG(LogUMG, Log, TEXT("Animation: RemoveEvaluationData '%s' [%d]"),
-				*GetNameSafe(Animation), StateIndex);
+		UE_LOG(LogUMG, Log, TEXT("Animation: RemoveEvaluationData '%s'"),
+				*GetNameSafe(Animation));
 	}
 #endif
 }
@@ -733,8 +697,8 @@ void FWidgetAnimationState::TearDown()
 #if !NO_LOGGING
 	if (UE::UMG::GVarAnimationDetailedLogging)
 	{
-		UE_LOG(LogUMG, Log, TEXT("Animation: TearDown '%s' [%d]"),
-				*GetNameSafe(Animation), StateIndex);
+		UE_LOG(LogUMG, Log, TEXT("Animation: TearDown '%s'"),
+				*GetNameSafe(Animation));
 	}
 #endif
 }
diff --git a/Engine/Source/Runtime/UMG/Private/UserWidget.cpp b/Engine/Source/Runtime/UMG/Private/UserWidget.cpp
--- a/Engine/Source/Runtime/UMG/Private/UserWidget.cpp
+++ b/Engine/Source/Runtime/UMG/Private/UserWidget.cpp
@@ -483,35 +483,33 @@ UWorld* UUserWidget::GetWorld() const
 	return nullptr;
 }
 
-const FWidgetAnimationState* UUserWidget::GetAnimationState(const UWidgetAnimation* InAnimation) const
+TSharedPtr<const FWidgetAnimationState> UUserWidget::GetAnimationState(const UWidgetAnimation* InAnimation) const
 {
-	const FWidgetAnimationState* FoundState = nullptr;
 	for (auto It = ActiveAnimations.CreateConstIterator(); It; ++It)
 	{
-		if (It->GetAnimation() == InAnimation && It->IsValid())
+		TSharedRef<FWidgetAnimationState> State(*It);
+		if (State->GetAnimation() == InAnimation && State->IsValid())
 		{
-			FoundState = &(*It);
-			break;
+			return State;
 		}
 	};
-	return FoundState;
+	return nullptr;
 }
 
-FWidgetAnimationState* UUserWidget::GetAnimationState(const UWidgetAnimation* InAnimation)
+TSharedPtr<FWidgetAnimationState> UUserWidget::GetAnimationState(const UWidgetAnimation* InAnimation)
 {
-	FWidgetAnimationState* FoundState = nullptr;
-	for (auto It = ActiveAnimations.CreateIterator(); It; ++It)
+	for (auto It = ActiveAnimations.CreateConstIterator(); It; ++It)
 	{
-		if (It->GetAnimation() == InAnimation && It->IsValid())
+		TSharedRef<FWidgetAnimationState> State(*It);
+		if (State->GetAnimation() == InAnimation && State->IsValid())
 		{
-			FoundState = &(*It);
-			break;
+			return State;
 		}
 	};
-	return FoundState;
+	return nullptr;
 }
 
-FWidgetAnimationState* UUserWidget::GetOrAddAnimationState(UWidgetAnimation* InAnimation)
+TSharedPtr<FWidgetAnimationState> UUserWidget::GetOrAddAnimationState(UWidgetAnimation* InAnimation)
 {
 	if (InAnimation && !bStoppingAllAnimations)
 	{
@@ -524,15 +522,15 @@ FWidgetAnimationState* UUserWidget::GetOrAddAnimationState(UWidgetAnimation* InA
 		// If this widget is currently offscreen it would very well not be being tracked by the tick manager
 		AnimationTickManager->AddWidget(this);
 
-		FWidgetAnimationState* FoundState = GetAnimationState(InAnimation);
+		TSharedPtr<FWidgetAnimationState> FoundState = GetAnimationState(InAnimation);
 		if (!FoundState || FoundState->IsPendingDelete())
 		{
 			// Create a new state and initialize it.
-			int32 NewStateIndex = ActiveAnimations.Emplace();
-			FWidgetAnimationState& NewState = ActiveAnimations[NewStateIndex];
-			NewState.Initialize(InAnimation, this, NewStateIndex, NextSerialNumber++);
+			TSharedRef<FWidgetAnimationState> NewState = MakeShared<FWidgetAnimationState>();
+			ActiveAnimations.Add(NewState);
+			NewState->Initialize(InAnimation, this);
 
-			return &NewState;
+			return NewState;
 		}
 		else
 		{
@@ -597,26 +595,30 @@ void UUserWidget::ExecuteQueuedAnimationTransitions()
 
 void UUserWidget::ConditionalTearDownAnimations()
 {
+	ensureMsgf(!bIsTickingAnimations, TEXT("When active animations are ticked, animations can only be added, not removed"));
+
 	for (auto It = ActiveAnimations.CreateIterator(); It; ++It)
 	{
-		FWidgetAnimationState& State(*It);
-		if (!State.IsValid())
+		TSharedRef<FWidgetAnimationState> State(*It);
+		if (!State->IsValid())
 		{
 			It.RemoveCurrent();
 		}
-		else if (!State.IsStopping())
+		else if (!State->IsStopping())
 		{
-			State.TearDown();
+			State->TearDown();
 		}
 	}
 }
 
 void UUserWidget::TearDownAnimations()
 {
+	ensureMsgf(!bIsTickingAnimations, TEXT("When active animations are ticked, animations can only be added, not removed"));
+
 	for (auto It = ActiveAnimations.CreateIterator(); It; ++It)
 	{
-		FWidgetAnimationState& State(*It);
-		State.TearDown();
+		TSharedRef<FWidgetAnimationState> State(*It);
+		State->TearDown();
 	}
 
 	ActiveAnimations.Reset();
@@ -626,7 +628,8 @@ void UUserWidget::DisableAnimations()
 {
 	for (auto It = ActiveAnimations.CreateIterator(); It; ++It)
 	{
-		It->RemoveEvaluationData();
+		TSharedRef<FWidgetAnimationState> State(*It);
+		State->RemoveEvaluationData();
 	}
 }
 
@@ -771,11 +774,11 @@ void UUserWidget::QueueStopAllAnimations()
 		QueuedWidgetAnimationTransition.TransitionMode = EQueuedWidgetAnimationMode::Stop;
 	}
 
-	for (FWidgetAnimationState& AnimationState : ActiveAnimations)
+	for (TSharedRef<FWidgetAnimationState> State : ActiveAnimations)
 	{
-		if (AnimationState.GetPlaybackStatus() == EMovieScenePlayerStatus::Playing)
+		if (State->GetPlaybackStatus() == EMovieScenePlayerStatus::Playing)
 		{
-			QueueStopAnimation(AnimationState.GetAnimation());
+			QueueStopAnimation(State->GetAnimation());
 		}
 	}
 
@@ -794,7 +797,7 @@ float UUserWidget::QueuePauseAnimation(const UWidgetAnimation* InAnimation)
 
 		UpdateCanTick();
 
-		if (FWidgetAnimationState* FoundState = GetAnimationState(InAnimation))
+		if (TSharedPtr<FWidgetAnimationState> FoundState = GetAnimationState(InAnimation))
 		{
 			return (float)FoundState->GetCurrentTime().AsSeconds();
 		}
@@ -807,7 +810,7 @@ FWidgetAnimationHandle UUserWidget::PlayAnimation(UWidgetAnimation* InAnimation,
 {
 	SCOPED_NAMED_EVENT_TEXT("Widget::PlayAnimation", FColor::Emerald);
 
-	if (FWidgetAnimationState* AnimationState = GetOrAddAnimationState(InAnimation))
+	if (TSharedPtr<FWidgetAnimationState> AnimationState = GetOrAddAnimationState(InAnimation))
 	{
 		FWidgetAnimationStatePlayParams PlayParams;
 		PlayParams.StartAtTime = StartAtTime;
@@ -822,7 +825,7 @@ FWidgetAnimationHandle UUserWidget::PlayAnimation(UWidgetAnimation* InAnimation,
 
 		UpdateCanTick();
 
-		return FWidgetAnimationHandle(this, AnimationState->StateIndex, AnimationState->SerialNumber);
+		return FWidgetAnimationHandle(AnimationState);
 	}
 
 	return FWidgetAnimationHandle();
@@ -832,7 +835,7 @@ FWidgetAnimationHandle UUserWidget::PlayAnimationTimeRange(UWidgetAnimation* InA
 {
 	SCOPED_NAMED_EVENT_TEXT("Widget::PlayAnimationTimeRange", FColor::Emerald);
 
-	if (FWidgetAnimationState* AnimationState = GetOrAddAnimationState(InAnimation))
+	if (TSharedPtr<FWidgetAnimationState> AnimationState = GetOrAddAnimationState(InAnimation))
 	{
 		FWidgetAnimationStatePlayParams PlayParams;
 		PlayParams.StartAtTime = StartAtTime;
@@ -848,7 +851,7 @@ FWidgetAnimationHandle UUserWidget::PlayAnimationTimeRange(UWidgetAnimation* InA
 
 		UpdateCanTick();
 
-		return FWidgetAnimationHandle(this, AnimationState->StateIndex, AnimationState->SerialNumber);
+		return FWidgetAnimationHandle(AnimationState);
 	}
 
 	return FWidgetAnimationHandle();
@@ -856,7 +859,7 @@ FWidgetAnimationHandle UUserWidget::PlayAnimationTimeRange(UWidgetAnimation* InA
 
 FWidgetAnimationHandle UUserWidget::PlayAnimationForward(UWidgetAnimation* InAnimation, float PlaybackSpeed, bool bRestoreState)
 {
-	FWidgetAnimationState* AnimationState = GetAnimationState(InAnimation);
+	TSharedPtr<FWidgetAnimationState> AnimationState = GetAnimationState(InAnimation);
 
 	if (AnimationState && AnimationState->GetPlaybackStatus() == EMovieScenePlayerStatus::Playing)
 	{
@@ -865,7 +868,7 @@ FWidgetAnimationHandle UUserWidget::PlayAnimationForward(UWidgetAnimation* InAni
 			AnimationState->Reverse();
 		}
 
-		return FWidgetAnimationHandle(this, AnimationState->StateIndex, AnimationState->SerialNumber);
+		return FWidgetAnimationHandle(AnimationState);
 	}
 
 	return PlayAnimation(InAnimation, 0.0f, 1, EUMGSequencePlayMode::Forward, PlaybackSpeed, bRestoreState);
@@ -873,7 +876,7 @@ FWidgetAnimationHandle UUserWidget::PlayAnimationForward(UWidgetAnimation* InAni
 
 FWidgetAnimationHandle UUserWidget::PlayAnimationReverse(UWidgetAnimation* InAnimation, float PlaybackSpeed, bool bRestoreState)
 {
-	FWidgetAnimationState* AnimationState = GetAnimationState(InAnimation);
+	TSharedPtr<FWidgetAnimationState> AnimationState = GetAnimationState(InAnimation);
 
 	if (AnimationState && AnimationState->GetPlaybackStatus() == EMovieScenePlayerStatus::Playing)
 	{
@@ -882,7 +885,7 @@ FWidgetAnimationHandle UUserWidget::PlayAnimationReverse(UWidgetAnimation* InAni
 			AnimationState->Reverse();
 		}
 
-		return FWidgetAnimationHandle(this, AnimationState->StateIndex, AnimationState->SerialNumber);
+		return FWidgetAnimationHandle(AnimationState);
 	}
 
 	return PlayAnimation(InAnimation, 0.0f, 1, EUMGSequencePlayMode::Reverse, PlaybackSpeed, bRestoreState);
@@ -892,7 +895,7 @@ void UUserWidget::StopAnimation(const UWidgetAnimation* InAnimation)
 {
 	if (InAnimation)
 	{
-		if (FWidgetAnimationState* FoundState = GetAnimationState(InAnimation))
+		if (TSharedPtr<FWidgetAnimationState> FoundState = GetAnimationState(InAnimation))
 		{
 			FoundState->Stop();
 
@@ -905,11 +908,11 @@ void UUserWidget::StopAllAnimations()
 {
 	bStoppingAllAnimations = true;
 
-	for (FWidgetAnimationState& AnimationState : ActiveAnimations)
+	for (TSharedRef<FWidgetAnimationState> State : ActiveAnimations)
 	{
-		if (AnimationState.GetPlaybackStatus() == EMovieScenePlayerStatus::Playing)
+		if (State->GetPlaybackStatus() == EMovieScenePlayerStatus::Playing)
 		{
-			AnimationState.Stop();
+			State->Stop();
 		}
 	}
 
@@ -923,7 +926,7 @@ float UUserWidget::PauseAnimation(const UWidgetAnimation* InAnimation)
 	if (InAnimation)
 	{
 		// @todo UMG sequencer - Restart animations which have had Play called on them?
-		if (FWidgetAnimationState* FoundState = GetAnimationState(InAnimation))
+		if (TSharedPtr<FWidgetAnimationState> FoundState = GetAnimationState(InAnimation))
 		{
 			FoundState->Pause();
 			return (float)FoundState->GetCurrentTime().AsSeconds();
@@ -937,7 +940,7 @@ float UUserWidget::GetAnimationCurrentTime(const UWidgetAnimation* InAnimation)
 {
 	if (InAnimation)
 	{
-		if (const FWidgetAnimationState* FoundState = GetAnimationState(InAnimation))
+		if (TSharedPtr<const FWidgetAnimationState> FoundState = GetAnimationState(InAnimation))
 		{
 			return (float)FoundState->GetCurrentTime().AsSeconds();
 		}
@@ -950,7 +953,7 @@ void UUserWidget::SetAnimationCurrentTime(const UWidgetAnimation* InAnimation, f
 {
 	if (InAnimation)
 	{
-		if (FWidgetAnimationState* FoundState = GetAnimationState(InAnimation))
+		if (TSharedPtr<FWidgetAnimationState> FoundState = GetAnimationState(InAnimation))
 		{
 			FoundState->SetCurrentTime(InTime);
 		}
@@ -961,7 +964,7 @@ bool UUserWidget::IsAnimationPlaying(const UWidgetAnimation* InAnimation) const
 {
 	if (InAnimation)
 	{
-		if (const FWidgetAnimationState* FoundState = GetAnimationState(InAnimation))
+		if (TSharedPtr<const FWidgetAnimationState> FoundState = GetAnimationState(InAnimation))
 		{
 			return FoundState->GetPlaybackStatus() == EMovieScenePlayerStatus::Playing;
 		}
@@ -977,7 +980,7 @@ bool UUserWidget::IsAnyAnimationPlaying() const
 
 void UUserWidget::SetNumLoopsToPlay(const UWidgetAnimation* InAnimation, int32 InNumLoopsToPlay)
 {
-	if (FWidgetAnimationState* FoundState = GetAnimationState(InAnimation))
+	if (TSharedPtr<FWidgetAnimationState> FoundState = GetAnimationState(InAnimation))
 	{
 		FoundState->SetNumLoopsToPlay(InNumLoopsToPlay);
 	}
@@ -985,7 +988,7 @@ void UUserWidget::SetNumLoopsToPlay(const UWidgetAnimation* InAnimation, int32 I
 
 void UUserWidget::SetPlaybackSpeed(const UWidgetAnimation* InAnimation, float PlaybackSpeed)
 {
-	if (FWidgetAnimationState* FoundState = GetAnimationState(InAnimation))
+	if (TSharedPtr<FWidgetAnimationState> FoundState = GetAnimationState(InAnimation))
 	{
 		FoundState->SetPlaybackSpeed(PlaybackSpeed);
 	}
@@ -993,7 +996,7 @@ void UUserWidget::SetPlaybackSpeed(const UWidgetAnimation* InAnimation, float Pl
 
 void UUserWidget::ReverseAnimation(const UWidgetAnimation* InAnimation)
 {
-	if (FWidgetAnimationState* FoundState = GetAnimationState(InAnimation))
+	if (TSharedPtr<FWidgetAnimationState> FoundState = GetAnimationState(InAnimation))
 	{
 		FoundState->Reverse();
 	}
@@ -1019,7 +1022,7 @@ bool UUserWidget::IsAnimationPlayingForward(const UWidgetAnimation* InAnimation)
 {
 	if (InAnimation)
 	{
-		if (FWidgetAnimationState* FoundState = GetAnimationState(InAnimation))
+		if (TSharedPtr<FWidgetAnimationState> FoundState = GetAnimationState(InAnimation))
 		{
 			return FoundState->IsPlayingForward();
 		}
@@ -1942,19 +1945,23 @@ void UUserWidget::TickActionsAndAnimation(float InDeltaTime)
 		return;
 	}
 
+	if (!ensureMsgf(!bIsTickingAnimations, TEXT("Re-entrant animation ticking detected!")))
+	{
+		return;
+	}
+
+	TGuardValue TickingGuard(bIsTickingAnimations, true);
+
 	ExecuteQueuedAnimationTransitions();
 
 	// Update active widget animations. None will be removed here, but new
 	// ones can be added during the tick, if one animation ends and triggers
-	// starting another animation. So iterate with indices always compared 
-	// to the new array size, and not an iterator.
-	for (int32 Index = 0; Index < ActiveAnimations.GetMaxIndex(); ++Index)
+	// starting another animation. So iterate with an index to keep going 
+	// when the array size increases.
+	for (int32 Index = 0; Index < ActiveAnimations.Num(); ++Index)
 	{
-		if (ActiveAnimations.IsAllocated(Index))
-		{
-			FWidgetAnimationState& State = ActiveAnimations[Index];
-			State.Tick(InDeltaTime);
-		}
+		TSharedRef<FWidgetAnimationState> State = ActiveAnimations[Index];
+		State->Tick(InDeltaTime);
 	}
 }
 
@@ -2600,22 +2607,25 @@ void UUserWidget::AddReferencedObjects(UObject* InThis, FReferenceCollector& Col
 
 	for (auto It = TypedThis->ActiveAnimations.CreateIterator(); It; ++It)
 	{
-		It->AddReferencedObjects(Collector);
+		TSharedRef<FWidgetAnimationState> State(*It);
+		State->AddReferencedObjects(Collector);
 	}
 }
 
 void UUserWidget::ClearStoppedAnimationStates()
 {
+	ensureMsgf(!bIsTickingAnimations, TEXT("When active animations are ticked, animations can only be added, not removed"));
+
 	for (auto It = ActiveAnimations.CreateIterator(); It; ++It)
 	{
-		FWidgetAnimationState& State(*It);
-		if (!State.IsValid())
+		TSharedRef<FWidgetAnimationState> State(*It);
+		if (!State->IsValid())
 		{
 			It.RemoveCurrent();
 		}
-		else if ((State.GetPlaybackStatus() == EMovieScenePlayerStatus::Stopped) && !State.IsStopping())
+		else if ((State->GetPlaybackStatus() == EMovieScenePlayerStatus::Stopped) && !State->IsStopping())
 		{
-			State.TearDown();
+			State->TearDown();
 			It.RemoveCurrent();
 		}
 	}
diff --git a/Engine/Source/Runtime/UMG/Public/Animation/UMGSequenceTickManager.h b/Engine/Source/Runtime/UMG/Public/Animation/UMGSequenceTickManager.h
--- a/Engine/Source/Runtime/UMG/Public/Animation/UMGSequenceTickManager.h
+++ b/Engine/Source/Runtime/UMG/Public/Animation/UMGSequenceTickManager.h
@@ -60,6 +60,11 @@ private:
 
 	UMG_API void HandleSlatePostTick(float DeltaSeconds);
 	UMG_API void TickWidgetAnimations(float DeltaSeconds);
+	UMG_API bool TickWidgetAnimations(
+			float DeltaSeconds, 
+			UUserWidget* UserWidget, 
+			FSequenceTickManagerWidgetData& UserWidgetData,
+			bool bIsCurrentlyEvaluating);
 
 private:
 
@@ -67,7 +72,7 @@ private:
 	TMap<TWeakObjectPtr<UUserWidget>, FSequenceTickManagerWidgetData> WeakUserWidgetData;
 
 	UPROPERTY(transient)
-	TArray<TWeakObjectPtr<UUserWidget>> PendingWeakUserWidgetData;
+	TArray<TObjectPtr<UUserWidget>> PendingUserWidgets;
 
 	UPROPERTY(transient)
 	TObjectPtr<UMovieSceneEntitySystemLinker> Linker;
diff --git a/Engine/Source/Runtime/UMG/Public/Animation/WidgetAnimationHandle.h b/Engine/Source/Runtime/UMG/Public/Animation/WidgetAnimationHandle.h
--- a/Engine/Source/Runtime/UMG/Public/Animation/WidgetAnimationHandle.h
+++ b/Engine/Source/Runtime/UMG/Public/Animation/WidgetAnimationHandle.h
@@ -4,7 +4,7 @@
 
 #include "CoreTypes.h"
 #include "Kismet/BlueprintFunctionLibrary.h"
-#include "UObject/WeakObjectPtr.h"
+#include "Templates/SharedPointerFwd.h"
 
 #include "WidgetAnimationHandle.generated.h"
 
@@ -30,6 +30,9 @@ public:
 	/** Gets the animation state. */
 	UMG_API FWidgetAnimationState* GetAnimationState() const;
 
+	/** Gets the animation state. */
+	UMG_API TSharedPtr<FWidgetAnimationState> PinAnimationState() const;
+
 public:
 
 	// For backwards compatibility in C++ code.
@@ -55,27 +58,18 @@ public:
 	/** Type hash for using handles in maps and other containers. */
 	friend uint32 GetTypeHash(const FWidgetAnimationHandle& Handle)
 	{
-		return HashCombine(GetTypeHash(Handle.WeakUserWidget), Handle.StateIndex, Handle.StateSerial);
+		return GetTypeHash(Handle.WeakState);
 	}
 
 private:
 
 	/** Creates a handle for the given running animation on the given widget. */
-	UMG_API FWidgetAnimationHandle(UUserWidget* InUserWidget, int32 InStateIndex, uint32 InStateSerial);
+	UMG_API FWidgetAnimationHandle(TSharedPtr<FWidgetAnimationState> InState);
 
 private:
 
-	/** The widget this handle relates to. */
-	UPROPERTY()
-	TWeakObjectPtr<UUserWidget> WeakUserWidget;
-
-	/** The animation state index. */
-	UPROPERTY()
-	int32 StateIndex = INDEX_NONE;
-
-	/** The animation state serial. */
-	UPROPERTY()
-	uint32 StateSerial = 0;
+	/** Weak pointer to the animation state. */
+	TWeakPtr<FWidgetAnimationState> WeakState;
 
 	// Only UUserWidget and FWidgetAnimationState can create handles.
 	friend class UUserWidget;
diff --git a/Engine/Source/Runtime/UMG/Public/Animation/WidgetAnimationState.h b/Engine/Source/Runtime/UMG/Public/Animation/WidgetAnimationState.h
--- a/Engine/Source/Runtime/UMG/Public/Animation/WidgetAnimationState.h
+++ b/Engine/Source/Runtime/UMG/Public/Animation/WidgetAnimationState.h
@@ -10,6 +10,7 @@
 #include "EntitySystem/MovieSceneSharedPlaybackState.h"
 #include "Misc/QualifiedFrameTime.h"
 #include "MovieSceneFwd.h"
+#include "Templates/SharedPointer.h"
 #include "UObject/ObjectMacros.h"
 
 class UUserWidget;
@@ -43,7 +44,7 @@ struct FWidgetAnimationStatePlayParams
 /**
  * A class that runs an animation on a widget.
  */
-struct FWidgetAnimationState
+struct FWidgetAnimationState : public TSharedFromThis<FWidgetAnimationState>
 {
 public:
 
@@ -51,7 +52,7 @@ public:
 	FWidgetAnimationState();
 
 	/** Initializes the widget animation state. */
-	void Initialize(UWidgetAnimation* InAnimation, UUserWidget* InUserWidget, int32 InIndex, int32 InSerialNumber);
+	void Initialize(UWidgetAnimation* InAnimation, UUserWidget* InUserWidget);
 
 	/** Gets the user widget this state is animating. */
 	UUserWidget* GetUserWidget() const { return WeakUserWidget.Get(); }
@@ -145,9 +146,6 @@ private:
 	void OnBegunPlay();
 	void OnStopped();
 
-	static void OnBegunPlay(FWidgetAnimationHandle InHandle);
-	static void OnStopped(FWidgetAnimationHandle InHandle);
-
 private:
 
 	/** Animation being played */
@@ -159,16 +157,6 @@ private:
 	/** The user widget this sequence is animating */
 	TWeakObjectPtr<UUserWidget> WeakUserWidget;
 
-	/** The index of this state. */
-	int32 StateIndex;
-
-	/** The serial number of this state. */
-	int32 SerialNumber;
-
-	// Access to index and serial number.
-	friend class UUserWidget;
-	friend struct FWidgetAnimationHandle;
-
 private:
 
 	/** Shared playback state for the animation. */
diff --git a/Engine/Source/Runtime/UMG/Public/Blueprint/UserWidget.h b/Engine/Source/Runtime/UMG/Public/Blueprint/UserWidget.h
--- a/Engine/Source/Runtime/UMG/Public/Blueprint/UserWidget.h
+++ b/Engine/Source/Runtime/UMG/Public/Blueprint/UserWidget.h
@@ -959,10 +959,10 @@ protected:
 	UE_DEPRECATED(5.6, "UUMGSequencePlayer is deprecated")
 	UMG_API void BroadcastAnimationStateChange(const UUMGSequencePlayer& Player, EWidgetAnimationEvent AnimationEvent);
 
-	UMG_API void BroadcastAnimationStateChange(const FWidgetAnimationState& State, EWidgetAnimationEvent AnimationEvent);
-
 private:
 
+	UMG_API void BroadcastAnimationStateChange(const FWidgetAnimationState& State, EWidgetAnimationEvent AnimationEvent);
+
 	UMG_API void BroadcastAnimationStateChange(const UWidgetAnimation* Animation, FName UserTag, EWidgetAnimationEvent AnimationEvent);
 
 	UMG_API void BroadcastAnimationStartedPlaying(FWidgetAnimationState& State);
@@ -1491,10 +1491,10 @@ public:
 private:
 
 	/** All the sequences currently playing */
-	TSparseArray<FWidgetAnimationState> ActiveAnimations;
+	TArray<TSharedRef<FWidgetAnimationState>> ActiveAnimations;
 
-	/** Next serial number for animation handles */
-	uint32 NextSerialNumber = 0;
+	/** True when we are ticking our active animations */
+	bool bIsTickingAnimations = false;
 
 private:
 
@@ -1635,9 +1635,9 @@ protected:
 
 	UMG_API void RemoveObsoleteBindings(const TArray<FName>& NamedSlots);
 
-	UMG_API const FWidgetAnimationState* GetAnimationState(const UWidgetAnimation* InAnimation) const;
-	UMG_API FWidgetAnimationState* GetAnimationState(const UWidgetAnimation* InAnimation);
-	UMG_API FWidgetAnimationState* GetOrAddAnimationState(UWidgetAnimation* InAnimation);
+	UMG_API TSharedPtr<const FWidgetAnimationState> GetAnimationState(const UWidgetAnimation* InAnimation) const;
+	UMG_API TSharedPtr<FWidgetAnimationState> GetAnimationState(const UWidgetAnimation* InAnimation);
+	UMG_API TSharedPtr<FWidgetAnimationState> GetOrAddAnimationState(UWidgetAnimation* InAnimation);
 
 	UMG_API void ExecuteQueuedAnimationTransitions();
 
