[SObjectWidget]
- Fix race on global value since the check is only a concern for the runtime

diff --git a/Engine/Source/Runtime/UMG/Private/Slate/SObjectWidget.cpp b/Engine/Source/Runtime/UMG/Private/Slate/SObjectWidget.cpp
--- a/Engine/Source/Runtime/UMG/Private/Slate/SObjectWidget.cpp
+++ b/Engine/Source/Runtime/UMG/Private/Slate/SObjectWidget.cpp
@@ -36,16 +36,11 @@ SObjectWidget::SObjectWidget() = default;
 
 SObjectWidget::~SObjectWidget()
 {
-#if SLATE_VERBOSE_NAMED_EVENTS
-	// This can happen during blueprint compiling, so just ignore it if it happens then, this is really only a concern
-	// in a running game.
-	if (!GCompilingBlueprint && !GIsGCingAfterBlueprintCompile)
+#if SLATE_VERBOSE_NAMED_EVENTS && !WITH_EDITOR
+	// This is only a concern during a running game - design-time instances can be destroyed from GC quite often when recompiling
+	if (!WidgetObject || !WidgetObject->IsDesignTime())
 	{
-		// This is only a concern during a running game - design-time instances can be destroyed from GC quite often when recompiling
-		if (!WidgetObject || !WidgetObject->IsDesignTime())
-		{
-			ensureMsgf(GExitPurge || !IsGarbageCollecting(), TEXT("SObjectWidget for '%s' destroyed while collecting garbage.  This can lead to multiple GCs being required to cleanup the object.  Possible causes might be,\n1) ReleaseSlateResources not being implemented for the owner of this pointer.\n2) You may just be holding onto some slate pointers on an actor that don't get reset until the actor is Garbage Collected.  You should avoid doing this, and instead reset those references when the actor is Destroyed."), *DebugName);
-		}
+		ensureMsgf(GExitPurge || !IsGarbageCollecting(), TEXT("SObjectWidget for '%s' destroyed while collecting garbage.  This can lead to multiple GCs being required to cleanup the object.  Possible causes might be,\n1) ReleaseSlateResources not being implemented for the owner of this pointer.\n2) You may just be holding onto some slate pointers on an actor that don't get reset until the actor is Garbage Collected.  You should avoid doing this, and instead reset those references when the actor is Destroyed."), *DebugName);
 	}
 #endif
 
