Engine
- Adding SetAllowOverScroll function to ListViewBase
#rnx

diff --git a/Engine/Source/Runtime/UMG/Private/Components/ListViewBase.cpp b/Engine/Source/Runtime/UMG/Private/Components/ListViewBase.cpp
--- a/Engine/Source/Runtime/UMG/Private/Components/ListViewBase.cpp
+++ b/Engine/Source/Runtime/UMG/Private/Components/ListViewBase.cpp
@@ -102,6 +102,15 @@ void UListViewBase::SetScrollbarVisibility(ESlateVisibility InVisibility)
 	}
 }
 
+void  UListViewBase::SetAllowOverScroll(bool NewAllowOverscroll)
+{
+	if (MyTableViewBase)
+	{
+		AllowOverscroll = NewAllowOverscroll;
+		MyTableViewBase->SetAllowOverscroll(AllowOverscroll ? EAllowOverscroll::Yes : EAllowOverscroll::No);
+	}
+}
+
 UMG_API void UListViewBase::SetIsPointerScrollingEnabled(bool bInIsPointerScrollingEnabled)
 {
 	bIsPointerScrollingEnabled = bInIsPointerScrollingEnabled;
diff --git a/Engine/Source/Runtime/UMG/Public/Components/ListViewBase.h b/Engine/Source/Runtime/UMG/Public/Components/ListViewBase.h
--- a/Engine/Source/Runtime/UMG/Public/Components/ListViewBase.h
+++ b/Engine/Source/Runtime/UMG/Public/Components/ListViewBase.h
@@ -600,6 +600,9 @@ public:
 	UFUNCTION(BlueprintCallable, Category = ListViewBase)
 	UMG_API void SetScrollbarVisibility(ESlateVisibility InVisibility);
 
+	UFUNCTION(BlueprintCallable, Category = ListViewBase)
+	UMG_API void SetAllowOverScroll(bool NewAllowOverscroll);
+
 	/** Enable/Disable the ability of the list to scroll. This should be use as a temporary disable. */
 	UFUNCTION(BlueprintCallable, Category = ListViewBase)
 	UMG_API void SetIsPointerScrollingEnabled(bool bInIsPointerScrollingEnabled);
