Fix to Crash in SObjectWidget OnFinishedPointerInput missing CanRouteEvents that causes nullptr access.

diff --git a/Engine/Source/Runtime/UMG/Private/Slate/SObjectWidget.cpp b/Engine/Source/Runtime/UMG/Private/Slate/SObjectWidget.cpp
--- a/Engine/Source/Runtime/UMG/Private/Slate/SObjectWidget.cpp
+++ b/Engine/Source/Runtime/UMG/Private/Slate/SObjectWidget.cpp
@@ -526,7 +526,10 @@ FReply SObjectWidget::OnTouchFirstMove(const FGeometry& MyGeometry, const FPoint
 
 void SObjectWidget::OnFinishedPointerInput()
 {
-	WidgetObject->NativeOnFinishedPointerInput();
+	if (CanRouteEvent())
+	{
+		WidgetObject->NativeOnFinishedPointerInput();
+	}
 }
 
 FNavigationReply SObjectWidget::OnNavigation(const FGeometry& MyGeometry, const FNavigationEvent& InNavigationEvent)
@@ -537,7 +540,7 @@ FNavigationReply SObjectWidget::OnNavigation(const FGeometry& MyGeometry, const
 	}
 	FNavigationReply Reply = SCompoundWidget::OnNavigation(MyGeometry, InNavigationEvent);
 
-	if (WidgetObject && CanRouteEvent() )
+	if ( CanRouteEvent() )
 	{
 		return WidgetObject->NativeOnNavigation(MyGeometry, InNavigationEvent, Reply);
 	}
