Make player and world context objects const in Blueprint Libraries.

#tests PIE Verified usages of nodes appeared and functioned unchanged

diff --git a/Engine/Source/Runtime/UMG/Private/SlateBlueprintLibrary.cpp b/Engine/Source/Runtime/UMG/Private/SlateBlueprintLibrary.cpp
--- a/Engine/Source/Runtime/UMG/Private/SlateBlueprintLibrary.cpp
+++ b/Engine/Source/Runtime/UMG/Private/SlateBlueprintLibrary.cpp
@@ -62,13 +62,13 @@ bool USlateBlueprintLibrary::EqualEqual_SlateBrush(const FSlateBrush& A, const F
 	return A == B;
 }
 
-void USlateBlueprintLibrary::LocalToViewport(UObject* WorldContextObject, const FGeometry& Geometry, FVector2D LocalCoordinate, FVector2D& PixelPosition, FVector2D& ViewportPosition)
+void USlateBlueprintLibrary::LocalToViewport(const UObject* WorldContextObject, const FGeometry& Geometry, FVector2D LocalCoordinate, FVector2D& PixelPosition, FVector2D& ViewportPosition)
 {
 	FVector2D AbsoluteCoordinate = Geometry.LocalToAbsolute(LocalCoordinate);
 	AbsoluteToViewport(WorldContextObject, AbsoluteCoordinate, PixelPosition, ViewportPosition);
 }
 
-void USlateBlueprintLibrary::AbsoluteToViewport(UObject* WorldContextObject, FVector2D AbsoluteDesktopCoordinate, FVector2D& PixelPosition, FVector2D& ViewportPosition)
+void USlateBlueprintLibrary::AbsoluteToViewport(const UObject* WorldContextObject, FVector2D AbsoluteDesktopCoordinate, FVector2D& PixelPosition, FVector2D& ViewportPosition)
 {
 	UWorld* World = GEngine->GetWorldFromContextObject(WorldContextObject, EGetWorldErrorMode::LogAndReturnNull);
 	if ( World && World->IsGameWorld() )
@@ -95,7 +95,7 @@ void USlateBlueprintLibrary::AbsoluteToViewport(UObject* WorldContextObject, FVe
 	ViewportPosition = FVector2D(0, 0);
 }
 
-void USlateBlueprintLibrary::ScreenToWidgetLocal(UObject* WorldContextObject, const FGeometry& Geometry, FVector2D ScreenPosition, FVector2D& LocalCoordinate, bool bIncludeWindowPosition /*= false*/)
+void USlateBlueprintLibrary::ScreenToWidgetLocal(const UObject* WorldContextObject, const FGeometry& Geometry, FVector2D ScreenPosition, FVector2D& LocalCoordinate, bool bIncludeWindowPosition /*= false*/)
 {
 	FVector2D AbsoluteCoordinate;
 	ScreenToWidgetAbsolute(WorldContextObject, ScreenPosition, AbsoluteCoordinate, bIncludeWindowPosition);
@@ -103,7 +103,7 @@ void USlateBlueprintLibrary::ScreenToWidgetLocal(UObject* WorldContextObject, co
 	LocalCoordinate = Geometry.AbsoluteToLocal(AbsoluteCoordinate);
 }
 
-void USlateBlueprintLibrary::ScreenToWidgetAbsolute(UObject* WorldContextObject, FVector2D ScreenPosition, FVector2D& AbsoluteCoordinate, bool bIncludeWindowPosition /*= false*/)
+void USlateBlueprintLibrary::ScreenToWidgetAbsolute(const UObject* WorldContextObject, FVector2D ScreenPosition, FVector2D& AbsoluteCoordinate, bool bIncludeWindowPosition /*= false*/)
 {
 	UWorld* World = GEngine->GetWorldFromContextObject(WorldContextObject, EGetWorldErrorMode::LogAndReturnNull);
 	if ( World && World->IsGameWorld() )
@@ -135,7 +135,7 @@ void USlateBlueprintLibrary::ScreenToWidgetAbsolute(UObject* WorldContextObject,
 	AbsoluteCoordinate = FVector2D(0, 0);
 }
 
-void USlateBlueprintLibrary::ScreenToViewport(UObject* WorldContextObject, FVector2D ScreenPosition, FVector2D& ViewportPosition)
+void USlateBlueprintLibrary::ScreenToViewport(const UObject* WorldContextObject, FVector2D ScreenPosition, FVector2D& ViewportPosition)
 {
 	FVector2D AbsolutePosition;
 	USlateBlueprintLibrary::ScreenToWidgetAbsolute(WorldContextObject, ScreenPosition, AbsolutePosition);
diff --git a/Engine/Source/Runtime/UMG/Private/WidgetBlueprintLibrary.cpp b/Engine/Source/Runtime/UMG/Private/WidgetBlueprintLibrary.cpp
--- a/Engine/Source/Runtime/UMG/Private/WidgetBlueprintLibrary.cpp
+++ b/Engine/Source/Runtime/UMG/Private/WidgetBlueprintLibrary.cpp
@@ -548,7 +548,7 @@ void UWidgetBlueprintLibrary::DismissAllMenus()
 	FSlateApplication::Get().DismissAllMenus();
 }
 
-void UWidgetBlueprintLibrary::GetAllWidgetsOfClass(UObject* WorldContextObject, TArray<UUserWidget*>& FoundWidgets, TSubclassOf<UUserWidget> WidgetClass, bool TopLevelOnly)
+void UWidgetBlueprintLibrary::GetAllWidgetsOfClass(const UObject* WorldContextObject, TArray<UUserWidget*>& FoundWidgets, TSubclassOf<UUserWidget> WidgetClass, bool TopLevelOnly)
 {
 	//Prevent possibility of an ever-growing array if user uses this in a loop
 	FoundWidgets.Empty();
@@ -588,7 +588,7 @@ void UWidgetBlueprintLibrary::GetAllWidgetsOfClass(UObject* WorldContextObject,
 	}
 }
 
-void UWidgetBlueprintLibrary::GetAllWidgetsWithInterface(UObject* WorldContextObject, TArray<UUserWidget*>& FoundWidgets, TSubclassOf<UInterface> Interface, bool TopLevelOnly)
+void UWidgetBlueprintLibrary::GetAllWidgetsWithInterface(const UObject* WorldContextObject, TArray<UUserWidget*>& FoundWidgets, TSubclassOf<UInterface> Interface, bool TopLevelOnly)
 {
 	//Prevent possibility of an ever-growing array if user uses this in a loop
 	FoundWidgets.Empty();
@@ -653,7 +653,7 @@ FInputEvent UWidgetBlueprintLibrary::GetInputEventFromNavigationEvent(const FNav
 	return Event;
 }
 
-void UWidgetBlueprintLibrary::GetSafeZonePadding(UObject* WorldContextObject, FVector4& SafePadding, FVector2D& SafePaddingScale, FVector4& SpillOverPadding)
+void UWidgetBlueprintLibrary::GetSafeZonePadding(const UObject* WorldContextObject, FVector4& SafePadding, FVector2D& SafePaddingScale, FVector4& SpillOverPadding)
 {
 	FVector2D ViewportSize = UWidgetLayoutLibrary::GetViewportSize(WorldContextObject);
 
@@ -675,7 +675,7 @@ void UWidgetBlueprintLibrary::SetColorVisionDeficiencyType(EColorVisionDeficienc
 }
 
 
-bool UWidgetBlueprintLibrary::SetHardwareCursor(UObject* WorldContextObject, EMouseCursor::Type CursorShape, FName CursorName, FVector2D HotSpot)
+bool UWidgetBlueprintLibrary::SetHardwareCursor(const UObject* WorldContextObject, EMouseCursor::Type CursorShape, FName CursorName, FVector2D HotSpot)
 {
 	UWorld* World = GEngine->GetWorldFromContextObject(WorldContextObject, EGetWorldErrorMode::LogAndReturnNull);
 	if ( World && World->IsGameWorld() )
diff --git a/Engine/Source/Runtime/UMG/Private/WidgetLayoutLibrary.cpp b/Engine/Source/Runtime/UMG/Private/WidgetLayoutLibrary.cpp
--- a/Engine/Source/Runtime/UMG/Private/WidgetLayoutLibrary.cpp
+++ b/Engine/Source/Runtime/UMG/Private/WidgetLayoutLibrary.cpp
@@ -44,7 +44,7 @@ UWidgetLayoutLibrary::UWidgetLayoutLibrary(const FObjectInitializer& ObjectIniti
 {
 }
 
-bool UWidgetLayoutLibrary::ProjectWorldLocationToWidgetPosition(APlayerController* PlayerController, FVector WorldLocation, FVector2D& ViewportPosition, bool bPlayerViewportRelative)
+bool UWidgetLayoutLibrary::ProjectWorldLocationToWidgetPosition(const APlayerController* PlayerController, FVector WorldLocation, FVector2D& ViewportPosition, bool bPlayerViewportRelative)
 {
 	FVector ScreenPosition3D;
 	const bool bSuccess = ProjectWorldLocationToWidgetPositionWithDistance(PlayerController, WorldLocation, ScreenPosition3D, bPlayerViewportRelative);
@@ -52,7 +52,7 @@ bool UWidgetLayoutLibrary::ProjectWorldLocationToWidgetPosition(APlayerControlle
 	return bSuccess;
 }
 
-bool UWidgetLayoutLibrary::ProjectWorldLocationToWidgetPositionWithDistance(APlayerController* PlayerController, FVector WorldLocation, FVector& ViewportPosition, bool bPlayerViewportRelative)
+bool UWidgetLayoutLibrary::ProjectWorldLocationToWidgetPositionWithDistance(const APlayerController* PlayerController, FVector WorldLocation, FVector& ViewportPosition, bool bPlayerViewportRelative)
 {	
 	if ( PlayerController )
 	{
@@ -134,7 +134,7 @@ FVector2D UWidgetLayoutLibrary::GetMousePositionOnPlatform()
 	return FVector2D(0, 0);
 }
 
-FVector2D UWidgetLayoutLibrary::GetMousePositionOnViewport(UObject* WorldContextObject)
+FVector2D UWidgetLayoutLibrary::GetMousePositionOnViewport(const UObject* WorldContextObject)
 {
 	if ( FSlateApplication::IsInitialized() )
 	{
@@ -146,7 +146,7 @@ FVector2D UWidgetLayoutLibrary::GetMousePositionOnViewport(UObject* WorldContext
 	return FVector2D(0, 0);
 }
 
-bool UWidgetLayoutLibrary::GetMousePositionScaledByDPI(APlayerController* Player, float& LocationX, float& LocationY)
+bool UWidgetLayoutLibrary::GetMousePositionScaledByDPI(const APlayerController* Player, float& LocationX, float& LocationY)
 {
 	if ( Player && Player->GetMousePosition(LocationX, LocationY) )
 	{
@@ -159,7 +159,7 @@ bool UWidgetLayoutLibrary::GetMousePositionScaledByDPI(APlayerController* Player
 
 	return false;
 }
-bool UWidgetLayoutLibrary::GetMousePositionScaledByDPI(APlayerController* Player, double& LocationX, double& LocationY)
+bool UWidgetLayoutLibrary::GetMousePositionScaledByDPI(const APlayerController* Player, double& LocationX, double& LocationY)
 {
 	float X = (float)LocationX, Y = (float)LocationY;
 	if(GetMousePositionScaledByDPI(Player, X, Y))
@@ -170,7 +170,7 @@ bool UWidgetLayoutLibrary::GetMousePositionScaledByDPI(APlayerController* Player
 	return false;
 }
 
-FVector2D UWidgetLayoutLibrary::GetViewportSize(UObject* WorldContextObject)
+FVector2D UWidgetLayoutLibrary::GetViewportSize(const UObject* WorldContextObject)
 {
 	static TFrameValue<FVector2D> ViewportSizeCache;
 
@@ -193,7 +193,7 @@ FVector2D UWidgetLayoutLibrary::GetViewportSize(UObject* WorldContextObject)
 	return ViewportSizeCache.GetValue();
 }
 
-FGeometry UWidgetLayoutLibrary::GetViewportWidgetGeometry(UObject* WorldContextObject)
+FGeometry UWidgetLayoutLibrary::GetViewportWidgetGeometry(const UObject* WorldContextObject)
 {
 	UWorld* World = GEngine->GetWorldFromContextObject(WorldContextObject, EGetWorldErrorMode::LogAndReturnNull);
 	if ( World && World->IsGameWorld() )
@@ -369,7 +369,7 @@ UWidgetSwitcherSlot* UWidgetLayoutLibrary::SlotAsWidgetSwitcherSlot(UWidget* Wid
 	return nullptr;
 }
 
-void UWidgetLayoutLibrary::RemoveAllWidgets(UObject* WorldContextObject)
+void UWidgetLayoutLibrary::RemoveAllWidgets(const UObject* WorldContextObject)
 {
 	UWorld* World = GEngine->GetWorldFromContextObject(WorldContextObject, EGetWorldErrorMode::LogAndReturnNull);
 	if ( World && World->IsGameWorld() )
diff --git a/Engine/Source/Runtime/UMG/Public/Blueprint/SlateBlueprintLibrary.h b/Engine/Source/Runtime/UMG/Public/Blueprint/SlateBlueprintLibrary.h
--- a/Engine/Source/Runtime/UMG/Public/Blueprint/SlateBlueprintLibrary.h
+++ b/Engine/Source/Runtime/UMG/Public/Blueprint/SlateBlueprintLibrary.h
@@ -107,7 +107,7 @@ public:
 	 * what you would use.
 	 */
 	UFUNCTION(BlueprintPure, Category="User Interface|Geometry", meta=( WorldContext="WorldContextObject" ))
-	static UMG_API void LocalToViewport(UObject* WorldContextObject, const FGeometry& Geometry, FVector2D LocalCoordinate, FVector2D& PixelPosition, FVector2D& ViewportPosition);
+	static UMG_API void LocalToViewport(const UObject* WorldContextObject, const FGeometry& Geometry, FVector2D LocalCoordinate, FVector2D& PixelPosition, FVector2D& ViewportPosition);
 
 	/**
 	 * Translates absolute coordinate in desktop space of the geometry provided into local viewport coordinates.
@@ -119,25 +119,25 @@ public:
 	 * what you would use.
 	 */
 	UFUNCTION(BlueprintPure, Category="User Interface|Geometry", meta=( WorldContext="WorldContextObject" ))
-	static UMG_API void AbsoluteToViewport(UObject* WorldContextObject, FVector2D AbsoluteDesktopCoordinate, FVector2D& PixelPosition, FVector2D& ViewportPosition);
+	static UMG_API void AbsoluteToViewport(const UObject* WorldContextObject, FVector2D AbsoluteDesktopCoordinate, FVector2D& PixelPosition, FVector2D& ViewportPosition);
 
 	/**
 	 * Translates a screen position in pixels into the local space of a widget with the given geometry. 
 	 * If bIncludeWindowPosition is true, then this method will also remove the game window's position (useful when in windowed mode).
 	 */
 	UFUNCTION(BlueprintPure, Category="User Interface|Geometry", meta=( WorldContext="WorldContextObject", DisplayName="ScreenToLocal" ))
-	static UMG_API void ScreenToWidgetLocal(UObject* WorldContextObject, const FGeometry& Geometry, FVector2D ScreenPosition, FVector2D& LocalCoordinate, bool bIncludeWindowPosition = false);
+	static UMG_API void ScreenToWidgetLocal(const UObject* WorldContextObject, const FGeometry& Geometry, FVector2D ScreenPosition, FVector2D& LocalCoordinate, bool bIncludeWindowPosition = false);
 
 	/**
 	 * Translates a screen position in pixels into absolute application coordinates.
 	 * If bIncludeWindowPosition is true, then this method will also remove the game window's position (useful when in windowed mode).
 	 */
 	UFUNCTION(BlueprintPure, Category="User Interface|Geometry", meta=( WorldContext="WorldContextObject", DisplayName="ScreenToAbsolute" ))
-	static UMG_API void ScreenToWidgetAbsolute(UObject* WorldContextObject, FVector2D ScreenPosition, FVector2D& AbsoluteCoordinate, bool bIncludeWindowPosition = false);
+	static UMG_API void ScreenToWidgetAbsolute(const UObject* WorldContextObject, FVector2D ScreenPosition, FVector2D& AbsoluteCoordinate, bool bIncludeWindowPosition = false);
 
 	/**
 	 * Translates a screen position in pixels into the local space of the viewport widget.
 	 */
 	UFUNCTION(BlueprintPure, Category="User Interface|Geometry", meta=( WorldContext="WorldContextObject" ))
-	static UMG_API void ScreenToViewport(UObject* WorldContextObject, FVector2D ScreenPosition, FVector2D& ViewportPosition);
+	static UMG_API void ScreenToViewport(const UObject* WorldContextObject, FVector2D ScreenPosition, FVector2D& ViewportPosition);
 };
diff --git a/Engine/Source/Runtime/UMG/Public/Blueprint/WidgetBlueprintLibrary.h b/Engine/Source/Runtime/UMG/Public/Blueprint/WidgetBlueprintLibrary.h
--- a/Engine/Source/Runtime/UMG/Public/Blueprint/WidgetBlueprintLibrary.h
+++ b/Engine/Source/Runtime/UMG/Public/Blueprint/WidgetBlueprintLibrary.h
@@ -300,7 +300,7 @@ public:
 	 * @param TopLevelOnly Only the widgets that are direct children of the viewport will be returned.
 	 */
 	UFUNCTION(BlueprintCallable, BlueprintCosmetic, Category="Widget", meta=( WorldContext="WorldContextObject", DeterminesOutputType = "WidgetClass", DynamicOutputParam = "FoundWidgets" ))
-	static UMG_API void GetAllWidgetsOfClass(UObject* WorldContextObject, TArray<UUserWidget*>& FoundWidgets, TSubclassOf<UUserWidget> WidgetClass, bool TopLevelOnly = true);
+	static UMG_API void GetAllWidgetsOfClass(const UObject* WorldContextObject, TArray<UUserWidget*>& FoundWidgets, TSubclassOf<UUserWidget> WidgetClass, bool TopLevelOnly = true);
 
 	/**
 	* Find all widgets in the world with the specified interface.
@@ -310,7 +310,7 @@ public:
 	* @param TopLevelOnly Only the widgets that are direct children of the viewport will be returned.
 	*/
 	UFUNCTION(BlueprintCallable, BlueprintCosmetic, Category="Widget", meta = (WorldContext = "WorldContextObject", DeterminesOutputType = "WidgetClass", DynamicOutputParam = "FoundWidgets"))
-	static UMG_API void GetAllWidgetsWithInterface(UObject* WorldContextObject, TArray<UUserWidget*>& FoundWidgets, TSubclassOf<UInterface> Interface, bool TopLevelOnly);
+	static UMG_API void GetAllWidgetsWithInterface(const UObject* WorldContextObject, TArray<UUserWidget*>& FoundWidgets, TSubclassOf<UInterface> Interface, bool TopLevelOnly);
 
 	UFUNCTION(BlueprintPure, Category="Widget", meta = (CompactNodeTitle = "->", BlueprintAutocast))
 	static UMG_API FInputEvent GetInputEventFromKeyEvent(const FKeyEvent& Event);
@@ -331,7 +331,7 @@ public:
 	 * Gets the amount of padding that needs to be added when accounting for the safe zone on TVs.
 	 */
 	UFUNCTION(BlueprintCallable, Category="Widget|Safe Zone", meta=( WorldContext="WorldContextObject" ))
-	static UMG_API void GetSafeZonePadding(UObject* WorldContextObject, FVector4& SafePadding, FVector2D& SafePaddingScale, FVector4& SpillOverPadding);
+	static UMG_API void GetSafeZonePadding(const UObject* WorldContextObject, FVector4& SafePadding, FVector2D& SafePaddingScale, FVector4& SpillOverPadding);
 
 	/**
 	* Apply color deficiency correction settings to the game window 
@@ -347,7 +347,7 @@ public:
 	 * Loads or sets a hardware cursor from the content directory in the game.
 	 */
 	UFUNCTION(BlueprintCallable, Category="Widget|Hardware Cursor", meta=( WorldContext="WorldContextObject" ))
-	static UMG_API bool SetHardwareCursor(UObject* WorldContextObject, EMouseCursor::Type CursorShape, FName CursorName, FVector2D HotSpot);
+	static UMG_API bool SetHardwareCursor(const UObject* WorldContextObject, EMouseCursor::Type CursorShape, FName CursorName, FVector2D HotSpot);
 
 	UFUNCTION(BlueprintCallable, Category = "Widget|Window Title Bar")
 	static UMG_API void SetWindowTitleBarState(UWidget* TitleBarContent, EWindowTitleBarMode Mode, bool bTitleBarDragEnabled, bool bWindowButtonsVisible, bool bTitleBarVisible);
diff --git a/Engine/Source/Runtime/UMG/Public/Blueprint/WidgetLayoutLibrary.h b/Engine/Source/Runtime/UMG/Public/Blueprint/WidgetLayoutLibrary.h
--- a/Engine/Source/Runtime/UMG/Public/Blueprint/WidgetLayoutLibrary.h
+++ b/Engine/Source/Runtime/UMG/Public/Blueprint/WidgetLayoutLibrary.h
@@ -42,7 +42,7 @@ public:
 	 * @return true if the position projects onto the screen.
 	 */
 	UFUNCTION(BlueprintPure, BlueprintCosmetic, Category="Viewport")
-	static UMG_API bool ProjectWorldLocationToWidgetPosition(APlayerController* PlayerController, FVector WorldLocation, FVector2D& ScreenPosition, bool bPlayerViewportRelative);
+	static UMG_API bool ProjectWorldLocationToWidgetPosition(const APlayerController* PlayerController, FVector WorldLocation, FVector2D& ScreenPosition, bool bPlayerViewportRelative);
 
 	/**
 	 * Convert a World Space 3D position into a 2D Widget Screen Space position, with distance from the camera the Z component.  This
@@ -50,7 +50,7 @@ public:
 	 *
 	 * @return true if the world coordinate was successfully projected to the screen.
 	 */
-	static UMG_API bool ProjectWorldLocationToWidgetPositionWithDistance(APlayerController* PlayerController, FVector WorldLocation, FVector& ScreenPosition, bool bPlayerViewportRelative);
+	static UMG_API bool ProjectWorldLocationToWidgetPositionWithDistance(const APlayerController* PlayerController, FVector WorldLocation, FVector& ScreenPosition, bool bPlayerViewportRelative);
 
 	/**
 	 * Gets the current DPI Scale being applied to the viewport and all the Widgets.
@@ -67,7 +67,7 @@ public:
 	 * Gets the size of the game viewport.
 	 */
 	UFUNCTION(BlueprintPure, BlueprintCosmetic, Category="Viewport", meta=( WorldContext="WorldContextObject", Keywords = "screen size"))
-	static UMG_API FVector2D GetViewportSize(UObject* WorldContextObject);
+	static UMG_API FVector2D GetViewportSize(const UObject* WorldContextObject);
 
 	/**
 	 * Gets the geometry of the widget holding all widgets added to the "Viewport".  You
@@ -75,7 +75,7 @@ public:
 	 * widget.
 	 */
 	UFUNCTION(BlueprintCallable, Category="Viewport", meta=( WorldContext="WorldContextObject" ))
-	static UMG_API FGeometry GetViewportWidgetGeometry(UObject* WorldContextObject);
+	static UMG_API FGeometry GetViewportWidgetGeometry(const UObject* WorldContextObject);
 
 	/**
 	 * Gets the geometry of the widget holding all widgets added to the "Player Screen". You
@@ -95,7 +95,7 @@ public:
 	 * Gets the platform's mouse cursor position in the local space of the viewport widget.
 	 */
 	UFUNCTION(BlueprintCallable, Category="Viewport", meta=( WorldContext="WorldContextObject" ))
-	static UMG_API FVector2D GetMousePositionOnViewport(UObject* WorldContextObject);
+	static UMG_API FVector2D GetMousePositionOnViewport(const UObject* WorldContextObject);
 
 	/**
 	 * Gets the mouse position of the player controller, scaled by the DPI.  If you're trying to go from raw mouse screenspace coordinates
@@ -105,8 +105,8 @@ public:
 	 */
 	//UE_DEPRECATED(4.17, "Use GetMousePositionOnViewport() instead.  Optionally and for more options, you can use GetViewportWidgetGeometry and GetPlayerScreenWidgetGeometry are newly introduced to give you the geometry of the viewport and the player screen for widgets to help convert between spaces.")
 	UFUNCTION(BlueprintPure, BlueprintCosmetic, Category="Viewport")
-	static UMG_API bool GetMousePositionScaledByDPI(APlayerController* Player, float& LocationX, float& LocationY);
-	static UMG_API bool GetMousePositionScaledByDPI(APlayerController* Player, double& LocationX, double& LocationY);	// LWC_TODO: Temp stand in for native calls with FVector2D components.
+	static UMG_API bool GetMousePositionScaledByDPI(const APlayerController* Player, float& LocationX, float& LocationY);
+	static UMG_API bool GetMousePositionScaledByDPI(const APlayerController* Player, double& LocationX, double& LocationY);	// LWC_TODO: Temp stand in for native calls with FVector2D components.
 
 	/**
 	* Gets the slot object on the child widget as a Border Slot, allowing you to manipulate layout information.
@@ -211,5 +211,5 @@ public:
 	 * Removes all widgets from the viewport.
 	 */
 	UFUNCTION(BlueprintCallable, BlueprintCosmetic, Category="Viewport", meta=( WorldContext="WorldContextObject" ))
-	static UMG_API void RemoveAllWidgets(UObject* WorldContextObject);
+	static UMG_API void RemoveAllWidgets(const UObject* WorldContextObject);
 };
