Fixed issue where OnEntriesGenerated would fire for each entry instead of firing once all entries are generated

diff --git a/Engine/Source/Runtime/UMG/Private/Components/ListViewBase.cpp b/Engine/Source/Runtime/UMG/Private/Components/ListViewBase.cpp
--- a/Engine/Source/Runtime/UMG/Private/Components/ListViewBase.cpp
+++ b/Engine/Source/Runtime/UMG/Private/Components/ListViewBase.cpp
@@ -3,7 +3,6 @@
 #include "Components/ListViewBase.h"
 #include "UMGPrivate.h"
 #include "Widgets/Text/STextBlock.h"
-#include "TimerManager.h"
 #include "Engine/Blueprint.h"
 #include "Editor/WidgetCompilerLog.h"
 #include "Blueprint/WidgetBlueprintLibrary.h"
@@ -303,11 +302,11 @@ void UListViewBase::FinishGeneratingEntry(UUserWidget& GeneratedEntry)
 		// Announcing the row generation now is just a bit too soon, as we haven't finished generating the row as far as the underlying list is concerned. 
 		// So we cache the row/item pair here and announce their generation on the next tick
 		GeneratedEntriesToAnnounce.AddUnique(&GeneratedEntry);
-		if (!EntryGenAnnouncementTimerHandle.IsValid())
+		if (!EntryGenAnnouncementDelegateHandle.IsValid())
 		{
 			if (UWorld* World = GetWorld())
 			{
-				FTSTicker::GetCoreTicker().AddTicker(FTickerDelegate::CreateWeakLambda(this, [this](float InDelta)
+				EntryGenAnnouncementDelegateHandle = FTSTicker::GetCoreTicker().AddTicker(FTickerDelegate::CreateWeakLambda(this, [this](float InDelta)
 				{
 					HandleAnnounceGeneratedEntries();
 					return false;
@@ -319,7 +318,7 @@ void UListViewBase::FinishGeneratingEntry(UUserWidget& GeneratedEntry)
 
 void UListViewBase::HandleAnnounceGeneratedEntries()
 {
-	EntryGenAnnouncementTimerHandle.Invalidate();
+	EntryGenAnnouncementDelegateHandle.Reset();
 	for (TWeakObjectPtr<UUserWidget> EntryWidget : GeneratedEntriesToAnnounce)
 	{
 		if (EntryWidget.IsValid())
diff --git a/Engine/Source/Runtime/UMG/Public/Components/ListViewBase.h b/Engine/Source/Runtime/UMG/Public/Components/ListViewBase.h
--- a/Engine/Source/Runtime/UMG/Public/Components/ListViewBase.h
+++ b/Engine/Source/Runtime/UMG/Public/Components/ListViewBase.h
@@ -873,7 +873,7 @@ private:
 	UPROPERTY(Transient)
 	FUserWidgetPool EntryWidgetPool;
 
-	FTimerHandle EntryGenAnnouncementTimerHandle;
+	FTSTicker::FDelegateHandle EntryGenAnnouncementDelegateHandle;
 	
 	FOnListEntryGenerated OnListEntryGeneratedEvent;
 	FOnEntryWidgetReleased OnEntryWidgetReleasedEvent;
