Added ClearFill to UGridPanel.

Currently, Blueprint users can reset ColumnFill and RowFill by calling the Array_Clear function on them. However, this is exploiting a loophole: these variables are read-only and the Blueprint compiler currently allows array functions to change read-only variables.

Eventually, this will become an error. Instead, users can use UGridPanel::ClearFill, which also clears the fill coeffecients on the underlying Slate widget.

diff --git a/Engine/Source/Runtime/UMG/Private/Components/GridPanel.cpp b/Engine/Source/Runtime/UMG/Private/Components/GridPanel.cpp
--- a/Engine/Source/Runtime/UMG/Private/Components/GridPanel.cpp
+++ b/Engine/Source/Runtime/UMG/Private/Components/GridPanel.cpp
@@ -133,6 +133,17 @@ void UGridPanel::SetRowFill(int32 RowIndex, float Coefficient)
 	}
 }
 
+void UGridPanel::ClearFill()
+{
+	ColumnFill.Reset();
+	RowFill.Reset();
+
+	if (MyGridPanel.IsValid())
+	{
+		MyGridPanel->ClearFill();
+	}
+}
+
 #if WITH_EDITOR
 
 const FText UGridPanel::GetPaletteCategory()
diff --git a/Engine/Source/Runtime/UMG/Public/Components/GridPanel.h b/Engine/Source/Runtime/UMG/Public/Components/GridPanel.h
--- a/Engine/Source/Runtime/UMG/Public/Components/GridPanel.h
+++ b/Engine/Source/Runtime/UMG/Public/Components/GridPanel.h
@@ -14,8 +14,6 @@ class UGridSlot;
 
 /**
  * A table-like panel that retains the width of every column throughout the table.
- * 
- * * Many Children
  */
 UCLASS(MinimalAPI)
 class UGridPanel : public UPanelWidget
@@ -34,7 +32,6 @@ public:
 
 public:
 
-	/**  */
 	UFUNCTION(BlueprintCallable, Category="Widget")
 	UMG_API UGridSlot* AddChildToGrid(UWidget* Content, int32 InRow = 0, int32 InColumn = 0);
 
@@ -44,6 +41,9 @@ public:
 	UFUNCTION(BlueprintCallable, Category = "Widget")
 	UMG_API void SetRowFill(int32 RowIndex, float Coefficient);
 
+	UFUNCTION(BlueprintCallable, Category = "Widget")
+	UMG_API void ClearFill();
+
 public:
 
 	// UWidget interface
