UI Component binding: pass the correct component object at runtime from the UserWidgetExtension instead of WidgetBlueprintGeneratedClassExtension

diff --git a/Engine/Source/Runtime/UMG/Private/Extensions/UIComponentWidgetBlueprintGeneratedClassExtension.cpp b/Engine/Source/Runtime/UMG/Private/Extensions/UIComponentWidgetBlueprintGeneratedClassExtension.cpp
--- a/Engine/Source/Runtime/UMG/Private/Extensions/UIComponentWidgetBlueprintGeneratedClassExtension.cpp
+++ b/Engine/Source/Runtime/UMG/Private/Extensions/UIComponentWidgetBlueprintGeneratedClassExtension.cpp
@@ -118,7 +118,7 @@ void UUIComponentWidgetBlueprintGeneratedClassExtension::PreConstruct(UUserWidge
 	if (!IsDesignTime)
 	{
 		// Set the component object in the corresponding property on the userwidget.
-		ComponentContainer->ForEachComponentTarget([UserWidget](const FUIComponentTarget& ComponentTarget) {
+		ComponentContainer->ForEachComponentTarget([UserWidgetExtension, UserWidget](const FUIComponentTarget& ComponentTarget) {
 			FName ComponentPropertyName = UUIComponentContainer::GetPropertyNameForComponent(ComponentTarget.GetComponent(), ComponentTarget.GetTargetName());
 
 			FObjectPropertyBase* FoundPanelObjectProperty = FindFProperty<FObjectPropertyBase>(UserWidget->GetClass(), ComponentPropertyName);
@@ -127,7 +127,8 @@ void UUIComponentWidgetBlueprintGeneratedClassExtension::PreConstruct(UUserWidge
 			{
 				if (ensure(FoundPanelObjectProperty->PropertyClass == ComponentTarget.GetComponent()->GetClass()))
 				{
-					FoundPanelObjectProperty->SetObjectPropertyValue_InContainer(UserWidget, ComponentTarget.GetComponent());
+					UUIComponent* UserWidgetComponent = UserWidgetExtension->GetComponent(FoundPanelObjectProperty->PropertyClass, ComponentTarget.GetTargetName());
+					FoundPanelObjectProperty->SetObjectPropertyValue_InContainer(UserWidget, UserWidgetComponent);
 				}
 			}
 			});
