[UMG + Enhanced Input] Define a much better lifecycle for how input components are treated in widgets.

- When a widget is initalized, create the input component and bind its script delegates.
- On construction, push the component to the input stack
- on destruction, pop the component from the input stack.
- Finally, when the widget is completely destroyed, unbind and remove all input delegates

This provides a more expected and useful behavior when binding enhanced input (or legacy inputs) in widgets, where the delegates will only be processed when the widget is active in the viewport.

Before, the input component was only being added via the UUserWidget::NativeOnInitialized function, which meant it only would register once. You would need to call the register function every time it was added to the viewport manually. This was pretty confusing behavior and felt very buggy.

Deprecated the old InitializeInputComponent in favor of the new input component lifecycle functions.

[RN] Any input bindings in widgets will not happen on Construction, not initalization, of the widget. The input component will also be pushed on construction, and popped on NativeDestruct. This means you don't need to manually call RegisterInputComponent anymore.

diff --git a/Engine/Source/Runtime/UMG/Private/UserWidget.cpp b/Engine/Source/Runtime/UMG/Private/UserWidget.cpp
--- a/Engine/Source/Runtime/UMG/Private/UserWidget.cpp
+++ b/Engine/Source/Runtime/UMG/Private/UserWidget.cpp
@@ -316,6 +316,14 @@ void UUserWidget::BeginDestroy()
 {
 	Super::BeginDestroy();
 
+	// At this point, the widget is being completely destroyed. We want to ensure that
+	// the input component is no longer being processed so that any input delegate nodes
+	// are no longer going to trigger
+	DestroyInputComponent();
+
+	// Ensure that we have removed all the delegate bindings that we set up
+	StopListeningForPlayerControllerChanges();
+	
 	TearDownAnimations();
 
 	if (AnimationTickManager)
@@ -1382,6 +1390,10 @@ void UUserWidget::SetVisibility(ESlateVisibility InVisibility)
 
 void UUserWidget::SetPlayerContext(const FLocalPlayerContext& InPlayerContext)
 {
+	// Pop the input component off of the current player controller, if we have one...
+	StopProcessingInputScriptDelegates();
+	StopListeningForPlayerControllerChanges();
+	
 	PlayerContext = InPlayerContext;
 	CachedWorld.Reset();
 
@@ -1397,6 +1409,9 @@ void UUserWidget::SetPlayerContext(const FLocalPlayerContext& InPlayerContext)
 			});
 	}
 	
+	// ... and start processing input again and push onto this new player controller's input stack if we can
+	StartListeningForPlayerControllerChanges();
+	StartProcessingInputScriptDelegates();
 }
 
 const FLocalPlayerContext& UUserWidget::GetPlayerContext() const
@@ -1415,10 +1430,22 @@ ULocalPlayer* UUserWidget::GetOwningLocalPlayer() const
 
 void UUserWidget::SetOwningLocalPlayer(ULocalPlayer* LocalPlayer)
 {
-	if ( LocalPlayer )
+	if (LocalPlayer)
 	{
+		// Pop the input component off of the current player controller, if we have one...
+		StopProcessingInputScriptDelegates();
+
+		// Remove bindings from any callbacks for PC changes on the old local player
+		StopListeningForPlayerControllerChanges();
+		
 		PlayerContext = FLocalPlayerContext(LocalPlayer, GetWorld());
 		CachedWorld.Reset();
+
+		// And start listening to the player controller changes again
+		StartListeningForPlayerControllerChanges();
+
+		// ... and start processing input again and push onto this new player controller's input stack if we can
+		StartProcessingInputScriptDelegates();
 	}
 }
 
@@ -1431,8 +1458,14 @@ void UUserWidget::SetOwningPlayer(APlayerController* LocalPlayerController)
 {
 	if (LocalPlayerController && LocalPlayerController->IsLocalController())
 	{
+		// Pop the input component off of the current player controller, if we have one...
+		StopProcessingInputScriptDelegates();
+		
 		PlayerContext = FLocalPlayerContext(LocalPlayerController);
 		CachedWorld.Reset();
+
+		// ... and start processing input again and push onto this new player controller's input stack if we can
+		StartProcessingInputScriptDelegates();
 	}
 }
 
@@ -1766,29 +1799,7 @@ void UUserWidget::BindToAnimationEvent(UWidgetAnimation* InAnimation, FWidgetAni
 
 void UUserWidget::NativeOnInitialized()
 {
-	// The widget tree may be constructed from instanced subobjects so we need to set this value
-	// based on this widgets actual CDO.
-	if (!HasAnyFlags(RF_ClassDefaultObject))
-	{
-		const UUserWidget* DefaultWidget = CastChecked<UUserWidget>(GetClass()->GetDefaultObject());
-		bAutomaticallyRegisterInputOnConstruction = DefaultWidget->bAutomaticallyRegisterInputOnConstruction;
-	}
-
-	// Bind any input delegates that may be on this widget to its owning player controller
-	if (bAutomaticallyRegisterInputOnConstruction)
-	{
-		// Only widgets with a valid player controller can bind to input delegates
-		if (GetOwningPlayer() != nullptr)
-		{
-			InitializeInputComponent();
-			check(InputComponent);
-			UInputDelegateBinding::BindInputDelegates(GetClass(), InputComponent, this);
-		}
-		else if (!IsEditorUtility())
-		{
-			UE_LOG(LogUMG, Error, TEXT("[%hs] Widget '%s' has bAutomaticallyRegisterInputOnConstruction set to true, but has no valid player controller. Input delegates will not work!"), __func__, *GetNameSafe(this));
-		}
-	}
+	CreateInputComponent();
 	
 	if (UWidgetBlueprintGeneratedClass* BPClass = Cast<UWidgetBlueprintGeneratedClass>(GetClass()))
 	{
@@ -1840,6 +1851,10 @@ void UUserWidget::NativeConstruct()
 {
 	LLM_SCOPE_BYTAG(UI_UMG);
 
+	// Upon construction, allow for our input delegates to be called
+	StartProcessingInputScriptDelegates();
+	StartListeningForPlayerControllerChanges();
+	
 	if (UWidgetBlueprintGeneratedClass* BPClass = Cast<UWidgetBlueprintGeneratedClass>(GetClass()))
 	{
 		BPClass->ForEachExtension([this](UWidgetBlueprintGeneratedClassExtension* Extension)
@@ -1868,7 +1883,12 @@ void UUserWidget::NativeConstruct()
 
 void UUserWidget::NativeDestruct()
 {
-	StopListeningForAllInputActions();
+	// On NativeDestruct we want to stop processing any input delegates so that
+	// when this widget isn't visible they don't fire. But, the widget might be pushed
+	// back onto the viewport again, so we don't want to completely unbind all delegates yet. 
+	StopProcessingInputScriptDelegates();
+	StopListeningForPlayerControllerChanges();
+	
 	OnNativeDestruct.Broadcast(this);
 
 	Destruct();
@@ -1896,6 +1916,149 @@ void UUserWidget::NativeDestruct()
 	}
 }
 
+void UUserWidget::CreateInputComponent()
+{
+	// We cannot bind to any input events during design time.
+	if (IsDesignTime())
+	{
+		return;
+	}
+	
+	// The widget tree may be constructed from instanced subobjects so we need to set this value
+	// based on this widget's actual CDO.
+	if (!HasAnyFlags(RF_ClassDefaultObject))
+	{
+		const UUserWidget* DefaultWidget = CastChecked<UUserWidget>(GetClass()->GetDefaultObject());
+		bAutomaticallyRegisterInputOnConstruction = DefaultWidget->bAutomaticallyRegisterInputOnConstruction;
+	}
+
+	// If bAutomaticallyRegisterInputOnConstruction is false, then there are no input delegates in this widget
+	// that need to be processed. There is no need to create an input component
+	if (!bAutomaticallyRegisterInputOnConstruction)
+	{
+		return;
+	}
+	
+	if (const APlayerController* Controller = GetOwningPlayer())
+	{
+		// Use the existing PC's input class, or fallback to the project default. We should use the existing class
+		// instead of just the default one because if you have a plugin that has a PC with a different default input
+		// class then this would fail
+		const UClass* InputClass = Controller->InputComponent ? Controller->InputComponent->GetClass() : UInputSettings::GetDefaultInputComponentClass();
+		InputComponent = NewObject<UInputComponent>( this, InputClass, NAME_None, RF_Transient );
+
+		// TODO: The input consumption and priority behaviors of the component should be driven
+		// not by these properties, but by something like the Z-order of the widget. This way,
+		// you could create some "UI Input Actions" and have them routed in a more deterministic
+		// way. UE-306592
+		PRAGMA_DISABLE_DEPRECATION_WARNINGS
+		InputComponent->bBlockInput = bStopAction;
+		InputComponent->Priority = Priority;
+		PRAGMA_ENABLE_DEPRECATION_WARNINGS
+
+		UInputDelegateBinding::BindInputDelegates(GetClass(), InputComponent, this);
+	}
+	else
+	{
+		FMessageLog("PIE").Info(FText::Format(LOCTEXT("NoInputListeningWithoutPlayerController", "Unable to listen to input actions without a player controller in {0}."), FText::FromName(GetClass()->GetFName())));
+	}
+}
+
+void UUserWidget::StartProcessingInputScriptDelegates()
+{
+	// The input component should always be valid if bAutomaticallyRegisterInputOnConstruction is true,
+	// because it would have been created in UUserWidget::CreateInputComponent
+	if (!InputComponent)
+	{
+		return;
+	}
+
+	// Again, we should only ever have a valid input component if we have a valid player controller.
+	APlayerController* Controller = GetOwningPlayer();
+	if (!ensure(Controller))
+	{
+		return;
+	}
+
+	// Actually push the input component to the stack. This is what will make the input systems
+	// actually call the delegates bound in this graph if the input is activated
+	Controller->PushInputComponent(InputComponent);
+}
+
+void UUserWidget::StopProcessingInputScriptDelegates()
+{
+	// The input component should always be valid if bAutomaticallyRegisterInputOnConstruction is true,
+	// because it would have been created in UUserWidget::CreateInputComponent
+	if (!InputComponent)
+	{
+		return;
+	}
+
+	// Again, we should only ever have a valid input component if we have a valid player controller.
+	APlayerController* Controller = GetOwningPlayer();
+	if (!Controller)
+	{
+		return;
+	}
+
+	// Pop this component off of the input stack, which will stop the input delegates in the graph
+	// from being called.
+	Controller->PopInputComponent(InputComponent);
+}
+
+void UUserWidget::DestroyInputComponent()
+{	
+	if (!InputComponent)
+	{
+		return;
+	}
+
+	// Ensure that the component is removed from the owning controller
+	if (APlayerController* Controller = GetOwningPlayer())
+	{
+		Controller->PopInputComponent(InputComponent);
+	}
+
+	// Unbind all the action delegates which were bound to this widget
+	for (int32 ExistingIndex = InputComponent->GetNumActionBindings() - 1; ExistingIndex >= 0; --ExistingIndex)
+	{
+		InputComponent->RemoveActionBinding(ExistingIndex);
+	}
+	InputComponent->ClearActionBindings();
+
+	// Destroy the input component
+	InputComponent->MarkAsGarbage();
+	InputComponent = nullptr;
+}
+
+void UUserWidget::StartListeningForPlayerControllerChanges()
+{
+	if (ULocalPlayer* LP = GetOwningLocalPlayer())
+	{
+		LP->OnPlayerControllerChanged().AddUObject(this, &UUserWidget::HandleOwningLocalPlayerControllerChanged);
+	}
+}
+
+void UUserWidget::StopListeningForPlayerControllerChanges()
+{
+	if (ULocalPlayer* LP = GetOwningLocalPlayer())
+	{
+		LP->OnPlayerControllerChanged().RemoveAll(this);
+	}
+}
+
+void UUserWidget::HandleOwningLocalPlayerControllerChanged(APlayerController* NewPlayerController)
+{
+	// If the controller is the same, then we don't need to do anything.
+	if (NewPlayerController == GetOwningPlayer())
+	{
+		return;
+	}
+
+	// SetOwningPlayer will stop listening to input delegates on the old PC, and then start again on the new PC
+	SetOwningPlayer(NewPlayerController);
+}
+
 void UUserWidget::NativeTick(const FGeometry& MyGeometry, float InDeltaTime)
 { 
 	// If this ensure is hit it is likely UpdateCanTick as not called somewhere
@@ -1990,11 +2153,13 @@ void UUserWidget::StopAnimationsAndLatentActions()
 
 void UUserWidget::ListenForInputAction( FName ActionName, TEnumAsByte< EInputEvent > EventType, bool bConsume, FOnInputAction Callback )
 {
+	PRAGMA_DISABLE_DEPRECATION_WARNINGS
 	if ( !InputComponent )
 	{
 		InitializeInputComponent();
 	}
-
+	PRAGMA_ENABLE_DEPRECATION_WARNINGS
+	
 	if ( InputComponent )
 	{
 		FInputActionBinding NewBinding( ActionName, EventType.GetValue() );
@@ -2117,6 +2282,7 @@ void UUserWidget::OnInputAction( FOnInputAction Callback )
 	}
 }
 
+PRAGMA_DISABLE_DEPRECATION_WARNINGS
 void UUserWidget::InitializeInputComponent()
 {
 	if ( APlayerController* Controller = GetOwningPlayer() )
@@ -2126,10 +2292,10 @@ void UUserWidget::InitializeInputComponent()
 		// class then this would fail
 		UClass* InputClass = Controller->InputComponent ? Controller->InputComponent->GetClass() : UInputSettings::GetDefaultInputComponentClass();
 		InputComponent = NewObject< UInputComponent >( this, InputClass, NAME_None, RF_Transient );
-PRAGMA_DISABLE_DEPRECATION_WARNINGS
+
 		InputComponent->bBlockInput = bStopAction;
 		InputComponent->Priority = Priority;
-PRAGMA_ENABLE_DEPRECATION_WARNINGS
+
 		Controller->PushInputComponent( InputComponent );
 	}
 	else
@@ -2137,6 +2303,7 @@ PRAGMA_ENABLE_DEPRECATION_WARNINGS
 		FMessageLog("PIE").Info(FText::Format(LOCTEXT("NoInputListeningWithoutPlayerController", "Unable to listen to input actions without a player controller in {0}."), FText::FromName(GetClass()->GetFName())));
 	}
 }
+PRAGMA_ENABLE_DEPRECATION_WARNINGS
 
 void UUserWidget::UpdateCanTick() 
 {
diff --git a/Engine/Source/Runtime/UMG/Public/Blueprint/UserWidget.h b/Engine/Source/Runtime/UMG/Public/Blueprint/UserWidget.h
--- a/Engine/Source/Runtime/UMG/Public/Blueprint/UserWidget.h
+++ b/Engine/Source/Runtime/UMG/Public/Blueprint/UserWidget.h
@@ -1698,7 +1698,8 @@ protected:
 
 	UMG_API void OnInputAction( FOnInputAction Callback );
 
-	UMG_API virtual void InitializeInputComponent();
+	UE_DEPRECATED(5.7, "InitializeInputComponent is deprecated. bAutomaticallyRegisterInputOnConstruction will be set on compilation of this widget and automatically create the input component instead.")
+	UMG_API virtual void InitializeInputComponent() final;
 
 private:
 	FVector2D MinimumDesiredSize;
@@ -1715,9 +1716,48 @@ private:
 	FWidgetChild DesiredFocusWidget;
 
 protected:
+	
+	/**
+	 * Spawns the InputComponent and binds any script delegates if bAutomaticallyRegisterInputOnConstruction is true.
+	 * 
+	 * bAutomaticallyRegisterInputOnConstruction is true when there are input event nodes in the graph
+	 * of this widget, which would in turn need to be bound in order to receive their callbacks.
+	 *
+	 * This should only happen once per widget when it is spawned and initialized.
+	 */
+	UMG_API virtual void CreateInputComponent();
+
+	/**
+	 * Pushes the input component to the input stack and allows for any
+	 * script delegates on this widget to get called.
+	 */
+	UMG_API virtual void StartProcessingInputScriptDelegates();
+
+	/**
+	 * Removes this widget's input component from the input stack,
+	 * stopping any script delegates from being called.
+	 * 
+	 * This does not actually destroy or unbind any of the delegates yet, because the widget may be added 
+	 * to the viewport again the future, and the binding process can be expensive.
+	 */
+	UMG_API virtual void StopProcessingInputScriptDelegates();
+
+	/**
+	* Unbinds all script delegates that this widget has bound to on its input component,
+	* removes it from the input stack, and deletes the actual input component object.
+	*/
+	UMG_API virtual void DestroyInputComponent();
+	
 	UPROPERTY(Transient, DuplicateTransient)
 	TObjectPtr<class UInputComponent> InputComponent;
 
+	/**
+	 * A callback bound to this widget's owning local player which will fire when the player controller has changed.
+	 */
+	void HandleOwningLocalPlayerControllerChanged(APlayerController* NewPlayerController);
+	void StartListeningForPlayerControllerChanges();
+	void StopListeningForPlayerControllerChanges();
+
 protected:
 	UPROPERTY(Transient, DuplicateTransient)
 	TArray<FAnimationEventBinding> AnimationCallbacks;
