PR #13460: Fixed nullptr exception when loading PanelWidget

diff --git a/Engine/Source/Runtime/UMG/Public/Components/PanelWidget.h b/Engine/Source/Runtime/UMG/Public/Components/PanelWidget.h
--- a/Engine/Source/Runtime/UMG/Public/Components/PanelWidget.h
+++ b/Engine/Source/Runtime/UMG/Public/Components/PanelWidget.h
@@ -146,12 +146,15 @@ public:
 
 	virtual void ConnectEditorData() override
 	{
-		for ( UPanelSlot* PanelSlot : Slots )
+		for (UPanelSlot* PanelSlot : Slots)
 		{
-			PanelSlot->Parent = this;
-			if ( PanelSlot->Content )
+			if (PanelSlot)
 			{
-				PanelSlot->Content->Slot = PanelSlot;
+				PanelSlot->Parent = this;
+				if (PanelSlot->Content)
+				{
+					PanelSlot->Content->Slot = PanelSlot;
+				}
 			}
 		}
 	}
