UIComponents: Fix implicit TObjectPtr to raw pointer cast

#rnx

diff --git a/Engine/Source/Runtime/UMG/Private/Extensions/UIComponents/NavigationUIComponent.cpp b/Engine/Source/Runtime/UMG/Private/Extensions/UIComponents/NavigationUIComponent.cpp
--- a/Engine/Source/Runtime/UMG/Private/Extensions/UIComponents/NavigationUIComponent.cpp
+++ b/Engine/Source/Runtime/UMG/Private/Extensions/UIComponents/NavigationUIComponent.cpp
@@ -38,8 +38,8 @@ namespace UE::NavigationUIComponent::Private
 			{
 				Widget = UserWidget->GetCachedWidget().ToSharedRef();
 
-				UWidgetTree* WidgetTree = UserWidget != nullptr ? UserWidget->WidgetTree : nullptr;
-				if(TObjectPtr<UWidget> RootWidget = WidgetTree != nullptr ? WidgetTree->RootWidget : nullptr)
+				TObjectPtr<UWidgetTree> WidgetTree = (UserWidget != nullptr) ? UserWidget->WidgetTree : nullptr;
+				if(TObjectPtr<UWidget> RootWidget = (WidgetTree != nullptr) ? WidgetTree->RootWidget : nullptr)
 				{
 					Widget = RootWidget->GetCachedWidget().ToSharedRef();
 				}
@@ -99,15 +99,15 @@ void UNavigationUIComponent::HandleNavigationTransition(const FNavigationTransit
 
 	switch (NavigationTransition.Direction)
 	{
-	case ENavigationTransitionDirection::Incoming:
-	{
-		OnNavigationEnteredDelegate.ExecuteIfBound(NavigationTransition.Type, OldFocusedWidgetObject, NewFocusedWidgetObject);
-		break;
-	}
-	case ENavigationTransitionDirection::Outgoing:
-	{
-		OnNavigationExitedDelegate.ExecuteIfBound(NavigationTransition.Type, OldFocusedWidgetObject, NewFocusedWidgetObject);
-		break;
-	}
+		case ENavigationTransitionDirection::Incoming:
+		{
+			OnNavigationEnteredDelegate.ExecuteIfBound(NavigationTransition.Type, OldFocusedWidgetObject, NewFocusedWidgetObject);
+			break;
+		}
+		case ENavigationTransitionDirection::Outgoing:
+		{
+			OnNavigationExitedDelegate.ExecuteIfBound(NavigationTransition.Type, OldFocusedWidgetObject, NewFocusedWidgetObject);
+			break;
+		}
 	}
 }
diff --git a/Engine/Source/Runtime/UMG/Public/Extensions/UIComponents/NavigationUIComponent.h b/Engine/Source/Runtime/UMG/Public/Extensions/UIComponents/NavigationUIComponent.h
--- a/Engine/Source/Runtime/UMG/Public/Extensions/UIComponents/NavigationUIComponent.h
+++ b/Engine/Source/Runtime/UMG/Public/Extensions/UIComponents/NavigationUIComponent.h
@@ -6,11 +6,9 @@
 #include "Input/NavigationMetadata.h"
 #include "Extensions/UIComponent.h"
 #include "UObject/ObjectMacros.h"
-#include "Widgets/SWidget.h"
 
 #include "NavigationUIComponent.generated.h"
 
-class SWidget;
 class UUserWidget;
 
 /**
