Sequencer: fix possible GC leak with UMG animations

When a Sequencer runs an event track that invokes events defined in that Sequence's Blueprint director, it caches a pointer to the Blueprint in order to keep it alive (this is in particular needed for Level Sequences, where that Blueprint is not used by anybody else than the Sequence).

This can however lead to GC leaks with UMG animations. For User Widgets, the Blueprint director is the widget itself, so the Sequence keeps a strong reference to it while it is playing. So:

1. If the widget animation has an event track that invokes a Blueprint function on the widget itself, and
2. That widget is removed from its parent (or the viewport) while that animation is playing, and
3. The pool of inactive widgets is discarded and made available to the GC...

...then the GC will complain that something (the Sequence) is still holding onto the widget. This somehow happens before the widget's BeginDestroy is called, where we tear-down animations.

One possible fix is to tear-down all ongoing animations when the widget is removed from its parent and released back to the widget pool, but this isn't acceptable because we want the animation to resume if the widget resurfaces.

So instead, this fix makes Sequencer use weak pointers for Sequence Blueprint directors in the case of User Widgets. Normally this cache is lazily created when needed, but we create it ahead of time to make sure it's created with the flag to make it use weak pointers.
The change also adds more detailed logging to help with any future investigations.

diff --git a/Engine/Source/Runtime/MovieScene/Private/Evaluation/SequenceDirectorPlaybackCapability.cpp b/Engine/Source/Runtime/MovieScene/Private/Evaluation/SequenceDirectorPlaybackCapability.cpp
--- a/Engine/Source/Runtime/MovieScene/Private/Evaluation/SequenceDirectorPlaybackCapability.cpp
+++ b/Engine/Source/Runtime/MovieScene/Private/Evaluation/SequenceDirectorPlaybackCapability.cpp
@@ -31,6 +31,7 @@ void FSequenceDirectorPlaybackCapability::ResetDirectorInstances()
 	{
 		Cache->DirectorInstances.Reset();
 	}
+	WeakCache.Reset();
 }
 
 void FSequenceDirectorPlaybackCapability::InvalidateCachedData(UMovieSceneEntitySystemLinker* Linker)
@@ -40,12 +41,14 @@ void FSequenceDirectorPlaybackCapability::InvalidateCachedData(UMovieSceneEntity
 
 UObject* FSequenceDirectorPlaybackCapability::GetOrCreateDirectorInstance(TSharedRef<const FSharedPlaybackState> SharedPlaybackState, FMovieSceneSequenceIDRef SequenceID)
 {
-	if (!Cache)
+	if (bUseStrongCache && !Cache)
 	{
 		Cache = MakeUnique<FDirectorInstanceCache>();
 	}
 
-	UObject* ExistingDirectorInstance = Cache->DirectorInstances.FindRef(SequenceID);
+	UObject* ExistingDirectorInstance = bUseStrongCache ?
+		Cache->DirectorInstances.FindRef(SequenceID).Get() :
+		WeakCache.FindRef(SequenceID).Get();
 #if WITH_EDITOR
 	if (ExistingDirectorInstance)
 	{
@@ -79,7 +82,14 @@ UObject* FSequenceDirectorPlaybackCapability::GetOrCreateDirectorInstance(TShare
 
 	if (NewDirectorInstance)
 	{
-		Cache->DirectorInstances.Add(SequenceID, NewDirectorInstance);
+		if (bUseStrongCache)
+		{
+			Cache->DirectorInstances.Add(SequenceID, NewDirectorInstance);
+		}
+		else
+		{
+			WeakCache.Add(SequenceID, NewDirectorInstance);
+		}
 	}
 
 	return NewDirectorInstance;
diff --git a/Engine/Source/Runtime/MovieScene/Public/Evaluation/SequenceDirectorPlaybackCapability.h b/Engine/Source/Runtime/MovieScene/Public/Evaluation/SequenceDirectorPlaybackCapability.h
--- a/Engine/Source/Runtime/MovieScene/Public/Evaluation/SequenceDirectorPlaybackCapability.h
+++ b/Engine/Source/Runtime/MovieScene/Public/Evaluation/SequenceDirectorPlaybackCapability.h
@@ -21,6 +21,13 @@ struct FSequenceDirectorPlaybackCapability : public IPlaybackCapability
 {
 	UE_DECLARE_MOVIESCENE_PLAYBACK_CAPABILITY_API(MOVIESCENE_API, FSequenceDirectorPlaybackCapability)
 
+public:
+
+	FSequenceDirectorPlaybackCapability() {}
+	FSequenceDirectorPlaybackCapability(bool bInUseStrongCache) : bUseStrongCache(bInUseStrongCache) {}
+
+public:
+
 	/** Remove all director blueprint instances */
 	MOVIESCENE_API void ResetDirectorInstances();
 
@@ -45,6 +52,11 @@ private:
 		TSortedMap<FMovieSceneSequenceID, TObjectPtr<UObject>> DirectorInstances;
 	};
 	TUniquePtr<FDirectorInstanceCache> Cache;
+
+	/** An alternative cache using weak-pointers. */
+	TSortedMap<FMovieSceneSequenceID, TWeakObjectPtr<UObject>> WeakCache;
+
+	bool bUseStrongCache = true;
 };
 
 }  // namespace UE::MovieScene
diff --git a/Engine/Source/Runtime/UMG/Private/Animation/WidgetAnimationState.cpp b/Engine/Source/Runtime/UMG/Private/Animation/WidgetAnimationState.cpp
--- a/Engine/Source/Runtime/UMG/Private/Animation/WidgetAnimationState.cpp
+++ b/Engine/Source/Runtime/UMG/Private/Animation/WidgetAnimationState.cpp
@@ -12,6 +12,7 @@
 #include "Evaluation/MovieSceneEvaluationState.h"
 #include "Evaluation/MovieScenePlayback.h"
 #include "Evaluation/MovieSceneSequenceHierarchy.h"
+#include "Evaluation/SequenceDirectorPlaybackCapability.h"
 #include "MovieScene.h"
 #include "MovieSceneFwd.h"
 #include "MovieSceneLegacyPlayer.h"
@@ -94,6 +95,12 @@ void FWidgetAnimationState::Initialize(UWidgetAnimation* InAnimation, UUserWidge
 		WeakPlaybackState = SharedPlaybackState;
 
 		SharedPlaybackState->AddCapability<FMovieSceneEvaluationState>();
+
+		// The UUserWidget that owns us is its own Sequence Director Blueprint, so don't create a cycle of 
+		// strong GC references: create the cache right away and initialize it to use weak references.
+		const bool bUseStrongDirectorCache = false;
+		SharedPlaybackState->AddCapability<FSequenceDirectorPlaybackCapability>(bUseStrongDirectorCache);
+
 		if (bNeedsLegacyPlayer)
 		{
 			GetOrCreateLegacyPlayer();
@@ -114,8 +121,8 @@ void FWidgetAnimationState::Initialize(UWidgetAnimation* InAnimation, UUserWidge
 #if !NO_LOGGING
 	if (UE::UMG::GVarAnimationDetailedLogging)
 	{
-		UE_LOG(LogUMG, Log, TEXT("Animation: Initializing '%s', instance [%d]%s"), 
-				*GetNameSafe(Animation), RootInstanceHandle.InstanceID,
+		UE_LOG(LogUMG, Log, TEXT("Animation: Initializing '%s' playing on '%s' (%p), instance [%d]%s"), 
+				*GetNameSafe(Animation), *GetNameSafe(InUserWidget), InUserWidget, RootInstanceHandle.InstanceID,
 				(LegacyPlayer ? TEXT(", with legacy player") : TEXT("")));
 	}
 #endif
@@ -178,8 +185,8 @@ void FWidgetAnimationState::Tick(float InDeltaSeconds)
 #if !NO_LOGGING
 		if (UE::UMG::GVarAnimationDetailedLogging)
 		{
-			UE_LOG(LogUMG, Log, TEXT("Animation: Ticking '%s' aborted, invalid state."),
-					*GetNameSafe(Animation));
+			UE_LOG(LogUMG, Log, TEXT("Animation: Ticking '%s' on '%s' (%p) aborted, invalid state."),
+					*GetNameSafe(Animation), *GetNameSafe(UserWidget), UserWidget);
 		}
 #endif
 		return;
@@ -191,8 +198,8 @@ void FWidgetAnimationState::Tick(float InDeltaSeconds)
 #if !NO_LOGGING
 		if (UE::UMG::GVarAnimationDetailedLogging)
 		{
-			UE_LOG(LogUMG, Log, TEXT("Animation: Ticking '%s' aborted, shared playback state has no runner."),
-					*GetNameSafe(Animation));
+			UE_LOG(LogUMG, Log, TEXT("Animation: Ticking '%s' on '%s' (%p) aborted, shared playback state has no runner."),
+					*GetNameSafe(Animation), *GetNameSafe(UserWidget), UserWidget);
 		}
 #endif
 		return;
@@ -247,8 +254,8 @@ void FWidgetAnimationState::Tick(float InDeltaSeconds)
 #if !NO_LOGGING
 			if (UE::UMG::GVarAnimationDetailedLogging)
 			{
-				UE_LOG(LogUMG, Log, TEXT("Animation: Ticking '%s', Time=%s, Status=%s, Direction=%s (update %d)"),
-						*GetNameSafe(Animation),
+				UE_LOG(LogUMG, Log, TEXT("Animation: Ticking '%s' on '%s' (%p), Time=%s, Status=%s, Direction=%s (update %d)"),
+						*GetNameSafe(Animation), *GetNameSafe(UserWidget), UserWidget,
 						*LexToString(TickContext.GetTime()),
 						*UEnum::GetValueAsString(TickContext.GetStatus()),
 						((TickContext.GetDirection() == EPlayDirection::Forwards) ? TEXT("Forwards") : TEXT("Backwards")),
@@ -260,8 +267,8 @@ void FWidgetAnimationState::Tick(float InDeltaSeconds)
 #if !NO_LOGGING
 		if (UE::UMG::GVarAnimationDetailedLogging)
 		{
-			UE_LOG(LogUMG, Log, TEXT("Animation: Ticking done '%s', NumLoopsCompleted=%d, PlaybackSpeed=%f"),
-					*GetNameSafe(Animation),
+			UE_LOG(LogUMG, Log, TEXT("Animation: Ticking done '%s' on '%s' (%p), NumLoopsCompleted=%d, PlaybackSpeed=%f"),
+					*GetNameSafe(Animation), *GetNameSafe(UserWidget), UserWidget,
 					PlaybackManager.GetNumLoopsCompleted(),
 					PlaybackManager.GetPlayRate());
 		}
@@ -274,8 +281,8 @@ void FWidgetAnimationState::Tick(float InDeltaSeconds)
 #if !NO_LOGGING
 		if (UE::UMG::GVarAnimationDetailedLogging)
 		{
-			UE_LOG(LogUMG, Log, TEXT("Animation: Ticking '%s' skipped, state not playing."),
-					*GetNameSafe(Animation));
+			UE_LOG(LogUMG, Log, TEXT("Animation: Ticking '%s' on '%s' (%p) skipped, state not playing."),
+					*GetNameSafe(Animation), *GetNameSafe(UserWidget), UserWidget);
 		}
 #endif
 	}
@@ -290,8 +297,9 @@ void FWidgetAnimationState::OnBegunPlay()
 #if !NO_LOGGING
 	if (UE::UMG::GVarAnimationDetailedLogging)
 	{
-		UE_LOG(LogUMG, Log, TEXT("Animation: OnBegunPlay '%s'"),
-				*GetNameSafe(Animation));
+		UUserWidget* UserWidget = WeakUserWidget.Get();
+		UE_LOG(LogUMG, Log, TEXT("Animation: OnBegunPlay '%s' on '%s' (%p)"),
+				*GetNameSafe(Animation), *GetNameSafe(UserWidget), UserWidget);
 	}
 #endif
 }
@@ -328,8 +336,9 @@ void FWidgetAnimationState::OnStopped()
 	if (UE::UMG::GVarAnimationDetailedLogging)
 	{
 		PRAGMA_DISABLE_DEPRECATION_WARNINGS
-		UE_LOG(LogUMG, Log, TEXT("Animation: OnStopped '%s' %s%s"),
-				*GetNameSafe(Animation),
+		UUserWidget* UserWidget = WeakUserWidget.Get();
+		UE_LOG(LogUMG, Log, TEXT("Animation: OnStopped '%s' on '%s' (%p) %s%s"),
+				*GetNameSafe(Animation), *GetNameSafe(UserWidget), UserWidget,
 				OnWidgetAnimationFinishedEvent.IsBound() ? TEXT(", OnWidgetAnimationFinishedEvent bound") : TEXT(""),
 				(LegacyPlayer && LegacyPlayer->OnSequenceFinishedPlaying().IsBound()) ? TEXT(", OnSequenceFinishedPlaying bound") : TEXT(""));
 		PRAGMA_ENABLE_DEPRECATION_WARNINGS
@@ -534,8 +543,8 @@ void FWidgetAnimationState::Play(const FWidgetAnimationStatePlayParams& PlayPara
 #if !NO_LOGGING
 	if (UE::UMG::GVarAnimationDetailedLogging)
 	{
-		UE_LOG(LogUMG, Log, TEXT("Animation: Play '%s', StartTime=%s, NumLoopsToPlay=%d, PlaybackSpeed=%f, PlayMode=%s%s"),
-				*GetNameSafe(Animation),
+		UE_LOG(LogUMG, Log, TEXT("Animation: Play '%s' on '%s' (%p), StartTime=%s, NumLoopsToPlay=%d, PlaybackSpeed=%f, PlayMode=%s%s"),
+				*GetNameSafe(Animation), *GetNameSafe(UserWidget), UserWidget,
 				*LexToString(PlaybackManager.GetCurrentTime()),
 				PlaybackManager.GetNumLoopsToPlay(),
 				PlaybackManager.GetPlayRate(),
@@ -595,8 +604,9 @@ void FWidgetAnimationState::Stop()
 #if !NO_LOGGING
 	if (UE::UMG::GVarAnimationDetailedLogging)
 	{
-		UE_LOG(LogUMG, Log, TEXT("Animation: Stop '%s'"),
-				*GetNameSafe(Animation));
+		UUserWidget* UserWidget = WeakUserWidget.Get();
+		UE_LOG(LogUMG, Log, TEXT("Animation: Stop '%s' on '%s' (%p)"),
+				*GetNameSafe(Animation), *GetNameSafe(UserWidget), UserWidget);
 	}
 #endif
 
@@ -629,8 +639,9 @@ void FWidgetAnimationState::Pause()
 #if !NO_LOGGING
 	if (UE::UMG::GVarAnimationDetailedLogging)
 	{
-		UE_LOG(LogUMG, Log, TEXT("Animation: Pause '%s', PauseTime=%s"),
-				*GetNameSafe(Animation),
+		UUserWidget* UserWidget = WeakUserWidget.Get();
+		UE_LOG(LogUMG, Log, TEXT("Animation: Pause '%s' on '%s' (%p), PauseTime=%s"),
+				*GetNameSafe(Animation), *GetNameSafe(UserWidget), UserWidget,
 				*LexToString(PlaybackManager.GetCurrentTime()));
 	}
 #endif
@@ -673,8 +684,9 @@ void FWidgetAnimationState::RemoveEvaluationData()
 #if !NO_LOGGING
 	if (UE::UMG::GVarAnimationDetailedLogging)
 	{
-		UE_LOG(LogUMG, Log, TEXT("Animation: RemoveEvaluationData '%s'"),
-				*GetNameSafe(Animation));
+		UUserWidget* UserWidget = WeakUserWidget.Get();
+		UE_LOG(LogUMG, Log, TEXT("Animation: RemoveEvaluationData '%s' on '%s' (%p)"),
+				*GetNameSafe(Animation), *GetNameSafe(UserWidget), UserWidget);
 	}
 #endif
 }
@@ -711,8 +723,9 @@ void FWidgetAnimationState::TearDown()
 #if !NO_LOGGING
 	if (UE::UMG::GVarAnimationDetailedLogging)
 	{
-		UE_LOG(LogUMG, Log, TEXT("Animation: TearDown '%s'"),
-				*GetNameSafe(Animation));
+		UUserWidget* UserWidget = WeakUserWidget.Get();
+		UE_LOG(LogUMG, Log, TEXT("Animation: TearDown '%s' on '%s' (%p)"),
+				*GetNameSafe(Animation), *GetNameSafe(UserWidget), UserWidget);
 	}
 #endif
 }
