Created a Mouse Hover Component, which exposes a variable `Is Hovered` that is true when the Component's owner widget is hovered over by the mouse, and false when it is not.

#tests I created a bunch of different widgets that I displayed on the viewport, including nested widgets. I then made a binding from IsHovered to a mvvm view model variable, and used that variable to change the text in a textbox to "Hovered" vs "Not Hovered". As my mouse moved over those widgets, the text changed, and changed back when my mouse left them. Everything worked as expected.
[RN] minor ui/slate

diff --git a/Engine/Source/Runtime/UMG/Private/Components/MouseHoverComponent.cpp b/Engine/Source/Runtime/UMG/Private/Components/MouseHoverComponent.cpp
new file mode 100644
--- /dev/null
+++ b/Engine/Source/Runtime/UMG/Private/Components/MouseHoverComponent.cpp
@@ -0,0 +1,90 @@
+// Copyright Epic Games, Inc. All Rights Reserved.
+
+#include "Components/MouseHoverComponent.h"
+
+#include "Blueprint/UserWidget.h"
+#include "Components/Widget.h"
+#include "Widgets/SCompoundWidget.h"
+#include "Widgets/SWidget.h"
+
+class SMouseHoverWidget : public SCompoundWidget
+{
+public:
+	SLATE_BEGIN_ARGS(SMouseHoverWidget)
+		: _Content()
+		{
+		}
+		SLATE_DEFAULT_SLOT(FArguments, Content)
+	SLATE_END_ARGS()
+
+	void Construct(const FArguments& InArgs, UMouseHoverComponent* InMouseEventsComponent);
+
+protected:
+
+	virtual void OnMouseEnter(const FGeometry& MyGeometry, const FPointerEvent& MouseEvent) override;
+	virtual void OnMouseLeave(const FPointerEvent& MouseEvent) override;
+
+	TWeakObjectPtr<UMouseHoverComponent> MouseHoverComponent;
+};
+
+void SMouseHoverWidget::Construct(const FArguments& InArgs, UMouseHoverComponent* InMouseHoverComponent)
+{
+	MouseHoverComponent = InMouseHoverComponent;
+	ChildSlot
+	[
+		InArgs._Content.Widget
+	];
+}
+
+void SMouseHoverWidget::OnMouseEnter(const FGeometry& MyGeometry, const FPointerEvent& MouseEvent)
+{
+	if (MouseHoverComponent.IsValid())
+	{
+		const bool bWasHovered = IsHovered();
+
+		// Call super class's implementation
+		SWidget::OnMouseEnter(MyGeometry, MouseEvent);
+
+		if (!bWasHovered && IsHovered())
+		{
+			MouseHoverComponent->OnMouseHoverChanged(true);
+		}
+	}
+}
+
+void SMouseHoverWidget::OnMouseLeave(const FPointerEvent& MouseEvent)
+{
+	if (MouseHoverComponent.IsValid())
+	{
+		const bool bWasHovered = IsHovered();
+
+		// Call super class's implementation
+		SWidget::OnMouseLeave(MouseEvent);
+
+		if (bWasHovered && !IsHovered())
+		{
+			MouseHoverComponent->OnMouseHoverChanged(false);
+		}
+	}
+}
+
+void UMouseHoverComponent::OnMouseHoverChanged(bool InbIsHovered)
+{
+	if (GetIsHovered() != InbIsHovered)
+	{
+		bIsHovered = InbIsHovered;
+		BroadcastFieldValueChanged(FFieldNotificationClassDescriptor::bIsHovered);
+	}
+}
+
+TSharedRef<SWidget> UMouseHoverComponent::RebuildWidgetWithContent(TSharedRef<SWidget> OwnerContent)
+{
+	TSharedRef<SWidget> WrapperWidget = 
+		SNew(SMouseHoverWidget, this)
+		[
+			OwnerContent
+		];
+
+	return WrapperWidget;
+}
+
diff --git a/Engine/Source/Runtime/UMG/Public/Components/MouseHoverComponent.h b/Engine/Source/Runtime/UMG/Public/Components/MouseHoverComponent.h
new file mode 100644
--- /dev/null
+++ b/Engine/Source/Runtime/UMG/Public/Components/MouseHoverComponent.h
@@ -0,0 +1,29 @@
+// Copyright Epic Games, Inc. All Rights Reserved.
+
+#pragma once
+
+#include "CoreMinimal.h"
+#include "Extensions/UIComponent.h"
+#include "UObject/ObjectMacros.h"
+
+#include "MouseHoverComponent.generated.h"
+
+/** This is a class for a Component that exposes the mouse hover state of its Owner widget. */
+UCLASS(MinimalAPI, Experimental)
+class UMouseHoverComponent : public UUIComponent
+{
+	GENERATED_BODY()
+
+public:
+
+	bool GetIsHovered() const { return bIsHovered; }
+
+	void OnMouseHoverChanged(bool InbIsHovered);
+
+private:
+
+	UPROPERTY(VisibleAnywhere, BlueprintReadOnly, Getter = "GetIsHovered", FieldNotify, Category = "Mouse Hover State", meta = (AllowPrivateAccess))
+	bool bIsHovered = false;
+
+	virtual TSharedRef<SWidget> RebuildWidgetWithContent(const TSharedRef<SWidget> OwnerContent) override;
+};
\ No newline at end of file
