Add InitializeInEditor function to UserWidgetExtension and WidgetBlueprintGeneratedClassExtension to be able to support proper initialization in Editor for UIComponent Extensions

diff --git a/Engine/Source/Editor/UMGEditor/Private/UIComponentWidgetBlueprintExtension.cpp b/Engine/Source/Editor/UMGEditor/Private/UIComponentWidgetBlueprintExtension.cpp
--- a/Engine/Source/Editor/UMGEditor/Private/UIComponentWidgetBlueprintExtension.cpp
+++ b/Engine/Source/Editor/UMGEditor/Private/UIComponentWidgetBlueprintExtension.cpp
@@ -119,7 +119,6 @@ UUIComponentUserWidgetExtension* UUIComponentWidgetBlueprintExtension::GetOrCrea
 	{
 		UserWidgetExtension = PreviewWidget->AddExtension<UUIComponentUserWidgetExtension>();
 		UserWidgetExtension->InitializeContainer(DuplicateContainer(UserWidgetExtension));
-		UserWidgetExtension->Initialize();
 	}
 	else
 	{
diff --git a/Engine/Source/Runtime/UMG/Private/Extensions/UIComponent.cpp b/Engine/Source/Runtime/UMG/Private/Extensions/UIComponent.cpp
--- a/Engine/Source/Runtime/UMG/Private/Extensions/UIComponent.cpp
+++ b/Engine/Source/Runtime/UMG/Private/Extensions/UIComponent.cpp
@@ -11,15 +11,19 @@
 
 void UUIComponent::Initialize(UWidget* Target)
 {
+	ensure(!bInitialized);
+	ensure(!Target->IsConstructed());
 	ensure(Target);
 	ensure(Owner.IsExplicitlyNull());
 	Owner = Target;
 	Owner->bWrappedByComponent = true;
 	OnInitialize();
+	bInitialized = true;
 }
 
 void UUIComponent::PreConstruct(bool bIsDesignTime)
 {
+	ensure(bInitialized);
 	OnPreConstruct(bIsDesignTime);
 }
 
diff --git a/Engine/Source/Runtime/UMG/Private/Extensions/UIComponentUserWidgetExtension.cpp b/Engine/Source/Runtime/UMG/Private/Extensions/UIComponentUserWidgetExtension.cpp
--- a/Engine/Source/Runtime/UMG/Private/Extensions/UIComponentUserWidgetExtension.cpp
+++ b/Engine/Source/Runtime/UMG/Private/Extensions/UIComponentUserWidgetExtension.cpp
@@ -12,11 +12,14 @@
 
 void UUIComponentUserWidgetExtension::Initialize()
 {
+	ensure(!bInitialized);
 	InitializeComponents();
+	bInitialized = true;
 }
 
 void UUIComponentUserWidgetExtension::PreConstruct(bool bIsDesignTime)
 {
+	ensure(bInitialized);
 	Super::PreConstruct(bIsDesignTime);
 
 	if (ComponentContainer)
@@ -82,6 +85,12 @@ void UUIComponentUserWidgetExtension::InitializeComponents()
 
 #if WITH_EDITOR
 
+void UUIComponentUserWidgetExtension::InitializeInEditor()
+{
+	// We simply call the Initialize function as we need the same code to execute in Editor.
+	Initialize();
+}
+
 void UUIComponentUserWidgetExtension::RenameWidget(const FName& OldVarName, const FName& NewVarName)
 {
 	if (ensure(ComponentContainer))
diff --git a/Engine/Source/Runtime/UMG/Private/Extensions/UIComponentWidgetBlueprintGeneratedClassExtension.cpp b/Engine/Source/Runtime/UMG/Private/Extensions/UIComponentWidgetBlueprintGeneratedClassExtension.cpp
--- a/Engine/Source/Runtime/UMG/Private/Extensions/UIComponentWidgetBlueprintGeneratedClassExtension.cpp
+++ b/Engine/Source/Runtime/UMG/Private/Extensions/UIComponentWidgetBlueprintGeneratedClassExtension.cpp
@@ -58,6 +58,13 @@ bool UUIComponentWidgetBlueprintGeneratedClassExtension::VerifyContainer(const U
 }
 
 #if WITH_EDITOR
+
+void UUIComponentWidgetBlueprintGeneratedClassExtension::InitializeInEditor(UUserWidget* UserWidget)
+{
+	// We simply call the Initialize function as we need the same code to execute in Editor.
+	Initialize(UserWidget);
+}
+
 bool UUIComponentWidgetBlueprintGeneratedClassExtension::VerifyAllWidgetsExists(const UWidgetTree* WidgetTree) const
 {	
 	TSet<FName> Targets;
@@ -105,7 +112,6 @@ void UUIComponentWidgetBlueprintGeneratedClassExtension::PreConstruct(UUserWidge
 	{
 		UserWidgetExtension = UserWidget->AddExtension<UUIComponentUserWidgetExtension>();
 		UserWidgetExtension->InitializeContainer(DuplicateContainer(UserWidget));
-		UserWidgetExtension->Initialize();
 	}
 	else
 	{
diff --git a/Engine/Source/Runtime/UMG/Private/UserWidget.cpp b/Engine/Source/Runtime/UMG/Private/UserWidget.cpp
--- a/Engine/Source/Runtime/UMG/Private/UserWidget.cpp
+++ b/Engine/Source/Runtime/UMG/Private/UserWidget.cpp
@@ -169,7 +169,26 @@ bool UUserWidget::Initialize()
 		{
 			NativeOnInitialized();
 		}
+#if WITH_EDITOR
+		else if (IsDesignTime())
+		{
+			if (UWidgetBlueprintGeneratedClass* BPClass = Cast<UWidgetBlueprintGeneratedClass>(GetClass()))
+			{
+				BPClass->ForEachExtension([this](UWidgetBlueprintGeneratedClassExtension* Extension)
+					{
+						Extension->InitializeInEditor(this);
+					});
+			}
 
+			// Extension can add other extensions. Use index loop to initialize them all.
+			for (int32 Index = 0; Index < Extensions.Num(); ++Index)
+			{
+				UUserWidgetExtension* Extension = Extensions[Index];
+				check(Extension);
+				Extension->InitializeInEditor();
+			}			
+		}
+#endif // WITH_EDITOR
 		bInitialized = true;
 		return true;
 	}
diff --git a/Engine/Source/Runtime/UMG/Public/Extensions/UIComponent.h b/Engine/Source/Runtime/UMG/Public/Extensions/UIComponent.h
--- a/Engine/Source/Runtime/UMG/Public/Extensions/UIComponent.h
+++ b/Engine/Source/Runtime/UMG/Public/Extensions/UIComponent.h
@@ -83,4 +83,6 @@ private:
 	UPROPERTY(transient, DuplicateTransient)
 	TWeakObjectPtr<UWidget> Owner;
 	UE::FieldNotification::FFieldMulticastDelegate Delegates;
+
+	bool bInitialized = false;
 };
diff --git a/Engine/Source/Runtime/UMG/Public/Extensions/UIComponentUserWidgetExtension.h b/Engine/Source/Runtime/UMG/Public/Extensions/UIComponentUserWidgetExtension.h
--- a/Engine/Source/Runtime/UMG/Public/Extensions/UIComponentUserWidgetExtension.h
+++ b/Engine/Source/Runtime/UMG/Public/Extensions/UIComponentUserWidgetExtension.h
@@ -35,6 +35,8 @@ public:
 	
 
 #if WITH_EDITOR
+	UE_API virtual void InitializeInEditor() override;
+
 	UE_API void RenameWidget(const FName& OldVarName, const FName& NewVarName);
 	
 	// Used only to create a Component on the PreviewWidget in the editor, based on the Component Archetype object in the WidgetBlueprint.
@@ -52,6 +54,8 @@ private:
 	// Use a single TArray for the Entire UUserWidget to reduce memory usage.
 	UPROPERTY(Transient, DuplicateTransient)
 	TObjectPtr<UUIComponentContainer> ComponentContainer;
+
+	bool bInitialized = false;
 };
 
 #undef UE_API
diff --git a/Engine/Source/Runtime/UMG/Public/Extensions/UIComponentWidgetBlueprintGeneratedClassExtension.h b/Engine/Source/Runtime/UMG/Public/Extensions/UIComponentWidgetBlueprintGeneratedClassExtension.h
--- a/Engine/Source/Runtime/UMG/Public/Extensions/UIComponentWidgetBlueprintGeneratedClassExtension.h
+++ b/Engine/Source/Runtime/UMG/Public/Extensions/UIComponentWidgetBlueprintGeneratedClassExtension.h
@@ -30,6 +30,7 @@ public:
 	UE_API virtual void PreConstruct(UUserWidget* UserWidget, bool IsDesignTime) override;
 
 #if WITH_EDITOR
+	UE_API virtual void InitializeInEditor(UUserWidget* UserWidget) override;
 	UE_API bool VerifyAllWidgetsExists(const UWidgetTree* WidgetTree) const;
 #endif // WITH_EDITOR
 	
diff --git a/Engine/Source/Runtime/UMG/Public/Extensions/UserWidgetExtension.h b/Engine/Source/Runtime/UMG/Public/Extensions/UserWidgetExtension.h
--- a/Engine/Source/Runtime/UMG/Public/Extensions/UserWidgetExtension.h
+++ b/Engine/Source/Runtime/UMG/Public/Extensions/UserWidgetExtension.h
@@ -21,6 +21,12 @@ class UUserWidgetExtension : public UObject
 	GENERATED_BODY()
 	
 public:
+
+#ifdef WITH_EDITOR
+	/** Initialize function called only in Editor */
+	virtual void InitializeInEditor() { }
+#endif
+
 	/** Extend the UUserWidget::Initialize function */
 	virtual void Initialize()
 	{
diff --git a/Engine/Source/Runtime/UMG/Public/Extensions/WidgetBlueprintGeneratedClassExtension.h b/Engine/Source/Runtime/UMG/Public/Extensions/WidgetBlueprintGeneratedClassExtension.h
--- a/Engine/Source/Runtime/UMG/Public/Extensions/WidgetBlueprintGeneratedClassExtension.h
+++ b/Engine/Source/Runtime/UMG/Public/Extensions/WidgetBlueprintGeneratedClassExtension.h
@@ -17,6 +17,11 @@ class UWidgetBlueprintGeneratedClassExtension : public UObject
 	GENERATED_BODY()
 	
 public:
+
+#ifdef WITH_EDITOR
+	/** Initialize function called only in Editor */
+	virtual void InitializeInEditor(UUserWidget* UserWidget) { }
+#endif
 	/**
 	 * Extend the UUserWidget::Initialize function.
 	 * Called once when the UUserWidget is created.
