Adding utility functions to query private states in the widget. Useful to create widget validators

diff --git a/Engine/Source/Runtime/UMG/Private/Components/PostBufferUpdate.cpp b/Engine/Source/Runtime/UMG/Private/Components/PostBufferUpdate.cpp
--- a/Engine/Source/Runtime/UMG/Private/Components/PostBufferUpdate.cpp
+++ b/Engine/Source/Runtime/UMG/Private/Components/PostBufferUpdate.cpp
@@ -143,6 +143,21 @@ UMG_API void UPostBufferUpdate::ReleaseSlateResources(bool bReleaseChildren)
 	MyPostBufferUpdate.Reset();
 }
 
+#if WITH_EDITOR
+UMG_API bool UPostBufferUpdate::DoesUpdateSlatePostRT(ESlatePostRT SlateBufferBits) const
+{
+	for (const FSlatePostBufferUpdateInfo& Info : UpdateBufferInfos)
+	{
+		if (static_cast<uint8>(SlateBufferBits) & static_cast<uint8>(Info.BufferToUpdate))
+		{
+			return true;
+		}
+	}
+
+	return false;
+}
+#endif
+
 /////////////////////////////////////////////////////
 
 #undef LOCTEXT_NAMESPACE
diff --git a/Engine/Source/Runtime/UMG/Public/Components/PostBufferUpdate.h b/Engine/Source/Runtime/UMG/Public/Components/PostBufferUpdate.h
--- a/Engine/Source/Runtime/UMG/Public/Components/PostBufferUpdate.h
+++ b/Engine/Source/Runtime/UMG/Public/Components/PostBufferUpdate.h
@@ -76,6 +76,13 @@ public:
 	/** Set the orientation of the stack box. The existing elements will be rearranged. */
 	UMG_API void SetPerformDefaultPostBufferUpdate(bool bInPerformDefaultPostBufferUpdate);
 
+#if WITH_EDITOR
+	UMG_API bool GetUpdateOnlyPaintArea() const { return bUpdateOnlyPaintArea; }
+
+	/** True if this widget updates any buffer indicated in SlateBufferBits */
+	UMG_API bool DoesUpdateSlatePostRT(ESlatePostRT SlateBufferBits) const ;
+#endif
+
 protected:
 	//~ Begin UWidget Interface
 	UMG_API virtual TSharedRef<SWidget> RebuildWidget() override;
