UScrollBox : In OnDescendantDeselectedByDesigner ensure a valid this pointer in the lambda before calling ScrollToStart()

diff --git a/Engine/Source/Runtime/UMG/Private/Components/ScrollBox.cpp b/Engine/Source/Runtime/UMG/Private/Components/ScrollBox.cpp
--- a/Engine/Source/Runtime/UMG/Private/Components/ScrollBox.cpp
+++ b/Engine/Source/Runtime/UMG/Private/Components/ScrollBox.cpp
@@ -685,7 +685,7 @@ void UScrollBox::OnDescendantDeselectedByDesigner( UWidget* DescendantWidget )
 	}
 
 	// because we get a deselect before we get a select, we need to delay this call until we're sure we didn't scroll to another widget.
-	TickHandle = FTSTicker::GetCoreTicker().AddTicker( FTickerDelegate::CreateLambda( [this]( float ) -> bool
+	TickHandle = FTSTicker::GetCoreTicker().AddTicker( FTickerDelegate::CreateWeakLambda(this, [this]( float ) -> bool
 	                                                                                {
                                                                                         QUICK_SCOPE_CYCLE_COUNTER(STAT_UScrollBox_ScrollToStart_LambdaTick);
 		                                                                                this->ScrollToStart();
