Fixing a regression with a deprecated code path for fetching what animation is playing for UMG Sequences.

diff --git a/Engine/Source/Runtime/UMG/Private/Animation/UMGSequencePlayer.cpp b/Engine/Source/Runtime/UMG/Private/Animation/UMGSequencePlayer.cpp
--- a/Engine/Source/Runtime/UMG/Private/Animation/UMGSequencePlayer.cpp
+++ b/Engine/Source/Runtime/UMG/Private/Animation/UMGSequencePlayer.cpp
@@ -112,7 +112,7 @@ const UWidgetAnimation* UUMGSequencePlayer::GetAnimation() const
 {
 	if (FWidgetAnimationState* AnimState = WidgetAnimationHandle.GetAnimationState())
 	{
-		AnimState->GetAnimation();
+		return AnimState->GetAnimation();
 	}
 	return nullptr;
 }
