Created a Size Box UI Component which wraps the owner widget in a size box.

#tests Tried out the 8 different Size Box parameters in the designer view and PIE and everything looks good. Also made functions that were bound to a button that changed parameters at runtime, and it worked as expected.
[RN] minor ui/slate

diff --git a/Engine/Source/Runtime/UMG/Private/Components/SizeBoxComponent.cpp b/Engine/Source/Runtime/UMG/Private/Components/SizeBoxComponent.cpp
new file mode 100644
--- /dev/null
+++ b/Engine/Source/Runtime/UMG/Private/Components/SizeBoxComponent.cpp
@@ -0,0 +1,385 @@
+// Copyright Epic Games, Inc. All Rights Reserved.
+
+#include "Components/SizeBoxComponent.h"
+
+#include "Blueprint/UserWidget.h"
+#include "Components/Widget.h"
+#include "CoreGlobals.h"
+#include "Widgets/Layout/SBox.h"
+#include "Widgets/SWidget.h"
+
+#include "UMGPrivate.h"
+
+float USizeBoxComponent::GetWidthOverride() const
+{
+	return WidthOverride;
+}
+
+bool USizeBoxComponent::IsWidthOverride() const
+{
+	return bOverride_WidthOverride;
+}
+
+void USizeBoxComponent::SetWidthOverride(float InWidthOverride)
+{
+	bOverride_WidthOverride = true;
+	WidthOverride = InWidthOverride;
+	if (MySizeBox.IsValid())
+	{
+		MySizeBox->SetWidthOverride(InWidthOverride);
+	}
+}
+
+void USizeBoxComponent::ClearWidthOverride()
+{
+	bOverride_WidthOverride = false;
+	if (MySizeBox.IsValid())
+	{
+		MySizeBox->SetWidthOverride(FOptionalSize());
+	}
+}
+
+float USizeBoxComponent::GetHeightOverride() const
+{
+	return HeightOverride;
+}
+
+bool USizeBoxComponent::IsHeightOverride() const
+{
+	return bOverride_HeightOverride;
+}
+
+void USizeBoxComponent::SetHeightOverride(float InHeightOverride)
+{
+	bOverride_HeightOverride = true;
+	HeightOverride = InHeightOverride;
+	if (MySizeBox.IsValid())
+	{
+		MySizeBox->SetHeightOverride(InHeightOverride);
+	}
+}
+
+void USizeBoxComponent::ClearHeightOverride()
+{
+	bOverride_HeightOverride = false;
+	if (MySizeBox.IsValid())
+	{
+		MySizeBox->SetHeightOverride(FOptionalSize());
+	}
+}
+
+float USizeBoxComponent::GetMinDesiredWidth() const
+{
+	return MinDesiredWidth;
+}
+
+bool USizeBoxComponent::IsMinDesiredWidthOverride() const
+{
+	return bOverride_MinDesiredWidth;
+}
+
+void USizeBoxComponent::SetMinDesiredWidth(float InMinDesiredWidth)
+{
+	bOverride_MinDesiredWidth = true;
+	MinDesiredWidth = InMinDesiredWidth;
+	if (MySizeBox.IsValid())
+	{
+		MySizeBox->SetMinDesiredWidth(InMinDesiredWidth);
+	}
+}
+
+void USizeBoxComponent::ClearMinDesiredWidth()
+{
+	bOverride_MinDesiredWidth = false;
+	if (MySizeBox.IsValid())
+	{
+		MySizeBox->SetMinDesiredWidth(FOptionalSize());
+	}
+}
+
+float USizeBoxComponent::GetMinDesiredHeight() const
+{
+	return MinDesiredHeight;
+}
+
+bool USizeBoxComponent::IsMinDesiredHeightOverride() const
+{
+	return bOverride_MinDesiredHeight;
+}
+
+void USizeBoxComponent::SetMinDesiredHeight(float InMinDesiredHeight)
+{
+	bOverride_MinDesiredHeight = true;
+	MinDesiredHeight = InMinDesiredHeight;
+	if (MySizeBox.IsValid())
+	{
+		MySizeBox->SetMinDesiredHeight(InMinDesiredHeight);
+	}
+}
+
+void USizeBoxComponent::ClearMinDesiredHeight()
+{
+	bOverride_MinDesiredHeight = false;
+	if (MySizeBox.IsValid())
+	{
+		MySizeBox->SetMinDesiredHeight(FOptionalSize());
+	}
+}
+
+float USizeBoxComponent::GetMaxDesiredWidth() const
+{
+	return MaxDesiredWidth;
+}
+
+bool USizeBoxComponent::IsMaxDesiredWidthOverride() const
+{
+	return bOverride_MaxDesiredWidth;
+}
+
+void USizeBoxComponent::SetMaxDesiredWidth(float InMaxDesiredWidth)
+{
+	bOverride_MaxDesiredWidth = true;
+	MaxDesiredWidth = InMaxDesiredWidth;
+	if (MySizeBox.IsValid())
+	{
+		MySizeBox->SetMaxDesiredWidth(InMaxDesiredWidth);
+	}
+}
+
+void USizeBoxComponent::ClearMaxDesiredWidth()
+{
+	bOverride_MaxDesiredWidth = false;
+	if (MySizeBox.IsValid())
+	{
+		MySizeBox->SetMaxDesiredWidth(FOptionalSize());
+	}
+}
+
+float USizeBoxComponent::GetMaxDesiredHeight() const
+{
+	return MaxDesiredHeight;
+}
+
+bool USizeBoxComponent::IsMaxDesiredHeightOverride() const
+{
+	return bOverride_MaxDesiredHeight;
+}
+
+void USizeBoxComponent::SetMaxDesiredHeight(float InMaxDesiredHeight)
+{
+	bOverride_MaxDesiredHeight = true;
+	MaxDesiredHeight = InMaxDesiredHeight;
+	if (MySizeBox.IsValid())
+	{
+		MySizeBox->SetMaxDesiredHeight(InMaxDesiredHeight);
+	}
+}
+
+void USizeBoxComponent::ClearMaxDesiredHeight()
+{
+	bOverride_MaxDesiredHeight = false;
+	if (MySizeBox.IsValid())
+	{
+		MySizeBox->SetMaxDesiredHeight(FOptionalSize());
+	}
+}
+
+float USizeBoxComponent::GetMinAspectRatio() const
+{
+	return MinAspectRatio;
+}
+
+bool USizeBoxComponent::IsMinAspectRatioOverride() const
+{
+	return bOverride_MinAspectRatio;
+}
+
+void USizeBoxComponent::SetMinAspectRatio(float InMinAspectRatio)
+{
+	bOverride_MinAspectRatio = true;
+	MinAspectRatio = InMinAspectRatio;
+	if (MySizeBox.IsValid())
+	{
+		MySizeBox->SetMinAspectRatio(InMinAspectRatio);
+	}
+}
+
+void USizeBoxComponent::ClearMinAspectRatio()
+{
+	bOverride_MinAspectRatio = false;
+	if (MySizeBox.IsValid())
+	{
+		MySizeBox->SetMinAspectRatio(FOptionalSize());
+	}
+}
+
+float USizeBoxComponent::GetMaxAspectRatio() const
+{
+	return MaxAspectRatio;
+}
+
+bool USizeBoxComponent::IsMaxAspectRatioOverride() const
+{
+	return bOverride_MaxAspectRatio;
+}
+
+void USizeBoxComponent::SetMaxAspectRatio(float InMaxAspectRatio)
+{
+	bOverride_MaxAspectRatio = true;
+	MaxAspectRatio = InMaxAspectRatio;
+	if (MySizeBox.IsValid())
+	{
+		MySizeBox->SetMaxAspectRatio(InMaxAspectRatio);
+	}
+}
+
+void USizeBoxComponent::ClearMaxAspectRatio()
+{
+	bOverride_MaxAspectRatio = false;
+	if (MySizeBox.IsValid())
+	{
+		MySizeBox->SetMaxAspectRatio(FOptionalSize());
+	}
+}
+
+FMargin USizeBoxComponent::GetPadding() const
+{
+	return Padding;
+}
+
+void USizeBoxComponent::SetPadding(FMargin InPadding)
+{
+	Padding = InPadding;
+	if (MySizeBox.IsValid())
+	{
+		MySizeBox->SetPadding(InPadding);
+	}
+}
+
+EHorizontalAlignment USizeBoxComponent::GetHorizontalAlignment() const
+{
+	return HorizontalAlignment;
+}
+
+void USizeBoxComponent::SetHorizontalAlignment(EHorizontalAlignment InHorizontalAlignment)
+{
+	HorizontalAlignment = InHorizontalAlignment;
+	if (MySizeBox.IsValid())
+	{
+		MySizeBox->SetHAlign(InHorizontalAlignment);
+	}
+}
+
+EVerticalAlignment USizeBoxComponent::GetVerticalAlignment() const
+{
+	return VerticalAlignment;
+}
+
+void USizeBoxComponent::SetVerticalAlignment(EVerticalAlignment InVerticalAlignment)
+{
+	VerticalAlignment = InVerticalAlignment;
+	if (MySizeBox.IsValid())
+	{
+		MySizeBox->SetVAlign(InVerticalAlignment);
+	}
+}
+
+#if WITH_EDITOR
+void USizeBoxComponent::PostEditChangeProperty(FPropertyChangedEvent& PropertyChangedEvent)
+{
+	Super::PostEditChangeProperty(PropertyChangedEvent);
+
+	SynchronizeProperties();
+}
+#endif //WITH_EDITOR
+
+void USizeBoxComponent::SynchronizeProperties()
+{
+	if (bOverride_WidthOverride)
+	{
+		SetWidthOverride(WidthOverride);
+	}
+	else
+	{
+		ClearWidthOverride();
+	}
+
+	if (bOverride_HeightOverride)
+	{
+		SetHeightOverride(HeightOverride);
+	}
+	else
+	{
+		ClearHeightOverride();
+	}
+
+	if (bOverride_MinDesiredWidth)
+	{
+		SetMinDesiredWidth(MinDesiredWidth);
+	}
+	else
+	{
+		ClearMinDesiredWidth();
+	}
+
+	if (bOverride_MinDesiredHeight)
+	{
+		SetMinDesiredHeight(MinDesiredHeight);
+	}
+	else
+	{
+		ClearMinDesiredHeight();
+	}
+
+	if (bOverride_MaxDesiredWidth)
+	{
+		SetMaxDesiredWidth(MaxDesiredWidth);
+	}
+	else
+	{
+		ClearMaxDesiredWidth();
+	}
+
+	if (bOverride_MaxDesiredHeight)
+	{
+		SetMaxDesiredHeight(MaxDesiredHeight);
+	}
+	else
+	{
+		ClearMaxDesiredHeight();
+	}
+
+	if (bOverride_MinAspectRatio)
+	{
+		SetMinAspectRatio(MinAspectRatio);
+	}
+	else
+	{
+		ClearMinAspectRatio();
+	}
+
+	if (bOverride_MaxAspectRatio)
+	{
+		SetMaxAspectRatio(MaxAspectRatio);
+	}
+	else
+	{
+		ClearMaxAspectRatio();
+	}
+
+	// Set child slot properties
+	SetPadding(Padding);
+	SetHorizontalAlignment(HorizontalAlignment);
+	SetVerticalAlignment(VerticalAlignment);
+}
+
+TSharedRef<SWidget> USizeBoxComponent::RebuildWidgetWithContent(TSharedRef<SWidget> OwnerContent)
+{
+	MySizeBox = SNew(SBox)
+		[OwnerContent];
+
+	SynchronizeProperties();
+
+	return MySizeBox.ToSharedRef();
+}
+
diff --git a/Engine/Source/Runtime/UMG/Public/Components/SizeBoxComponent.h b/Engine/Source/Runtime/UMG/Public/Components/SizeBoxComponent.h
new file mode 100644
--- /dev/null
+++ b/Engine/Source/Runtime/UMG/Public/Components/SizeBoxComponent.h
@@ -0,0 +1,223 @@
+// Copyright Epic Games, Inc. All Rights Reserved.
+
+#pragma once
+
+#include "CoreMinimal.h"
+#include "Extensions/UIComponent.h"
+#include "Layout/Margin.h"
+#include "UObject/ObjectMacros.h"
+
+#include "SizeBoxComponent.generated.h"
+
+class SBox;
+
+/** This is a class for a Component that wraps the Owner widget with a Size Box. */
+UCLASS(MinimalAPI, Experimental)
+class USizeBoxComponent : public UUIComponent
+{
+	GENERATED_BODY()
+
+public:
+
+	/** */
+	UMG_API float GetWidthOverride() const;
+
+	/** */
+	UMG_API bool IsWidthOverride() const;
+
+	/** When specified, ignore the content's desired size and report the WidthOverride as the Box's desired width. */
+	UMG_API void SetWidthOverride(float InWidthOverride);
+
+	UFUNCTION(BlueprintCallable, Category = "Layout|Size Box")
+	UMG_API void ClearWidthOverride();
+
+	/** */
+	UMG_API float GetHeightOverride() const;
+
+	/** */
+	UMG_API bool IsHeightOverride() const;
+
+	/** When specified, ignore the content's desired size and report the HeightOverride as the Box's desired height. */
+	UMG_API void SetHeightOverride(float InHeightOverride);
+
+	UFUNCTION(BlueprintCallable, Category = "Layout|Size Box")
+	UMG_API void ClearHeightOverride();
+
+	/** */
+	UMG_API float GetMinDesiredWidth() const;
+
+	/** */
+	UMG_API bool IsMinDesiredWidthOverride() const;
+
+	/** When specified, will report the MinDesiredWidth if larger than the content's desired width. */
+	UMG_API void SetMinDesiredWidth(float InMinDesiredWidth);
+
+	UFUNCTION(BlueprintCallable, Category = "Layout|Size Box")
+	UMG_API void ClearMinDesiredWidth();
+
+	/** */
+	UMG_API float GetMinDesiredHeight() const;
+
+	/** */
+	UMG_API bool IsMinDesiredHeightOverride() const;
+
+	/** When specified, will report the MinDesiredHeight if larger than the content's desired height. */
+	UMG_API void SetMinDesiredHeight(float InMinDesiredHeight);
+
+	UFUNCTION(BlueprintCallable, Category = "Layout|Size Box")
+	UMG_API void ClearMinDesiredHeight();
+
+	/** */
+	UMG_API float GetMaxDesiredWidth() const;
+
+	/** */
+	UMG_API bool IsMaxDesiredWidthOverride() const;
+
+	/** When specified, will report the MaxDesiredWidth if smaller than the content's desired width. */
+	UMG_API void SetMaxDesiredWidth(float InMaxDesiredWidth);
+
+	UFUNCTION(BlueprintCallable, Category = "Layout|Size Box")
+	UMG_API void ClearMaxDesiredWidth();
+
+	/** */
+	UMG_API float GetMaxDesiredHeight() const;
+
+	/** */
+	UMG_API bool IsMaxDesiredHeightOverride() const;
+
+	/** When specified, will report the MaxDesiredHeight if smaller than the content's desired height. */
+	UMG_API void SetMaxDesiredHeight(float InMaxDesiredHeight);
+
+	UFUNCTION(BlueprintCallable, Category = "Layout|Size Box")
+	UMG_API void ClearMaxDesiredHeight();
+
+	/** */
+	UMG_API float GetMinAspectRatio() const;
+
+	/** */
+	UMG_API bool IsMinAspectRatioOverride() const;
+
+	UMG_API void SetMinAspectRatio(float InMinAspectRatio);
+
+	UFUNCTION(BlueprintCallable, Category = "Layout|Size Box")
+	UMG_API void ClearMinAspectRatio();
+
+	/** */
+	UMG_API float GetMaxAspectRatio() const;
+
+	/** */
+	UMG_API bool IsMaxAspectRatioOverride() const;
+
+	UMG_API void SetMaxAspectRatio(float InMaxAspectRatio);
+
+	UFUNCTION(BlueprintCallable, Category = "Layout|Size Box")
+	UMG_API void ClearMaxAspectRatio();
+
+	// Child slot functions:
+
+	UMG_API FMargin GetPadding() const;
+
+	UMG_API void SetPadding(FMargin InPadding);
+
+	UMG_API EHorizontalAlignment GetHorizontalAlignment() const;
+
+	UMG_API void SetHorizontalAlignment(EHorizontalAlignment InHorizontalAlignment);
+
+	UMG_API EVerticalAlignment GetVerticalAlignment() const;
+
+	UMG_API void SetVerticalAlignment(EVerticalAlignment InVerticalAlignment);
+
+#if WITH_EDITOR
+	// UObject interface
+	virtual void PostEditChangeProperty(struct FPropertyChangedEvent& PropertyChangedEvent) override;
+	// End of UObject interface
+#endif // WITH_EDITOR
+
+	UMG_API void SynchronizeProperties();
+
+protected:
+	TSharedPtr<SBox> MySizeBox;
+
+private:
+	// Child slot properties:
+
+	/** This property is for the Widget that owns this Component. The padding area between the slot and the content it contains. */
+	UPROPERTY(EditAnywhere, BlueprintReadWrite, Getter, Setter, Category = "Size Box Slot", meta = (AllowPrivateAccess = true))
+	FMargin Padding = FMargin(0.f, 0.f);;
+
+	/** This property is for the Widget that owns this Component. The alignment of the object horizontally. */
+	UPROPERTY(EditAnywhere, BlueprintReadWrite, Getter, Setter, Category = "Size Box Slot", meta = (AllowPrivateAccess = true))
+	TEnumAsByte<EHorizontalAlignment> HorizontalAlignment = HAlign_Fill;;
+
+	/** This property is for the Widget that owns this Component. The alignment of the object vertically. */
+	UPROPERTY(EditAnywhere, BlueprintReadWrite, Getter, Setter, Category = "Size Box Slot", meta = (AllowPrivateAccess = true))
+	TEnumAsByte<EVerticalAlignment> VerticalAlignment = VAlign_Fill;
+
+	// Size Box Properties:
+
+	/** When specified, ignore the content's desired size and report the WidthOverride as the Box's desired width. */
+	UPROPERTY(EditAnywhere, BlueprintReadWrite, Getter, Setter, Category = "Child Layout", meta = (editcondition = "bOverride_WidthOverride", AllowPrivateAccess = true))
+	float WidthOverride = 0.f;
+
+	/** When specified, ignore the content's desired size and report the HeightOverride as the Box's desired height. */
+	UPROPERTY(EditAnywhere, BlueprintReadWrite, Getter, Setter, Category = "Child Layout", meta = (editcondition = "bOverride_HeightOverride", AllowPrivateAccess = true))
+	float HeightOverride = 0.f;
+
+	/** When specified, will report the MinDesiredWidth if larger than the content's desired width. */
+	UPROPERTY(EditAnywhere, BlueprintReadWrite, Getter, Setter, Category = "Child Layout", meta = (editcondition = "bOverride_MinDesiredWidth", AllowPrivateAccess = true))
+	float MinDesiredWidth = 0.f;
+
+	/** When specified, will report the MinDesiredHeight if larger than the content's desired height. */
+	UPROPERTY(EditAnywhere, BlueprintReadWrite, Getter, Setter, Category = "Child Layout", meta = (editcondition = "bOverride_MinDesiredHeight", AllowPrivateAccess = true))
+	float MinDesiredHeight = 0.f;
+
+	/** When specified, will report the MaxDesiredWidth if smaller than the content's desired width. */
+	UPROPERTY(EditAnywhere, BlueprintReadWrite, Getter, Setter, Category = "Child Layout", meta = (editcondition = "bOverride_MaxDesiredWidth", AllowPrivateAccess = true))
+	float MaxDesiredWidth = 0.f;
+
+	/** When specified, will report the MaxDesiredHeight if smaller than the content's desired height. */
+	UPROPERTY(EditAnywhere, BlueprintReadWrite, Getter, Setter, Category = "Child Layout", meta = (editcondition = "bOverride_MaxDesiredHeight", AllowPrivateAccess = true))
+	float MaxDesiredHeight = 0.f;
+
+	/** */
+	UPROPERTY(EditAnywhere, BlueprintReadWrite, Getter, Setter, Category = "Child Layout", meta = (editcondition = "bOverride_MinAspectRatio", AllowPrivateAccess = true))
+	float MinAspectRatio = 1.f;
+
+	/** */
+	UPROPERTY(EditAnywhere, BlueprintReadWrite, Getter, Setter, Category = "Child Layout", meta = (editcondition = "bOverride_MaxAspectRatio", AllowPrivateAccess = true))
+	float MaxAspectRatio = 1.f;
+
+	/**  */
+	UPROPERTY(EditAnywhere, Category = "Child Layout", meta = (InlineEditConditionToggle, AllowPrivateAccess = true))
+	uint8 bOverride_WidthOverride : 1;	
+
+	/**  */
+	UPROPERTY(EditAnywhere, Category = "Child Layout", meta = (InlineEditConditionToggle, AllowPrivateAccess = true))
+	uint8 bOverride_HeightOverride : 1;
+
+	/**  */
+	UPROPERTY(EditAnywhere, Category = "Child Layout", meta = (InlineEditConditionToggle, AllowPrivateAccess = true))
+	uint8 bOverride_MinDesiredWidth : 1;
+
+	/**  */
+	UPROPERTY(EditAnywhere, Category = "Child Layout", meta = (InlineEditConditionToggle, AllowPrivateAccess = true))
+	uint8 bOverride_MinDesiredHeight : 1;
+
+	/**  */
+	UPROPERTY(EditAnywhere, Category = "Child Layout", meta = (InlineEditConditionToggle, AllowPrivateAccess = true))
+	uint8 bOverride_MaxDesiredWidth : 1;
+
+	/**  */
+	UPROPERTY(EditAnywhere, Category = "Child Layout", meta = (InlineEditConditionToggle, AllowPrivateAccess = true))
+	uint8 bOverride_MaxDesiredHeight : 1;
+
+	/**  */
+	UPROPERTY(EditAnywhere, Category = "Child Layout", meta = (InlineEditConditionToggle, AllowPrivateAccess = true))
+	uint8 bOverride_MinAspectRatio : 1;
+
+	/**  */
+	UPROPERTY(EditAnywhere, Category = "Child Layout", meta = (InlineEditConditionToggle, AllowPrivateAccess = true))
+	uint8 bOverride_MaxAspectRatio : 1;
+
+	virtual TSharedRef<SWidget> RebuildWidgetWithContent(const TSharedRef<SWidget> OwnerContent) override;
+};
\ No newline at end of file
