Created a Scale Box UI Component which wraps the owner widget in a size box.

#tests Tried out the different Scale Box parameters in the designer view and PIE and everything looks good. Also made functions that were bound to a button that changed parameters at runtime, and it worked as expected.
[RN] minor ui/umg

diff --git a/Engine/Source/Runtime/UMG/Private/Components/ScaleBoxComponent.cpp b/Engine/Source/Runtime/UMG/Private/Components/ScaleBoxComponent.cpp
new file mode 100644
--- /dev/null
+++ b/Engine/Source/Runtime/UMG/Private/Components/ScaleBoxComponent.cpp
@@ -0,0 +1,170 @@
+// Copyright Epic Games, Inc. All Rights Reserved.
+
+#include "Components/ScaleBoxComponent.h"
+
+#include "Blueprint/UserWidget.h"
+#include "Components/Widget.h"
+#include "CoreGlobals.h"
+#include "Widgets/Layout/SScaleBox.h"
+#include "Widgets/SWidget.h"
+
+#include "UMGPrivate.h"
+
+void UScaleBoxComponent::SetStretch(EStretch::Type InStretch)
+{
+	Stretch = InStretch;
+	if (MyScaleBox.IsValid())
+	{
+		MyScaleBox->SetStretch(InStretch);
+	}
+}
+
+EStretch::Type UScaleBoxComponent::GetStretch() const
+{
+	return Stretch;
+}
+
+void UScaleBoxComponent::SetStretchDirection(EStretchDirection::Type InStretchDirection)
+{
+	StretchDirection = InStretchDirection;
+	if (MyScaleBox.IsValid())
+	{
+		MyScaleBox->SetStretchDirection(InStretchDirection);
+	}
+}
+
+EStretchDirection::Type UScaleBoxComponent::GetStretchDirection() const
+{
+	return StretchDirection;
+}
+
+void UScaleBoxComponent::SetUserSpecifiedScale(float InUserSpecifiedScale)
+{
+	UserSpecifiedScale = InUserSpecifiedScale;
+	if (MyScaleBox.IsValid())
+	{
+		MyScaleBox->SetUserSpecifiedScale(InUserSpecifiedScale);
+	}
+}
+
+float UScaleBoxComponent::GetUserSpecifiedScale() const
+{
+	return UserSpecifiedScale;
+}
+
+void UScaleBoxComponent::SetIgnoreInheritedScale(bool bInIgnoreInheritedScale)
+{
+	IgnoreInheritedScale = bInIgnoreInheritedScale;
+	if (MyScaleBox.IsValid())
+	{
+		MyScaleBox->SetIgnoreInheritedScale(bInIgnoreInheritedScale);
+	}
+}
+
+bool UScaleBoxComponent::IsIgnoreInheritedScale() const
+{
+	return IgnoreInheritedScale;
+}
+
+EHorizontalAlignment UScaleBoxComponent::GetHorizontalAlignment() const
+{
+	return HorizontalAlignment;
+}
+
+void UScaleBoxComponent::SetHorizontalAlignment(EHorizontalAlignment InHorizontalAlignment)
+{
+	HorizontalAlignment = InHorizontalAlignment;
+	if (MyScaleBox.IsValid())
+	{
+		MyScaleBox->SetHAlign(InHorizontalAlignment);
+	}
+}
+
+EVerticalAlignment UScaleBoxComponent::GetVerticalAlignment() const
+{
+	return VerticalAlignment;
+}
+
+void UScaleBoxComponent::SetVerticalAlignment(EVerticalAlignment InVerticalAlignment)
+{
+	VerticalAlignment = InVerticalAlignment;
+	if (MyScaleBox.IsValid())
+	{
+		MyScaleBox->SetVAlign(InVerticalAlignment);
+	}
+}
+
+#if WITH_EDITOR
+
+bool UScaleBoxComponent::CanEditChange(const FProperty* InProperty) const
+{
+	bool bIsEditable = Super::CanEditChange(InProperty);
+	if (bIsEditable && InProperty)
+	{
+		const FName PropertyName = InProperty->GetFName();
+		if (PropertyName == GET_MEMBER_NAME_CHECKED(UScaleBoxComponent, StretchDirection))
+		{
+			return Stretch != EStretch::None && Stretch != EStretch::ScaleBySafeZone &&
+				Stretch != EStretch::UserSpecified && Stretch != EStretch::UserSpecifiedWithClipping;
+		}
+
+		if (PropertyName == GET_MEMBER_NAME_CHECKED(UScaleBoxComponent, UserSpecifiedScale))
+		{
+			return Stretch == EStretch::UserSpecified || Stretch == EStretch::UserSpecifiedWithClipping;
+		}
+	}
+
+	return bIsEditable;
+}
+
+// TODO vinz: Implement OnDesignerChanged to work for Components. This currently is just a virtual function for UWidgets. Note: this function was adapted from the UScaleBox code.
+//void UScaleBoxComponent::OnDesignerChanged(const FDesignerChangedEventArgs& EventArgs)
+//{
+//	if (EventArgs.bScreenPreview)
+//	{
+//		DesignerSize = EventArgs.Size;
+//	}
+//	else
+//	{
+//		DesignerSize = FVector2D(0, 0);
+//	}
+//
+//	if (MyScaleBox.IsValid())
+//	{
+//		MyScaleBox->SetOverrideScreenInformation(DesignerSize);
+//	}
+//}
+
+void UScaleBoxComponent::PostEditChangeProperty(FPropertyChangedEvent& PropertyChangedEvent)
+{
+	Super::PostEditChangeProperty(PropertyChangedEvent);
+
+	SynchronizeProperties();
+}
+#endif //WITH_EDITOR
+
+void UScaleBoxComponent::SynchronizeProperties()
+{
+	SetStretchDirection(StretchDirection);
+	SetStretch(Stretch);
+	SetUserSpecifiedScale(UserSpecifiedScale);
+	SetIgnoreInheritedScale(IgnoreInheritedScale);
+
+	// Set child slot properties
+	SetHorizontalAlignment(HorizontalAlignment);
+	SetVerticalAlignment(VerticalAlignment);
+}
+
+TSharedRef<SWidget> UScaleBoxComponent::RebuildWidgetWithContent(TSharedRef<SWidget> OwnerContent)
+{
+	MyScaleBox = SNew(SScaleBox)
+// TODO vinz: Implement OnDesignerChanged to work for Components. This currently is just a virtual function for UWidgets. Also, this probably should be a part of SynchronizeProperties.
+//#if WITH_EDITOR
+//		.OverrideScreenSize(DesignerSize)
+//#endif
+		[OwnerContent];
+
+	SynchronizeProperties();
+
+	return MyScaleBox.ToSharedRef();
+}
\ No newline at end of file
diff --git a/Engine/Source/Runtime/UMG/Public/Components/ScaleBoxComponent.h b/Engine/Source/Runtime/UMG/Public/Components/ScaleBoxComponent.h
new file mode 100644
--- /dev/null
+++ b/Engine/Source/Runtime/UMG/Public/Components/ScaleBoxComponent.h
@@ -0,0 +1,100 @@
+// Copyright Epic Games, Inc. All Rights Reserved.
+
+#pragma once
+
+#include "Components/ContentWidget.h"
+#include "CoreMinimal.h"
+#include "Extensions/UIComponent.h"
+#include "UObject/ObjectMacros.h"
+#include "Widgets/Layout/SScaleBox.h"
+#include "Widgets/SWidget.h"
+
+#include "ScaleBoxComponent.generated.h"
+
+class SScaleBox;
+
+/** This is a class for a Component that wraps the Owner widget with a Scale Box */
+UCLASS(MinimalAPI, Experimental)
+class UScaleBoxComponent : public UUIComponent
+{
+	GENERATED_BODY()
+
+public:
+	UMG_API void SetStretch(EStretch::Type InStretch);
+
+	UMG_API EStretch::Type GetStretch() const;
+
+	UMG_API void SetStretchDirection(EStretchDirection::Type InStretchDirection);
+
+	UMG_API EStretchDirection::Type GetStretchDirection() const;
+
+	UMG_API void SetUserSpecifiedScale(float InUserSpecifiedScale);
+
+	UMG_API float GetUserSpecifiedScale() const;
+
+	UMG_API void SetIgnoreInheritedScale(bool bInIgnoreInheritedScale);
+
+	UMG_API bool IsIgnoreInheritedScale() const;
+
+	// Child slot functions:
+
+	UMG_API EHorizontalAlignment GetHorizontalAlignment() const;
+
+	UMG_API void SetHorizontalAlignment(EHorizontalAlignment InHorizontalAlignment);
+
+	UMG_API EVerticalAlignment GetVerticalAlignment() const;
+
+	UMG_API void SetVerticalAlignment(EVerticalAlignment InVerticalAlignment);
+
+#if WITH_EDITOR
+	UMG_API virtual bool CanEditChange(const FProperty* InProperty) const override;
+	// TODO vinz: Implement OnDesignerChanged to work for Components. This currently is just a virtual function for UWidgets.
+	//UMG_API virtual void OnDesignerChanged(const FDesignerChangedEventArgs& EventArgs) override;
+
+// UObject interface
+	virtual void PostEditChangeProperty(struct FPropertyChangedEvent& PropertyChangedEvent) override;
+	// End of UObject interface
+#endif // WITH_EDITOR
+
+	UMG_API void SynchronizeProperties();
+
+protected:
+	TSharedPtr<SScaleBox> MyScaleBox;
+
+// TODO vinz: Implement OnDesignerChanged to work for Components. This currently is just a virtual function for UWidgets.
+//#if WITH_EDITOR
+//	TOptional<FVector2D> DesignerSize;
+//#endif
+
+private:
+
+	// Child slot properties
+
+	/** This property is for the Widget that owns this Component. The alignment of the object horizontally. */
+	UPROPERTY(EditAnywhere, BlueprintReadWrite, Getter, Setter, Category = "Scale Box Slot", meta = (AllowPrivateAccess = true))
+	TEnumAsByte<EHorizontalAlignment> HorizontalAlignment = HAlign_Center;
+
+	/** This property is for the Widget that owns this Component. The alignment of the object vertically. */
+	UPROPERTY(EditAnywhere, BlueprintReadWrite, Getter, Setter, Category = "Scale Box Slot", meta = (AllowPrivateAccess = true))
+	TEnumAsByte<EVerticalAlignment> VerticalAlignment = VAlign_Center;
+
+	// Scale Box properties:
+
+	/** The stretching rule to apply when content is stretched */
+	UPROPERTY(EditAnywhere, BlueprintReadWrite, Getter, Setter, Category = "Stretching", meta = (AllowPrivateAccess = true))
+	TEnumAsByte<EStretch::Type> Stretch = EStretch::ScaleToFit;
+
+	/** Controls in what direction content can be scaled */
+	UPROPERTY(EditAnywhere, BlueprintReadWrite, Getter, Setter, Category = "Stretching", meta = (AllowPrivateAccess = true))
+	TEnumAsByte<EStretchDirection::Type> StretchDirection = EStretchDirection::Both;
+
+	/** Optional scale that can be specified by the User. Used only for UserSpecified stretching. */
+	UPROPERTY(EditAnywhere, BlueprintReadWrite, Getter, Setter, Category = "Stretching", meta = (AllowPrivateAccess = true))
+	float UserSpecifiedScale = 1.0f;
+
+	/** Optional bool to ignore the inherited scale. Applies inverse scaling to counteract parents before applying the local scale operation. */
+	UPROPERTY(EditAnywhere, BlueprintReadWrite, Getter="IsIgnoreInheritedScale", Setter, Category = "Stretching", meta = (AllowPrivateAccess = true))
+	bool IgnoreInheritedScale = false;
+
+	virtual TSharedRef<SWidget> RebuildWidgetWithContent(const TSharedRef<SWidget> OwnerContent) override;
+};
\ No newline at end of file
