preventing 3D widget keyboard focus leak upon PIE shutdown

diff --git a/Engine/Source/Editor/UnrealEd/Private/PlayLevel.cpp b/Engine/Source/Editor/UnrealEd/Private/PlayLevel.cpp
--- a/Engine/Source/Editor/UnrealEd/Private/PlayLevel.cpp
+++ b/Engine/Source/Editor/UnrealEd/Private/PlayLevel.cpp
@@ -301,8 +301,9 @@ void UEditorEngine::EndPlayMap()
 	GetSelectedObjects()->DeselectAll();
 	GetSelectedComponents()->DeselectAll();
 
-	// let the editor know
+	// let the editor and the world know
 	FEditorDelegates::EndPIE.Broadcast(bIsSimulatingInEditor);
+	FWorldDelegates::OnPIEEnded.Broadcast(PlayWorld->GetGameInstance());
 
 	// clean up any previous Play From Here sessions
 	if ( GameViewport != NULL && GameViewport->Viewport != NULL )
diff --git a/Engine/Source/Runtime/UMG/Private/Components/WidgetInteractionComponent.cpp b/Engine/Source/Runtime/UMG/Private/Components/WidgetInteractionComponent.cpp
--- a/Engine/Source/Runtime/UMG/Private/Components/WidgetInteractionComponent.cpp
+++ b/Engine/Source/Runtime/UMG/Private/Components/WidgetInteractionComponent.cpp
@@ -46,6 +46,10 @@ UWidgetInteractionComponent::UWidgetInteractionComponent(const FObjectInitialize
 		ArrowComponent->AttachToComponent(this, FAttachmentTransformRules(EAttachmentRule::KeepRelative, false));
 	}
 #endif
+	
+#if WITH_EDITOR
+	FWorldDelegates::OnPIEEnded.AddUObject(this, &UWidgetInteractionComponent::OnPIEEnded);
+#endif
 }
 
 void UWidgetInteractionComponent::OnComponentCreated()
@@ -114,6 +118,22 @@ void UWidgetInteractionComponent::SetFocus(UWidget* FocusWidget)
 	}
 }
 
+#if WITH_EDITOR
+void UWidgetInteractionComponent::OnPIEEnded(UGameInstance*)
+{
+	if(LastWidgetPath.IsValid())
+	{
+		TSharedPtr<SWidget> LastVirtualUserWidget = LastWidgetPath.GetLastWidget().Pin();
+		int32 UserIndex = VirtualUser->GetUserIndex();
+		
+		if (VirtualUser.IsValid() && LastVirtualUserWidget != nullptr && LastVirtualUserWidget->HasUserFocus(UserIndex))
+		{
+			FSlateApplication::Get().ClearUserFocus(UserIndex, EFocusCause::SetDirectly);
+		}
+	}
+}
+#endif
+
 FWidgetPath UWidgetInteractionComponent::FindHoveredWidgetPath(const FWidgetTraceResult& TraceResult) const
 {
 	if (TraceResult.HitWidgetComponent)
diff --git a/Engine/Source/Runtime/UMG/Public/Components/WidgetInteractionComponent.h b/Engine/Source/Runtime/UMG/Public/Components/WidgetInteractionComponent.h
--- a/Engine/Source/Runtime/UMG/Public/Components/WidgetInteractionComponent.h
+++ b/Engine/Source/Runtime/UMG/Public/Components/WidgetInteractionComponent.h
@@ -362,6 +362,10 @@ protected:
 
 private:
 
+#if WITH_EDITOR
+	void OnPIEEnded(class UGameInstance*);
+#endif
+	
 	/** Returns the path to the widget that is currently beneath the pointer */
 	FWidgetPath DetermineWidgetUnderPointer();
 
