Moving VirtualUser isValid check to fix BugHawk crash issue for CL #44935615

diff --git a/Engine/Source/Runtime/UMG/Private/Components/WidgetInteractionComponent.cpp b/Engine/Source/Runtime/UMG/Private/Components/WidgetInteractionComponent.cpp
--- a/Engine/Source/Runtime/UMG/Private/Components/WidgetInteractionComponent.cpp
+++ b/Engine/Source/Runtime/UMG/Private/Components/WidgetInteractionComponent.cpp
@@ -121,12 +121,12 @@ void UWidgetInteractionComponent::SetFocus(UWidget* FocusWidget)
 #if WITH_EDITOR
 void UWidgetInteractionComponent::OnPIEEnded(UGameInstance*)
 {
-	if(LastWidgetPath.IsValid())
+	if(LastWidgetPath.IsValid() && VirtualUser.IsValid())
 	{
 		TSharedPtr<SWidget> LastVirtualUserWidget = LastWidgetPath.GetLastWidget().Pin();
 		int32 UserIndex = VirtualUser->GetUserIndex();
 		
-		if (VirtualUser.IsValid() && LastVirtualUserWidget != nullptr && LastVirtualUserWidget->HasUserFocus(UserIndex))
+		if (LastVirtualUserWidget != nullptr && LastVirtualUserWidget->HasUserFocus(UserIndex))
 		{
 			FSlateApplication::Get().ClearUserFocus(UserIndex, EFocusCause::SetDirectly);
 		}
