UMG: Forward focused UWidget to Navigation UI Component. Add utility BP function to get owning UserWidget

#rnx

diff --git a/Engine/Source/Runtime/UMG/Private/Extensions/UIComponents/NavigationUIComponent.cpp b/Engine/Source/Runtime/UMG/Private/Extensions/UIComponents/NavigationUIComponent.cpp
--- a/Engine/Source/Runtime/UMG/Private/Extensions/UIComponents/NavigationUIComponent.cpp
+++ b/Engine/Source/Runtime/UMG/Private/Extensions/UIComponents/NavigationUIComponent.cpp
@@ -8,46 +8,57 @@
 #include "Slate/SObjectWidget.h"
 #include "Types/ReflectionMetadata.h"
 
-namespace UE::NavigationUIComponent::Private
+namespace UIComponent::Private
 {
-	UUserWidget* FindUserWidget(TWeakPtr<SWidget> Widget)
+UUserWidget* FindNearestUserWidget(TWeakPtr<SWidget> Widget)
+{
+	TSharedPtr<SWidget> PinnedWidget = Widget.Pin();
+	while (PinnedWidget.IsValid())
 	{
-		// Find the outer user widget that contains this widget. 
-		while(Widget.IsValid())
+		if (PinnedWidget->GetWidgetClass().GetWidgetType() == SObjectWidget::StaticWidgetClass().GetWidgetType())
 		{
-			TSharedPtr<SWidget> PinnedWidget = Widget.Pin();
-			if (PinnedWidget->GetWidgetClass().GetWidgetType() == SObjectWidget::StaticWidgetClass().GetWidgetType())
-			{
-				TSharedPtr<SObjectWidget> ObjectWidget = StaticCastSharedPtr<SObjectWidget>(PinnedWidget);
-				return ObjectWidget->GetWidgetObject();
-			}
-		
-			Widget = PinnedWidget->GetParentWidget();
+			return StaticCastSharedPtr<SObjectWidget>(PinnedWidget)->GetWidgetObject();
 		}
-		return nullptr;
+		PinnedWidget = PinnedWidget->GetParentWidget();
 	}
+	return nullptr;
+}
 
-	TSharedPtr<SWidget> GetNavigationWidget(TSharedPtr<SWidget> Widget)
+UWidget* FindUMGWidget(TWeakPtr<SWidget> Widget)
+{
+	UUserWidget* OwningUserWidget = FindNearestUserWidget(Widget);
+	if (OwningUserWidget == nullptr || OwningUserWidget->WidgetTree == nullptr)
 	{
-		// Special condition, get internal SWidget that will be focused/navigated. 
-		// Otherwise, the outer SObject Widget will report that it is focusable, although the native widget is being focus
-		if (Widget->GetType() == FName(TEXT("SObjectWidget")))
+		return nullptr;
+	}
+
+	UWidget* NearestUMGWidget = nullptr;
+	int32 NearestUMGWidgetDistance = INT32_MAX;
+
+	OwningUserWidget->WidgetTree->ForEachWidgetUntil([Widget, &NearestUMGWidget](UWidget* UMGWidget) {
+		if (NearestUMGWidget == nullptr)
+		{
+			NearestUMGWidget = UMGWidget;
+		}
+		else
 		{
-			TSharedPtr<SObjectWidget> ObjectWidget = StaticCastSharedPtr<SObjectWidget>(Widget);
-			if (UUserWidget* UserWidget = ObjectWidget->GetWidgetObject())
+			TSharedPtr<SWidget> PinnedWidget = Widget.Pin();
+			while (PinnedWidget.IsValid())
 			{
-				Widget = UserWidget->GetCachedWidget().ToSharedRef();
-
-				TObjectPtr<UWidgetTree> WidgetTree = (UserWidget != nullptr) ? UserWidget->WidgetTree : nullptr;
-				if(TObjectPtr<UWidget> RootWidget = (WidgetTree != nullptr) ? WidgetTree->RootWidget : nullptr)
+				if (UMGWidget->GetCachedWidget() == PinnedWidget)
 				{
-					Widget = RootWidget->GetCachedWidget().ToSharedRef();
+					NearestUMGWidget = UMGWidget;
+					return false;
 				}
+				PinnedWidget = PinnedWidget->GetParentWidget();
 			}
 		}
+		// Stop searching when slate widget is found
+		return UMGWidget->GetCachedWidget() != Widget;
+	});
 
-		return Widget;
-	}
+	return NearestUMGWidget;
+}
 }
 
 void UNavigationUIComponent::OnConstruct()
@@ -78,10 +89,7 @@ void UNavigationUIComponent::OnConstruct()
 	NavigationTransitionMetadata = MakeShared<FNavigationTransitionMetadata>();
 	NavigationTransitionMetadata->OnNavigationTransition.BindUObject(this, &UNavigationUIComponent::HandleNavigationTransition);
 
-	if (TSharedPtr<SWidget> NavigationWidget = UE::NavigationUIComponent::Private::GetNavigationWidget(Content))
-	{
-		NavigationWidget->AddMetadata(NavigationTransitionMetadata.ToSharedRef());
-	}
+	Content->AddMetadata(NavigationTransitionMetadata.ToSharedRef());
 }
 
 void UNavigationUIComponent::OnDestruct()
@@ -93,20 +101,20 @@ void UNavigationUIComponent::OnDestruct()
 void UNavigationUIComponent::HandleNavigationTransition(const FNavigationTransition& NavigationTransition)
 {
 	// Forward slate navigation transition to user provided functions via component delegates
-
-	UUserWidget* OldFocusedWidgetObject = UE::NavigationUIComponent::Private::FindUserWidget(NavigationTransition.OldFocusedWidget);
-	UUserWidget* NewFocusedWidgetObject = UE::NavigationUIComponent::Private::FindUserWidget(NavigationTransition.NewFocusedWidget);
+	
+	UWidget* OldUMGWidget = UIComponent::Private::FindUMGWidget(NavigationTransition.OldFocusedWidget);
+	UWidget* NewUMGWidget = UIComponent::Private::FindUMGWidget(NavigationTransition.NewFocusedWidget);
 
 	switch (NavigationTransition.Direction)
 	{
 		case ENavigationTransitionDirection::Incoming:
 		{
-			OnNavigationEnteredDelegate.ExecuteIfBound(NavigationTransition.Type, OldFocusedWidgetObject, NewFocusedWidgetObject);
+			OnNavigationEnteredDelegate.ExecuteIfBound(NavigationTransition.Type, OldUMGWidget, NewUMGWidget);
 			break;
 		}
 		case ENavigationTransitionDirection::Outgoing:
 		{
-			OnNavigationExitedDelegate.ExecuteIfBound(NavigationTransition.Type, OldFocusedWidgetObject, NewFocusedWidgetObject);
+			OnNavigationExitedDelegate.ExecuteIfBound(NavigationTransition.Type, OldUMGWidget, NewUMGWidget);
 			break;
 		}
 	}
diff --git a/Engine/Source/Runtime/UMG/Private/Slate/SObjectWidget.cpp b/Engine/Source/Runtime/UMG/Private/Slate/SObjectWidget.cpp
--- a/Engine/Source/Runtime/UMG/Private/Slate/SObjectWidget.cpp
+++ b/Engine/Source/Runtime/UMG/Private/Slate/SObjectWidget.cpp
@@ -9,6 +9,11 @@
 
 CSV_DECLARE_CATEGORY_MODULE_EXTERN(SLATECORE_API, Slate);
 
+SLATE_IMPLEMENT_WIDGET(SObjectWidget)
+void SObjectWidget::PrivateRegisterAttributes(FSlateAttributeInitializer& AttributeInitializer)
+{
+}
+
 void SObjectWidget::Construct(const FArguments& InArgs, UUserWidget* InWidgetObject)
 {
 	WidgetObject = InWidgetObject;
diff --git a/Engine/Source/Runtime/UMG/Private/UserWidget.cpp b/Engine/Source/Runtime/UMG/Private/UserWidget.cpp
--- a/Engine/Source/Runtime/UMG/Private/UserWidget.cpp
+++ b/Engine/Source/Runtime/UMG/Private/UserWidget.cpp
@@ -79,6 +79,11 @@ UUMGSequencePlayer* UUserWidgetFunctionLibrary::Conv_UMGSequencePlayer(const FWi
 	return WidgetAnimationHandle.GetSequencePlayer();
 }
 
+UUserWidget* UUserWidgetFunctionLibrary::GetOuterUserWidget(const UWidget* Widget)
+{
+	return Widget != nullptr? Widget->GetTypedOuter<UUserWidget>() : nullptr;
+}
+
 /////////////////////////////////////////////////////
 // UUserWidget
 UUserWidget::UUserWidget(const FObjectInitializer& ObjectInitializer)
diff --git a/Engine/Source/Runtime/UMG/Private/WidgetTree.cpp b/Engine/Source/Runtime/UMG/Private/WidgetTree.cpp
--- a/Engine/Source/Runtime/UMG/Private/WidgetTree.cpp
+++ b/Engine/Source/Runtime/UMG/Private/WidgetTree.cpp
@@ -265,6 +265,87 @@ void UWidgetTree::ForWidgetAndChildren(UWidget* Widget, TFunctionRef<void(UWidge
 	}
 }
 
+void UWidgetTree::ForEachWidgetUntil(TFunctionRef<bool(UWidget*)> Predicate) const
+{
+	// Mimics UWidgetTree::ForWidget
+	if ( RootWidget )
+	{
+		if (!Predicate(RootWidget))
+		{
+			return;
+		}
+
+		if (!ForEachWidgetAndChildrenUntil(RootWidget, Predicate))
+		{
+			return;
+		}
+	}
+
+	for (const auto& NamedSlotsEntry : NamedSlotBindings)
+	{
+		if (UWidget* NamedSlotContent = NamedSlotsEntry.Value)
+		{
+			if(!Predicate(NamedSlotContent))
+			{
+				return;
+			}
+			
+			if(!ForEachWidgetAndChildrenUntil(NamedSlotContent, Predicate))
+			{
+				return;
+			}
+		}
+	}
+}
+
+bool UWidgetTree::ForEachWidgetAndChildrenUntil(UWidget* Widget, TFunctionRef<bool(UWidget*)> Predicate)
+{
+	// Mimics UWidgetTree::ForWidgetAndChildren
+	if ( INamedSlotInterface* NamedSlotHost = Cast<INamedSlotInterface>(Widget) )
+	{
+		TArray<FName> SlotNames;
+		NamedSlotHost->GetSlotNames(SlotNames);
+
+		for ( FName SlotName : SlotNames )
+		{
+			if ( UWidget* SlotContent = NamedSlotHost->GetContentForSlot(SlotName) )
+			{
+				if (!Predicate(SlotContent))
+				{
+					return false;
+				}
+
+				if (!ForEachWidgetAndChildrenUntil(SlotContent, Predicate))
+				{
+					return false;
+				}
+			}
+		}
+	}
+
+	// Search standard children.
+	if ( UPanelWidget* PanelParent = Cast<UPanelWidget>(Widget) )
+	{
+		for ( int32 ChildIndex = 0; ChildIndex < PanelParent->GetChildrenCount(); ChildIndex++ )
+		{
+			if ( UWidget* ChildWidget = PanelParent->GetChildAt(ChildIndex) )
+			{
+				if(!Predicate(ChildWidget))
+				{
+					return false;
+				}
+
+				if (!ForEachWidgetAndChildrenUntil(ChildWidget, Predicate))
+				{
+					return false;
+				}
+			}
+		}
+	}
+
+	return true;
+}
+
 // INamedSlotInterface
 //----------------------------------------------------------------------------------------
 
diff --git a/Engine/Source/Runtime/UMG/Public/Blueprint/UserWidget.h b/Engine/Source/Runtime/UMG/Public/Blueprint/UserWidget.h
--- a/Engine/Source/Runtime/UMG/Public/Blueprint/UserWidget.h
+++ b/Engine/Source/Runtime/UMG/Public/Blueprint/UserWidget.h
@@ -43,6 +43,7 @@ class UWidgetAnimation;
 class UWidgetBlueprintGeneratedClass;
 class UWidgetTree;
 class UNamedSlot;
+class UUserWidget;
 class UUserWidgetExtension;
 
 /** Describes overall action driving this animation transition. */
@@ -264,6 +265,9 @@ class UUserWidgetFunctionLibrary : public UBlueprintFunctionLibrary
 public:
 	UFUNCTION(BlueprintPure, Category = "UMG", meta = (BlueprintAutocast, DeprecatedFunction, DeprecatedMessage="UMGSequencePlayer is deprecated, please use FWidgetAnimationState"))
 	static UUMGSequencePlayer* Conv_UMGSequencePlayer(const FWidgetAnimationHandle& WidgetAnimationHandle);
+
+	UFUNCTION(BlueprintPure, Category = "UMG")
+	static UUserWidget* GetOuterUserWidget(const UWidget* Widget);
 };
 
 DECLARE_DYNAMIC_MULTICAST_DELEGATE(FOnConstructEvent);
diff --git a/Engine/Source/Runtime/UMG/Public/Blueprint/WidgetTree.h b/Engine/Source/Runtime/UMG/Public/Blueprint/WidgetTree.h
--- a/Engine/Source/Runtime/UMG/Public/Blueprint/WidgetTree.h
+++ b/Engine/Source/Runtime/UMG/Public/Blueprint/WidgetTree.h
@@ -86,6 +86,17 @@ public:
 	 */
 	static UMG_API void ForWidgetAndChildren(UWidget* Widget, TFunctionRef<void(UWidget*)> Predicate);
 
+	/**
+	 * Iterates widgets according to ForEachWidget until Predicate returns false.
+	 */
+	void ForEachWidgetUntil(TFunctionRef<bool(UWidget*)> Predicate) const;
+
+	/**
+	 * Iterates widgets according to ForWidgetAndChildren until Predicate returns false.
+	 */
+	static UMG_API bool ForEachWidgetAndChildrenUntil(UWidget* Widget, TFunctionRef<bool(UWidget*)> Predicate);
+
+
 	/** Constructs the widget, and adds it to the tree. */
 	template <typename WidgetT>
 	[[nodiscard]] inline WidgetT* ConstructWidget(TSubclassOf<WidgetT> WidgetClass = WidgetT::StaticClass(), FName WidgetName = NAME_None)
diff --git a/Engine/Source/Runtime/UMG/Public/Extensions/UIComponents/NavigationUIComponent.h b/Engine/Source/Runtime/UMG/Public/Extensions/UIComponents/NavigationUIComponent.h
--- a/Engine/Source/Runtime/UMG/Public/Extensions/UIComponents/NavigationUIComponent.h
+++ b/Engine/Source/Runtime/UMG/Public/Extensions/UIComponents/NavigationUIComponent.h
@@ -9,13 +9,13 @@
 
 #include "NavigationUIComponent.generated.h"
 
-class UUserWidget;
+class UWidget;
 
 /**
 * If Entered, Other Widget is the Widget Navigated From 
 * If Exited , Other Widget is the Widget Navigated To
 */
-DECLARE_DYNAMIC_DELEGATE_ThreeParams(FOnNavigationTransition, EUINavigation, Type, UUserWidget*, OldFocusedWidget, UUserWidget*, NewFocusedWidget);
+DECLARE_DYNAMIC_DELEGATE_ThreeParams(FOnNavigationTransition, EUINavigation, Type, UWidget*, OldFocusedWidget, UWidget*, NewFocusedWidget);
 
 UCLASS(BlueprintType, Experimental)
 class UNavigationUIComponent : public UUIComponent
diff --git a/Engine/Source/Runtime/UMG/Public/Slate/SObjectWidget.h b/Engine/Source/Runtime/UMG/Public/Slate/SObjectWidget.h
--- a/Engine/Source/Runtime/UMG/Public/Slate/SObjectWidget.h
+++ b/Engine/Source/Runtime/UMG/Public/Slate/SObjectWidget.h
@@ -29,6 +29,7 @@ template< typename ObjectType > class TAttribute;
  */
 class SObjectWidget : public SCompoundWidget, public FGCObject
 {
+	SLATE_DECLARE_WIDGET_API(SObjectWidget, SCompoundWidget, UMG_API)
 	SLATE_BEGIN_ARGS(SObjectWidget)
 	{
 		_Visibility = EVisibility::SelfHitTestInvisible;
