UMG: Fix named slot content not saving for named slot instances of widget blueprints that extend widget blueprints

#tests Verified that named slots on widgets that inherit from other widget blueprints now properly save their content. Also verified that named slots from base widget blueprints in instances of child widget blueprints still save their content.

diff --git a/Engine/Source/Runtime/UMG/Private/UserWidget.cpp b/Engine/Source/Runtime/UMG/Private/UserWidget.cpp
--- a/Engine/Source/Runtime/UMG/Private/UserWidget.cpp
+++ b/Engine/Source/Runtime/UMG/Private/UserWidget.cpp
@@ -2627,7 +2627,7 @@ void UUserWidget::PreSave(FObjectPreSaveContext ObjectSaveContext)
 	}
 
 	// Remove bindings that are no longer contained in the class.
-	if ( UWidgetBlueprintGeneratedClass* BGClass = GetWidgetTreeOwningClass())
+	if ( UWidgetBlueprintGeneratedClass* BGClass = Cast<UWidgetBlueprintGeneratedClass>(GetClass()))
 	{
 		RemoveObsoleteBindings(BGClass->NamedSlots);
 	}
