Fixed issue where widgets couldn't be dragged to the UMG designer's DesignerView

diff --git a/Engine/Source/Runtime/UMG/Private/Components/Widget.cpp b/Engine/Source/Runtime/UMG/Private/Components/Widget.cpp
--- a/Engine/Source/Runtime/UMG/Private/Components/Widget.cpp
+++ b/Engine/Source/Runtime/UMG/Private/Components/Widget.cpp
@@ -1460,7 +1460,7 @@ void UWidget::SynchronizeProperties()
 	// We want to apply the bindings to the cached widget, which could be the SWidget, or the SObjectWidget, 
 	// in the case where it's a user widget.  We always want to prefer the SObjectWidget so that bindings to 
 	// visibility and enabled status are not stomping values setup in the root widget in the User Widget.
-	TSharedPtr<SWidget> SafeWidget = GetCachedWidget();
+	TSharedPtr<SWidget> SafeWidget = GetCachedWrappedWidget();
 	if ( !SafeWidget.IsValid() )
 	{
 		return;
diff --git a/Engine/Source/Runtime/UMG/Private/WidgetTree.cpp b/Engine/Source/Runtime/UMG/Private/WidgetTree.cpp
--- a/Engine/Source/Runtime/UMG/Private/WidgetTree.cpp
+++ b/Engine/Source/Runtime/UMG/Private/WidgetTree.cpp
@@ -48,7 +48,7 @@ UWidget* UWidgetTree::FindWidget(TSharedRef<SWidget> InWidget) const
 	UWidget* FoundWidget = nullptr;
 
 	ForEachWidget([&] (UWidget* Widget) {
-		if ( Widget->GetCachedWidget() == InWidget )
+		if ( Widget->GetCachedWrappedWidget() == InWidget )
 		{
 			FoundWidget = Widget;
 		}
