UMG: UIComponents - Fix UMG widget search in Navigation Component

#rnx

diff --git a/Engine/Source/Runtime/UMG/Private/Extensions/UIComponents/NavigationUIComponent.cpp b/Engine/Source/Runtime/UMG/Private/Extensions/UIComponents/NavigationUIComponent.cpp
--- a/Engine/Source/Runtime/UMG/Private/Extensions/UIComponents/NavigationUIComponent.cpp
+++ b/Engine/Source/Runtime/UMG/Private/Extensions/UIComponents/NavigationUIComponent.cpp
@@ -35,21 +35,28 @@ UWidget* FindUMGWidget(TWeakPtr<SWidget> Widget)
 	UWidget* NearestUMGWidget = nullptr;
 	int32 NearestUMGWidgetDistance = INT32_MAX;
 
-	OwningUserWidget->WidgetTree->ForEachWidgetUntil([Widget, &NearestUMGWidget](UWidget* UMGWidget) {
+	OwningUserWidget->WidgetTree->ForEachWidgetUntil([Widget, &NearestUMGWidget, &NearestUMGWidgetDistance](UWidget* UMGWidget) {
 		if (NearestUMGWidget == nullptr)
 		{
 			NearestUMGWidget = UMGWidget;
 		}
 		else
 		{
+			int32 UMGWidgetDistance = 0;
+
 			TSharedPtr<SWidget> PinnedWidget = Widget.Pin();
 			while (PinnedWidget.IsValid())
 			{
-				if (UMGWidget->GetCachedWidget() == PinnedWidget)
+				if (UMGWidget->GetCachedWidget() == PinnedWidget && UMGWidgetDistance < NearestUMGWidgetDistance)
 				{
 					NearestUMGWidget = UMGWidget;
-					return false;
+					NearestUMGWidgetDistance = UMGWidgetDistance;
+					
+					//break widget search if can not find closer widget
+					return NearestUMGWidgetDistance != 0; 
 				}
+
+				++UMGWidgetDistance;
 				PinnedWidget = PinnedWidget->GetParentWidget();
 			}
 		}
