Fixing 3D Web Browser widget memory leakage after backing actor destruction

diff --git a/Engine/Source/Runtime/UMG/Private/Components/WidgetComponent.cpp b/Engine/Source/Runtime/UMG/Private/Components/WidgetComponent.cpp
--- a/Engine/Source/Runtime/UMG/Private/Components/WidgetComponent.cpp
+++ b/Engine/Source/Runtime/UMG/Private/Components/WidgetComponent.cpp
@@ -19,6 +19,7 @@
 #include "Widgets/SWindow.h"
 #include "Engine/TextureRenderTarget2D.h"
 #include "Framework/Application/SlateApplication.h"
+#include "Framework/Application/SlateUser.h"
 #include "Kismet/KismetSystemLibrary.h"
 #include "Input/HittestGrid.h"
 #include "SceneManagement.h"
@@ -1511,11 +1512,47 @@ double UWidgetComponent::GetCurrentTime() const
 	return (TimingPolicy == EWidgetTimingPolicy::RealTime) ? FApp::GetCurrentTime() : static_cast<double>(GetWorld()->GetTimeSeconds());
 }
 
+void UWidgetComponent::OnWorldUserEntered(int32 SlateUserIndex)
+{
+	UsersEnteredWorldWidget.Add(SlateUserIndex);
+}
+
+void UWidgetComponent::ResetWorldUserFocusOnWidget()
+{
+	// we want to go through all the users which could've focused
+	// on our widget to release any focus paths/strongs refs they might have to our widget.
+	for(int32 UserIndex : UsersEnteredWorldWidget)
+	{
+		TSharedPtr<SWidget> CurrentWidgetInPath = FSlateApplication::Get().GetUser(UserIndex)->GetFocusedWidget();
+		while(CurrentWidgetInPath != nullptr)
+		{
+			if(CurrentWidgetInPath == SlateWindow)
+			{
+				FSlateApplication::Get().ClearUserFocus(UserIndex, EFocusCause::SetDirectly);
+
+				break;
+			}
+			
+			CurrentWidgetInPath = CurrentWidgetInPath->GetParentWidget();
+		}
+	}
+	UsersEnteredWorldWidget.Reset();
+}
+
 void UWidgetComponent::RemoveWidgetFromScreen()
 {
 #if !UE_SERVER
 	if (!IsRunningDedicatedServer())
 	{
+		ResetWorldUserFocusOnWidget();
+		
+		if(GetWidgetSpace() == EWidgetSpace::World)
+		{
+			SlateWindow->SetContent(SNullWidget::NullWidget);
+		}
+		
+		SlateWindow.Reset();
+		
 		bAddedToScreen = false;
 
 		if (UWorld* World = GetWorld())
diff --git a/Engine/Source/Runtime/UMG/Private/Components/WidgetInteractionComponent.cpp b/Engine/Source/Runtime/UMG/Private/Components/WidgetInteractionComponent.cpp
--- a/Engine/Source/Runtime/UMG/Private/Components/WidgetInteractionComponent.cpp
+++ b/Engine/Source/Runtime/UMG/Private/Components/WidgetInteractionComponent.cpp
@@ -406,6 +406,14 @@ FWidgetPath UWidgetInteractionComponent::DetermineWidgetUnderPointer()
 		{
 			OldHoveredWidgetComponent->RequestRedraw();
 		}
+		
+		if ( NewHoveredWidgetComponent )
+		{
+			if(VirtualUser.IsValid())
+			{
+				NewHoveredWidgetComponent->OnWorldUserEntered(VirtualUser->GetUserIndex());
+			}
+		}
 
 		OnHoveredWidgetChanged.Broadcast( NewHoveredWidgetComponent, OldHoveredWidgetComponent );
 	}
diff --git a/Engine/Source/Runtime/UMG/Public/Components/WidgetComponent.h b/Engine/Source/Runtime/UMG/Public/Components/WidgetComponent.h
--- a/Engine/Source/Runtime/UMG/Public/Components/WidgetComponent.h
+++ b/Engine/Source/Runtime/UMG/Public/Components/WidgetComponent.h
@@ -147,6 +147,9 @@ public:
 	* @param bDrawSizeChanged Whether the draw size of this component has changed since the last update call.
 	*/
 	UMG_API void UpdateBodySetup( bool bDrawSizeChanged = false );
+	
+	/** To be called when a SlateUser enters this world space widget. */
+	UMG_API void OnWorldUserEntered(int32 SlateUserIndex);
 
 	/**
 	 * Converts a world-space hit result to a hit location on the widget
@@ -651,6 +654,9 @@ private:
 
 	/** The Slate widget to be displayed by this component.  Only one of either Widget or SlateWidget can be used */
 	TSharedPtr<SWidget> SlateWidget;
+	
+	/** The Slate user ids of the users which have entered our world-space widget (and thus could've gotten focus). */
+	TSet<int32> UsersEnteredWorldWidget;
 
 	/** The slate widget currently being drawn. */
 	TWeakPtr<SWidget> CurrentSlateWidget;
@@ -664,6 +670,9 @@ private:
 
 	/** When using Screen space, this code will add the widget to the screen. **/
 	UMG_API void AddWidgetToScreen(ULocalPlayer* TargetPlayer);
+	
+	/** When using World space, this will reset all the users focused on the active widget.  */
+	UMG_API void ResetWorldUserFocusOnWidget();
 
 	/** Set to true after a draw of an empty component.*/
 	bool bRenderCleared;
@@ -708,4 +717,4 @@ public:
 public:
 	TSubclassOf<UUserWidget> WidgetClass;
 	TObjectPtr<UTextureRenderTarget2D> RenderTarget;
-};
\ No newline at end of file
+};
