Added SlateWindow nullptr check to fix an editor crash caused by CL #45190987

diff --git a/Engine/Source/Runtime/UMG/Private/Components/WidgetComponent.cpp b/Engine/Source/Runtime/UMG/Private/Components/WidgetComponent.cpp
--- a/Engine/Source/Runtime/UMG/Private/Components/WidgetComponent.cpp
+++ b/Engine/Source/Runtime/UMG/Private/Components/WidgetComponent.cpp
@@ -1544,15 +1544,18 @@ void UWidgetComponent::RemoveWidgetFromScreen()
 #if !UE_SERVER
 	if (!IsRunningDedicatedServer())
 	{
-		ResetWorldUserFocusOnWidget();
-		
-		if(GetWidgetSpace() == EWidgetSpace::World)
+		if(SlateWindow != nullptr)
 		{
-			SlateWindow->SetContent(SNullWidget::NullWidget);
+			ResetWorldUserFocusOnWidget();
+			
+			if(GetWidgetSpace() == EWidgetSpace::World)
+			{
+				SlateWindow->SetContent(SNullWidget::NullWidget);
+			}
+			
+			SlateWindow.Reset();
 		}
 		
-		SlateWindow.Reset();
-		
 		bAddedToScreen = false;
 
 		if (UWorld* World = GetWorld())
