UMG: Fix anchor dragging canvas panel slots

#rnx

diff --git a/Engine/Source/Runtime/UMG/Private/Components/CanvasPanel.cpp b/Engine/Source/Runtime/UMG/Private/Components/CanvasPanel.cpp
--- a/Engine/Source/Runtime/UMG/Private/Components/CanvasPanel.cpp
+++ b/Engine/Source/Runtime/UMG/Private/Components/CanvasPanel.cpp
@@ -44,7 +44,7 @@ void UCanvasPanel::OnSlotRemoved(UPanelSlot* InSlot)
 	// Remove the widget from the live slot if it exists.
 	if ( MyCanvas.IsValid() && InSlot->Content)
 	{
-		TSharedPtr<SWidget> Widget = InSlot->Content->GetCachedWidget();
+		TSharedPtr<SWidget> Widget = InSlot->Content->GetCachedWrappedWidget();
 		if ( Widget.IsValid() )
 		{
 			MyCanvas->RemoveSlot(Widget.ToSharedRef());
@@ -99,7 +99,7 @@ bool UCanvasPanel::GetGeometryForSlot(UCanvasPanelSlot* InSlot, FGeometry& Arran
 
 		for ( int32 ChildIndex = 0; ChildIndex < ArrangedChildren.Num(); ChildIndex++ )
 		{
-			if ( ArrangedChildren[ChildIndex].Widget == InSlot->Content->GetCachedWidget() )
+			if ( ArrangedChildren[ChildIndex].Widget == InSlot->Content->GetCachedWrappedWidget() )
 			{
 				ArrangedGeometry = ArrangedChildren[ChildIndex].Geometry;
 				return true;
