UMG Sequencer: Fixed 2 issues.

1) When adding an event track as the first item in a widget animation, the duration of the animation would reset to 0, and notably the playback range start would be set to inclusize (range as (0,0)). This would persist after the length changed, for example to (0,500), meaning that any items placed on the 0 frame, such as events, would never evaluate.

2) When not encountering the above issue, items placed on the 0 frame would evaluate twice. This was due to a refactor to the way widget animation would playback.

diff --git a/Engine/Source/Editor/Sequencer/Private/Sequencer.cpp b/Engine/Source/Editor/Sequencer/Private/Sequencer.cpp
--- a/Engine/Source/Editor/Sequencer/Private/Sequencer.cpp
+++ b/Engine/Source/Editor/Sequencer/Private/Sequencer.cpp
@@ -3328,7 +3328,7 @@ void FSequencer::UpdatePlaybackRange()
 		}
 
 		// When the playback range is determined by the section bounds, don't mark the change in the playback range otherwise the scene will be marked dirty
-		if (!NewBounds.IsDegenerate())
+		if (!NewBounds.IsDegenerate() && !NewBounds.IsEmpty())
 		{
 			// Playback ranges should always have exclusive upper bounds
 			if (NewBounds.GetUpperBound().IsInclusive())
diff --git a/Engine/Source/Runtime/UMG/Private/Animation/WidgetAnimationState.cpp b/Engine/Source/Runtime/UMG/Private/Animation/WidgetAnimationState.cpp
--- a/Engine/Source/Runtime/UMG/Private/Animation/WidgetAnimationState.cpp
+++ b/Engine/Source/Runtime/UMG/Private/Animation/WidgetAnimationState.cpp
@@ -532,12 +532,23 @@ void FWidgetAnimationState::Play(const FWidgetAnimationStatePlayParams& PlayPara
 	TSharedPtr<FMovieSceneEntitySystemRunner> Runner = SharedPlaybackState->GetRunner();
 	if (Runner)
 	{
-		const FMovieSceneContext FirstContext = PlaybackManager.UpdateAtCurrentTime();
-		Runner->QueueUpdate(
-				FirstContext, 
+		// Generate a full context for this start using Update rather than UpdateAtCurrentTime.
+		// This is so we update the 'previous' time to the start time so we don't evaluate it again during our first tick
+		FMovieScenePlaybackManager::FContexts StartContexts;
+		PlaybackManager.Update(0, StartContexts);
+
+		// We should just have a single context
+		ensure(StartContexts.Num() == 1);
+		if (StartContexts.Num() > 0)
+		{
+			FMovieSceneContext& FirstContext(StartContexts[0]);
+
+			Runner->QueueUpdate(
+				FirstContext,
 				SharedPlaybackState->GetRootInstanceHandle(),
 				FSimpleDelegate::CreateSP(this, &FWidgetAnimationState::OnBegunPlay),
 				ERunnerUpdateFlags::Flush);
+		}
 	}
 
 #if !NO_LOGGING
