Correction to SizeBox and ScaleBox Component to use a Weak Ptr to the SlateWidget created. The We do not need a Shared Ptr since the Owner Widget keeps a reference to it.

diff --git a/Engine/Source/Runtime/UMG/Private/Components/ScaleBoxComponent.cpp b/Engine/Source/Runtime/UMG/Private/Components/ScaleBoxComponent.cpp
--- a/Engine/Source/Runtime/UMG/Private/Components/ScaleBoxComponent.cpp
+++ b/Engine/Source/Runtime/UMG/Private/Components/ScaleBoxComponent.cpp
@@ -15,9 +15,9 @@
 void UScaleBoxComponent::SetStretch(EStretch::Type InStretch)
 {
 	Stretch = InStretch;
-	if (MyScaleBox.IsValid())
+	if (const TSharedPtr<SScaleBox> ScaleBox = MyScaleBox.Pin())
 	{
-		MyScaleBox->SetStretch(InStretch);
+		ScaleBox->SetStretch(InStretch);
 	}
 }
 
@@ -29,9 +29,9 @@ EStretch::Type UScaleBoxComponent::GetStretch() const
 void UScaleBoxComponent::SetStretchDirection(EStretchDirection::Type InStretchDirection)
 {
 	StretchDirection = InStretchDirection;
-	if (MyScaleBox.IsValid())
+	if (const TSharedPtr<SScaleBox> ScaleBox = MyScaleBox.Pin())
 	{
-		MyScaleBox->SetStretchDirection(InStretchDirection);
+		ScaleBox->SetStretchDirection(InStretchDirection);
 	}
 }
 
@@ -43,9 +43,9 @@ EStretchDirection::Type UScaleBoxComponent::GetStretchDirection() const
 void UScaleBoxComponent::SetUserSpecifiedScale(float InUserSpecifiedScale)
 {
 	UserSpecifiedScale = InUserSpecifiedScale;
-	if (MyScaleBox.IsValid())
+	if (const TSharedPtr<SScaleBox> ScaleBox = MyScaleBox.Pin())
 	{
-		MyScaleBox->SetUserSpecifiedScale(InUserSpecifiedScale);
+		ScaleBox->SetUserSpecifiedScale(InUserSpecifiedScale);
 	}
 }
 
@@ -57,9 +57,9 @@ float UScaleBoxComponent::GetUserSpecifiedScale() const
 void UScaleBoxComponent::SetIgnoreInheritedScale(bool bInIgnoreInheritedScale)
 {
 	IgnoreInheritedScale = bInIgnoreInheritedScale;
-	if (MyScaleBox.IsValid())
+	if (const TSharedPtr<SScaleBox> ScaleBox = MyScaleBox.Pin())
 	{
-		MyScaleBox->SetIgnoreInheritedScale(bInIgnoreInheritedScale);
+		ScaleBox->SetIgnoreInheritedScale(bInIgnoreInheritedScale);
 	}
 }
 
@@ -76,9 +76,9 @@ EHorizontalAlignment UScaleBoxComponent::GetHorizontalAlignment() const
 void UScaleBoxComponent::SetHorizontalAlignment(EHorizontalAlignment InHorizontalAlignment)
 {
 	HorizontalAlignment = InHorizontalAlignment;
-	if (MyScaleBox.IsValid())
+	if (const TSharedPtr<SScaleBox> ScaleBox = MyScaleBox.Pin())
 	{
-		MyScaleBox->SetHAlign(InHorizontalAlignment);
+		ScaleBox->SetHAlign(InHorizontalAlignment);
 	}
 }
 
@@ -90,9 +90,9 @@ EVerticalAlignment UScaleBoxComponent::GetVerticalAlignment() const
 void UScaleBoxComponent::SetVerticalAlignment(EVerticalAlignment InVerticalAlignment)
 {
 	VerticalAlignment = InVerticalAlignment;
-	if (MyScaleBox.IsValid())
+	if (const TSharedPtr<SScaleBox> ScaleBox = MyScaleBox.Pin())
 	{
-		MyScaleBox->SetVAlign(InVerticalAlignment);
+		ScaleBox->SetVAlign(InVerticalAlignment);
 	}
 }
 
@@ -131,9 +131,9 @@ bool UScaleBoxComponent::CanEditChange(const FProperty* InProperty) const
 //		DesignerSize = FVector2D(0, 0);
 //	}
 //
-//	if (MyScaleBox.IsValid())
+//	if (const TSharedPtr<SScaleBox> ScaleBox = MyScaleBox.Pin())
 //	{
-//		MyScaleBox->SetOverrideScreenInformation(DesignerSize);
+//		ScaleBox->SetOverrideScreenInformation(DesignerSize);
 //	}
 //}
 
@@ -141,32 +141,35 @@ void UScaleBoxComponent::PostEditChangeProperty(FPropertyChangedEvent& PropertyC
 {
 	Super::PostEditChangeProperty(PropertyChangedEvent);
 
-	SynchronizeProperties();
+	if (const TSharedPtr<SScaleBox> ScaleBox = MyScaleBox.Pin())
+	{
+		SynchronizeProperties(ScaleBox.ToSharedRef());
+	}
 }
 #endif //WITH_EDITOR
 
-void UScaleBoxComponent::SynchronizeProperties()
-{
-	SetStretchDirection(StretchDirection);
-	SetStretch(Stretch);
-	SetUserSpecifiedScale(UserSpecifiedScale);
-	SetIgnoreInheritedScale(IgnoreInheritedScale);
+void UScaleBoxComponent::SynchronizeProperties(TSharedRef<SScaleBox> ScaleBox)
+{	
+	ScaleBox->SetStretchDirection(StretchDirection);
+	ScaleBox->SetStretch(Stretch);
+	ScaleBox->SetUserSpecifiedScale(UserSpecifiedScale);
+	ScaleBox->SetIgnoreInheritedScale(IgnoreInheritedScale);
 
 	// Set child slot properties
-	SetHorizontalAlignment(HorizontalAlignment);
-	SetVerticalAlignment(VerticalAlignment);
+	ScaleBox->SetHAlign(HorizontalAlignment);
+	ScaleBox->SetVAlign(VerticalAlignment);
 }
 
 TSharedRef<SWidget> UScaleBoxComponent::RebuildWidgetWithContent(TSharedRef<SWidget> OwnerContent)
 {
-	MyScaleBox = SNew(SScaleBox)
+	TSharedRef<SScaleBox> ScaleBox = SNew(SScaleBox)
 // TODO vinz: Implement OnDesignerChanged to work for Components. This currently is just a virtual function for UWidgets. Also, this probably should be a part of SynchronizeProperties.
 //#if WITH_EDITOR
 //		.OverrideScreenSize(DesignerSize)
 //#endif
 		[OwnerContent];
 
-	SynchronizeProperties();
-
-	return MyScaleBox.ToSharedRef();
+	MyScaleBox = ScaleBox.ToWeakPtr();
+	SynchronizeProperties(ScaleBox);
+	return ScaleBox;
 }
diff --git a/Engine/Source/Runtime/UMG/Private/Components/SizeBoxComponent.cpp b/Engine/Source/Runtime/UMG/Private/Components/SizeBoxComponent.cpp
--- a/Engine/Source/Runtime/UMG/Private/Components/SizeBoxComponent.cpp
+++ b/Engine/Source/Runtime/UMG/Private/Components/SizeBoxComponent.cpp
@@ -26,18 +26,18 @@ void USizeBoxComponent::SetWidthOverride(float InWidthOverride)
 {
 	bOverride_WidthOverride = true;
 	WidthOverride = InWidthOverride;
-	if (MySizeBox.IsValid())
+	if (const TSharedPtr<SBox> SizeBox = MySizeBox.Pin())
 	{
-		MySizeBox->SetWidthOverride(InWidthOverride);
+		SizeBox->SetWidthOverride(InWidthOverride);
 	}
 }
 
 void USizeBoxComponent::ClearWidthOverride()
 {
 	bOverride_WidthOverride = false;
-	if (MySizeBox.IsValid())
+	if (const TSharedPtr<SBox> SizeBox = MySizeBox.Pin())
 	{
-		MySizeBox->SetWidthOverride(FOptionalSize());
+		SizeBox->SetWidthOverride(FOptionalSize());
 	}
 }
 
@@ -55,18 +55,18 @@ void USizeBoxComponent::SetHeightOverride(float InHeightOverride)
 {
 	bOverride_HeightOverride = true;
 	HeightOverride = InHeightOverride;
-	if (MySizeBox.IsValid())
+	if (const TSharedPtr<SBox> SizeBox = MySizeBox.Pin())
 	{
-		MySizeBox->SetHeightOverride(InHeightOverride);
+		SizeBox->SetHeightOverride(InHeightOverride);
 	}
 }
 
 void USizeBoxComponent::ClearHeightOverride()
 {
 	bOverride_HeightOverride = false;
-	if (MySizeBox.IsValid())
+	if (const TSharedPtr<SBox> SizeBox = MySizeBox.Pin())
 	{
-		MySizeBox->SetHeightOverride(FOptionalSize());
+		SizeBox->SetHeightOverride(FOptionalSize());
 	}
 }
 
@@ -84,18 +84,18 @@ void USizeBoxComponent::SetMinDesiredWidth(float InMinDesiredWidth)
 {
 	bOverride_MinDesiredWidth = true;
 	MinDesiredWidth = InMinDesiredWidth;
-	if (MySizeBox.IsValid())
+	if (const TSharedPtr<SBox> SizeBox = MySizeBox.Pin())
 	{
-		MySizeBox->SetMinDesiredWidth(InMinDesiredWidth);
+		SizeBox->SetMinDesiredWidth(InMinDesiredWidth);
 	}
 }
 
 void USizeBoxComponent::ClearMinDesiredWidth()
 {
 	bOverride_MinDesiredWidth = false;
-	if (MySizeBox.IsValid())
+	if (const TSharedPtr<SBox> SizeBox = MySizeBox.Pin())
 	{
-		MySizeBox->SetMinDesiredWidth(FOptionalSize());
+		SizeBox->SetMinDesiredWidth(FOptionalSize());
 	}
 }
 
@@ -113,18 +113,18 @@ void USizeBoxComponent::SetMinDesiredHeight(float InMinDesiredHeight)
 {
 	bOverride_MinDesiredHeight = true;
 	MinDesiredHeight = InMinDesiredHeight;
-	if (MySizeBox.IsValid())
+	if (const TSharedPtr<SBox> SizeBox = MySizeBox.Pin())
 	{
-		MySizeBox->SetMinDesiredHeight(InMinDesiredHeight);
+		SizeBox->SetMinDesiredHeight(InMinDesiredHeight);
 	}
 }
 
 void USizeBoxComponent::ClearMinDesiredHeight()
 {
 	bOverride_MinDesiredHeight = false;
-	if (MySizeBox.IsValid())
+	if (const TSharedPtr<SBox> SizeBox = MySizeBox.Pin())
 	{
-		MySizeBox->SetMinDesiredHeight(FOptionalSize());
+		SizeBox->SetMinDesiredHeight(FOptionalSize());
 	}
 }
 
@@ -142,18 +142,18 @@ void USizeBoxComponent::SetMaxDesiredWidth(float InMaxDesiredWidth)
 {
 	bOverride_MaxDesiredWidth = true;
 	MaxDesiredWidth = InMaxDesiredWidth;
-	if (MySizeBox.IsValid())
+	if (const TSharedPtr<SBox> SizeBox = MySizeBox.Pin())
 	{
-		MySizeBox->SetMaxDesiredWidth(InMaxDesiredWidth);
+		SizeBox->SetMaxDesiredWidth(InMaxDesiredWidth);
 	}
 }
 
 void USizeBoxComponent::ClearMaxDesiredWidth()
 {
 	bOverride_MaxDesiredWidth = false;
-	if (MySizeBox.IsValid())
+	if (const TSharedPtr<SBox> SizeBox = MySizeBox.Pin())
 	{
-		MySizeBox->SetMaxDesiredWidth(FOptionalSize());
+		SizeBox->SetMaxDesiredWidth(FOptionalSize());
 	}
 }
 
@@ -171,18 +171,18 @@ void USizeBoxComponent::SetMaxDesiredHeight(float InMaxDesiredHeight)
 {
 	bOverride_MaxDesiredHeight = true;
 	MaxDesiredHeight = InMaxDesiredHeight;
-	if (MySizeBox.IsValid())
+	if (const TSharedPtr<SBox> SizeBox = MySizeBox.Pin())
 	{
-		MySizeBox->SetMaxDesiredHeight(InMaxDesiredHeight);
+		SizeBox->SetMaxDesiredHeight(InMaxDesiredHeight);
 	}
 }
 
 void USizeBoxComponent::ClearMaxDesiredHeight()
 {
 	bOverride_MaxDesiredHeight = false;
-	if (MySizeBox.IsValid())
+	if (const TSharedPtr<SBox> SizeBox = MySizeBox.Pin())
 	{
-		MySizeBox->SetMaxDesiredHeight(FOptionalSize());
+		SizeBox->SetMaxDesiredHeight(FOptionalSize());
 	}
 }
 
@@ -200,18 +200,18 @@ void USizeBoxComponent::SetMinAspectRatio(float InMinAspectRatio)
 {
 	bOverride_MinAspectRatio = true;
 	MinAspectRatio = InMinAspectRatio;
-	if (MySizeBox.IsValid())
+	if (const TSharedPtr<SBox> SizeBox = MySizeBox.Pin())
 	{
-		MySizeBox->SetMinAspectRatio(InMinAspectRatio);
+		SizeBox->SetMinAspectRatio(InMinAspectRatio);
 	}
 }
 
 void USizeBoxComponent::ClearMinAspectRatio()
 {
 	bOverride_MinAspectRatio = false;
-	if (MySizeBox.IsValid())
+	if (const TSharedPtr<SBox> SizeBox = MySizeBox.Pin())
 	{
-		MySizeBox->SetMinAspectRatio(FOptionalSize());
+		SizeBox->SetMinAspectRatio(FOptionalSize());
 	}
 }
 
@@ -229,18 +229,18 @@ void USizeBoxComponent::SetMaxAspectRatio(float InMaxAspectRatio)
 {
 	bOverride_MaxAspectRatio = true;
 	MaxAspectRatio = InMaxAspectRatio;
-	if (MySizeBox.IsValid())
+	if (const TSharedPtr<SBox> SizeBox = MySizeBox.Pin())
 	{
-		MySizeBox->SetMaxAspectRatio(InMaxAspectRatio);
+		SizeBox->SetMaxAspectRatio(InMaxAspectRatio);
 	}
 }
 
 void USizeBoxComponent::ClearMaxAspectRatio()
 {
 	bOverride_MaxAspectRatio = false;
-	if (MySizeBox.IsValid())
+	if (const TSharedPtr<SBox> SizeBox = MySizeBox.Pin())
 	{
-		MySizeBox->SetMaxAspectRatio(FOptionalSize());
+		SizeBox->SetMaxAspectRatio(FOptionalSize());
 	}
 }
 
@@ -252,9 +252,9 @@ FMargin USizeBoxComponent::GetPadding() const
 void USizeBoxComponent::SetPadding(FMargin InPadding)
 {
 	Padding = InPadding;
-	if (MySizeBox.IsValid())
+	if (const TSharedPtr<SBox> SizeBox = MySizeBox.Pin())
 	{
-		MySizeBox->SetPadding(InPadding);
+		SizeBox->SetPadding(InPadding);
 	}
 }
 
@@ -266,9 +266,9 @@ EHorizontalAlignment USizeBoxComponent::GetHorizontalAlignment() const
 void USizeBoxComponent::SetHorizontalAlignment(EHorizontalAlignment InHorizontalAlignment)
 {
 	HorizontalAlignment = InHorizontalAlignment;
-	if (MySizeBox.IsValid())
+	if (const TSharedPtr<SBox> SizeBox = MySizeBox.Pin())
 	{
-		MySizeBox->SetHAlign(InHorizontalAlignment);
+		SizeBox->SetHAlign(InHorizontalAlignment);
 	}
 }
 
@@ -280,9 +280,9 @@ EVerticalAlignment USizeBoxComponent::GetVerticalAlignment() const
 void USizeBoxComponent::SetVerticalAlignment(EVerticalAlignment InVerticalAlignment)
 {
 	VerticalAlignment = InVerticalAlignment;
-	if (MySizeBox.IsValid())
+	if (const TSharedPtr<SBox> SizeBox = MySizeBox.Pin())
 	{
-		MySizeBox->SetVAlign(InVerticalAlignment);
+		SizeBox->SetVAlign(InVerticalAlignment);
 	}
 }
 
@@ -291,97 +291,38 @@ void USizeBoxComponent::PostEditChangeProperty(FPropertyChangedEvent& PropertyCh
 {
 	Super::PostEditChangeProperty(PropertyChangedEvent);
 
-	SynchronizeProperties();
-}
-#endif //WITH_EDITOR
-
-void USizeBoxComponent::SynchronizeProperties()
-{
-	if (bOverride_WidthOverride)
-	{
-		SetWidthOverride(WidthOverride);
-	}
-	else
-	{
-		ClearWidthOverride();
-	}
-
-	if (bOverride_HeightOverride)
-	{
-		SetHeightOverride(HeightOverride);
-	}
-	else
-	{
-		ClearHeightOverride();
-	}
-
-	if (bOverride_MinDesiredWidth)
-	{
-		SetMinDesiredWidth(MinDesiredWidth);
-	}
-	else
-	{
-		ClearMinDesiredWidth();
-	}
-
-	if (bOverride_MinDesiredHeight)
-	{
-		SetMinDesiredHeight(MinDesiredHeight);
-	}
-	else
-	{
-		ClearMinDesiredHeight();
-	}
-
-	if (bOverride_MaxDesiredWidth)
-	{
-		SetMaxDesiredWidth(MaxDesiredWidth);
-	}
-	else
-	{
-		ClearMaxDesiredWidth();
-	}
-
-	if (bOverride_MaxDesiredHeight)
-	{
-		SetMaxDesiredHeight(MaxDesiredHeight);
-	}
-	else
-	{
-		ClearMaxDesiredHeight();
-	}
-
-	if (bOverride_MinAspectRatio)
-	{
-		SetMinAspectRatio(MinAspectRatio);
-	}
-	else
-	{
-		ClearMinAspectRatio();
-	}
-
-	if (bOverride_MaxAspectRatio)
+	
+	if (const TSharedPtr<SBox> SizeBox = MySizeBox.Pin())
 	{
-		SetMaxAspectRatio(MaxAspectRatio);
-	}
-	else
-	{
-		ClearMaxAspectRatio();
+		SynchronizeProperties(SizeBox.ToSharedRef());
 	}
+}
+#endif //WITH_EDITOR
 
+void USizeBoxComponent::SynchronizeProperties(TSharedRef<SBox> SizeBox)
+{
+	SizeBox->SetWidthOverride(bOverride_WidthOverride? WidthOverride : FOptionalSize());
+	SizeBox->SetHeightOverride(bOverride_HeightOverride ? HeightOverride : FOptionalSize());
+	SizeBox->SetMinDesiredWidth(bOverride_MinDesiredWidth? MinDesiredWidth : FOptionalSize());
+	SizeBox->SetMinDesiredHeight(bOverride_MinDesiredHeight ? MinDesiredHeight : FOptionalSize());
+	SizeBox->SetMaxDesiredWidth(bOverride_MaxDesiredWidth ? MaxDesiredWidth : FOptionalSize());
+	SizeBox->SetMaxDesiredHeight(bOverride_MaxDesiredHeight ? MaxDesiredHeight : FOptionalSize());
+	SizeBox->SetMinAspectRatio(bOverride_MinAspectRatio ? MinAspectRatio : FOptionalSize());
+	SizeBox->SetMaxAspectRatio(bOverride_MaxAspectRatio ? MaxAspectRatio : FOptionalSize());
+	
 	// Set child slot properties
-	SetPadding(Padding);
-	SetHorizontalAlignment(HorizontalAlignment);
-	SetVerticalAlignment(VerticalAlignment);
+	SizeBox->SetPadding(Padding);
+	SizeBox->SetHAlign(HorizontalAlignment);
+	SizeBox->SetVAlign(VerticalAlignment);
 }
 
 TSharedRef<SWidget> USizeBoxComponent::RebuildWidgetWithContent(TSharedRef<SWidget> OwnerContent)
 {
-	MySizeBox = SNew(SBox)
+	TSharedRef<SBox> SizeBox = SNew(SBox)
 		[OwnerContent];
 
-	SynchronizeProperties();
-
-	return MySizeBox.ToSharedRef();
+	MySizeBox = SizeBox.ToWeakPtr();
+	SynchronizeProperties(SizeBox);
+	return SizeBox;
 }
 
diff --git a/Engine/Source/Runtime/UMG/Public/Components/ScaleBoxComponent.h b/Engine/Source/Runtime/UMG/Public/Components/ScaleBoxComponent.h
--- a/Engine/Source/Runtime/UMG/Public/Components/ScaleBoxComponent.h
+++ b/Engine/Source/Runtime/UMG/Public/Components/ScaleBoxComponent.h
@@ -55,11 +55,8 @@ public:
 	virtual void PostEditChangeProperty(struct FPropertyChangedEvent& PropertyChangedEvent) override;
 	// End of UObject interface
 #endif // WITH_EDITOR
-
-	UMG_API void SynchronizeProperties();
-
-protected:
-	TSharedPtr<SScaleBox> MyScaleBox;
+	
+	virtual TSharedRef<SWidget> RebuildWidgetWithContent(const TSharedRef<SWidget> OwnerContent) override;
 
 // TODO vinz: Implement OnDesignerChanged to work for Components. This currently is just a virtual function for UWidgets.
 //#if WITH_EDITOR
@@ -68,6 +65,8 @@ protected:
 
 private:
 
+	void SynchronizeProperties(TSharedRef<SScaleBox> ScaleBox);
+
 	// Child slot properties
 
 	/** This property is for the Widget that owns this Component. The alignment of the object horizontally. */
@@ -96,5 +95,5 @@ private:
 	UPROPERTY(EditAnywhere, BlueprintReadWrite, Getter="IsIgnoreInheritedScale", Setter, Category = "Stretching", meta = (AllowPrivateAccess = true))
 	bool IgnoreInheritedScale = false;
 
-	virtual TSharedRef<SWidget> RebuildWidgetWithContent(const TSharedRef<SWidget> OwnerContent) override;
+	TWeakPtr<SScaleBox> MyScaleBox;
 };
\ No newline at end of file
diff --git a/Engine/Source/Runtime/UMG/Public/Components/SizeBoxComponent.h b/Engine/Source/Runtime/UMG/Public/Components/SizeBoxComponent.h
--- a/Engine/Source/Runtime/UMG/Public/Components/SizeBoxComponent.h
+++ b/Engine/Source/Runtime/UMG/Public/Components/SizeBoxComponent.h
@@ -132,13 +132,13 @@ public:
 	virtual void PostEditChangeProperty(struct FPropertyChangedEvent& PropertyChangedEvent) override;
 	// End of UObject interface
 #endif // WITH_EDITOR
-
-	UMG_API void SynchronizeProperties();
-
-protected:
-	TSharedPtr<SBox> MySizeBox;
+	
+	virtual TSharedRef<SWidget> RebuildWidgetWithContent(const TSharedRef<SWidget> OwnerContent) override;
 
 private:
+		
+	void SynchronizeProperties(TSharedRef<SBox> SizeBox);
+	
 	// Child slot properties:
 
 	/** This property is for the Widget that owns this Component. The padding area between the slot and the content it contains. */
@@ -219,5 +219,5 @@ private:
 	UPROPERTY(EditAnywhere, Category = "Child Layout", meta = (InlineEditConditionToggle, AllowPrivateAccess = true))
 	uint8 bOverride_MaxAspectRatio : 1;
 
-	virtual TSharedRef<SWidget> RebuildWidgetWithContent(const TSharedRef<SWidget> OwnerContent) override;
+	TWeakPtr<SBox> MySizeBox;
 };
\ No newline at end of file
