UMG: Add helper function UUserObjectListEntryLibrary::GetListItemIndex

#changelist validated

#tests Frontend

diff --git a/Engine/Source/Runtime/UMG/Private/Blueprint/IUserObjectListEntry.cpp b/Engine/Source/Runtime/UMG/Private/Blueprint/IUserObjectListEntry.cpp
--- a/Engine/Source/Runtime/UMG/Private/Blueprint/IUserObjectListEntry.cpp
+++ b/Engine/Source/Runtime/UMG/Private/Blueprint/IUserObjectListEntry.cpp
@@ -46,3 +46,18 @@ UObject* UUserObjectListEntryLibrary::GetListItemObject(TScriptInterface<IUserOb
 	return nullptr;
 }
 
+int32 UUserObjectListEntryLibrary::GetListItemIndex(TScriptInterface<IUserObjectListEntry> UserObjectListEntry)
+{
+	if (UserObjectListEntry)
+	{
+		if (UUserWidget* Item = Cast<UUserWidget>(UserObjectListEntry->GetListItem()))
+		{
+			const UListView* OwningListView = Cast<UListView>(UUserListEntryLibrary::GetOwningListView(Item));
+			if (OwningListView != nullptr)
+			{
+				return OwningListView->GetIndexForItem(Item);
+			}
+		}
+	}
+	return INDEX_NONE;
+}
diff --git a/Engine/Source/Runtime/UMG/Public/Blueprint/IUserObjectListEntry.h b/Engine/Source/Runtime/UMG/Public/Blueprint/IUserObjectListEntry.h
--- a/Engine/Source/Runtime/UMG/Public/Blueprint/IUserObjectListEntry.h
+++ b/Engine/Source/Runtime/UMG/Public/Blueprint/IUserObjectListEntry.h
@@ -56,4 +56,8 @@ public:
 	 */
 	UFUNCTION(BlueprintPure, Category = UserObjectListEntry, meta = (DefaultToSelf = UserObjectListEntry))
 	static UMG_API UObject* GetListItemObject(TScriptInterface<IUserObjectListEntry> UserObjectListEntry);
+
+	/** Returns the index of the list item in the owning UListView */
+	UFUNCTION(BlueprintPure, Category = UserObjectListEntry, meta = (DefaultToSelf = UserObjectListEntry))
+	static UMG_API int32 GetListItemIndex(TScriptInterface<IUserObjectListEntry> UserObjectListEntry);
 };
