UUserObjectListEntryLibrary::GetListItemIndex - Replace cast to UUserWidget to cast to UObject

#changelist validated
#tests Frontend

diff --git a/Engine/Source/Runtime/UMG/Private/Blueprint/IUserObjectListEntry.cpp b/Engine/Source/Runtime/UMG/Private/Blueprint/IUserObjectListEntry.cpp
--- a/Engine/Source/Runtime/UMG/Private/Blueprint/IUserObjectListEntry.cpp
+++ b/Engine/Source/Runtime/UMG/Private/Blueprint/IUserObjectListEntry.cpp
@@ -50,7 +50,7 @@ int32 UUserObjectListEntryLibrary::GetListItemIndex(TScriptInterface<IUserObject
 {
 	if (UserObjectListEntry)
 	{
-		if (UUserWidget* Item = Cast<UUserWidget>(UserObjectListEntry->GetListItem()))
+		if (UObject* Item = Cast<UObject>(UserObjectListEntry->GetListItem()))
 		{
 			const UListView* OwningListView = Cast<UListView>(UUserListEntryLibrary::GetOwningListView(Item));
 			if (OwningListView != nullptr)
