UUserObjectListEntry::GetListItemIndex - correction - LogicError - Item object needs to be retrieved from "GetListItemObject", not UUserWidget->"GetListItem" - this would always fail

#changelist validated

#tests Frontend

diff --git a/Engine/Source/Runtime/UMG/Private/Blueprint/IUserObjectListEntry.cpp b/Engine/Source/Runtime/UMG/Private/Blueprint/IUserObjectListEntry.cpp
--- a/Engine/Source/Runtime/UMG/Private/Blueprint/IUserObjectListEntry.cpp
+++ b/Engine/Source/Runtime/UMG/Private/Blueprint/IUserObjectListEntry.cpp
@@ -48,16 +48,15 @@ UObject* UUserObjectListEntryLibrary::GetListItemObject(TScriptInterface<IUserOb
 
 int32 UUserObjectListEntryLibrary::GetListItemIndex(TScriptInterface<IUserObjectListEntry> UserObjectListEntry)
 {
-	if (UserObjectListEntry)
+	if (UUserWidget* EntryWidget = Cast<UUserWidget>(UserObjectListEntry.GetObject()))
 	{
-		if (UObject* Item = Cast<UObject>(UserObjectListEntry->GetListItem()))
+		if (const UListView* OwningListView = Cast<UListView>(UUserListEntryLibrary::GetOwningListView(EntryWidget)))
 		{
-			const UListView* OwningListView = Cast<UListView>(UUserListEntryLibrary::GetOwningListView(Item));
-			if (OwningListView != nullptr)
+			if (UObject* Item = UUserObjectListEntryLibrary::GetListItemObject(UserObjectListEntry))
 			{
 				return OwningListView->GetIndexForItem(Item);
 			}
 		}
 	}
 	return INDEX_NONE;
-}
+}
\ No newline at end of file
