Back-ported UI crash fix from Main

#rnx

diff --git a/Engine/Source/Runtime/UMG/Private/Components/Button.cpp b/Engine/Source/Runtime/UMG/Private/Components/Button.cpp
--- a/Engine/Source/Runtime/UMG/Private/Components/Button.cpp
+++ b/Engine/Source/Runtime/UMG/Private/Components/Button.cpp
@@ -241,7 +241,7 @@ bool UButton::GetIsFocusable() const
 UMG_API void UButton::SetAllowDragDrop(bool bInAllowDragDrop)
 {
 	bAllowDragDrop = bInAllowDragDrop;
-	if (MyButton->IsParentValid())
+	if (MyButton.IsValid() && MyButton->IsParentValid())
 	{
 		MyButton->SetAllowDragDrop(bAllowDragDrop);
 	}
