UUserObjectListEntryLibrary::IsFirstWidget, IsLastWidget
Helpers to allow the list entry to determine how it should display if it is the first or last widget in the listview scrollbox.

#tests Frontend

diff --git a/Engine/Source/Runtime/UMG/Private/Blueprint/IUserObjectListEntry.cpp b/Engine/Source/Runtime/UMG/Private/Blueprint/IUserObjectListEntry.cpp
--- a/Engine/Source/Runtime/UMG/Private/Blueprint/IUserObjectListEntry.cpp
+++ b/Engine/Source/Runtime/UMG/Private/Blueprint/IUserObjectListEntry.cpp
@@ -59,4 +59,25 @@ int32 UUserObjectListEntryLibrary::GetListItemIndex(TScriptInterface<IUserObject
 		}
 	}
 	return INDEX_NONE;
-}
\ No newline at end of file
+}
+
+bool UUserObjectListEntryLibrary::IsFirstWidget(TScriptInterface<IUserObjectListEntry> UserObjectListEntry)
+{
+	return UUserObjectListEntryLibrary::GetListItemIndex(UserObjectListEntry) == 0;
+}
+
+bool UUserObjectListEntryLibrary::IsLastWidget(TScriptInterface<IUserObjectListEntry> UserObjectListEntry)
+{
+	if (UUserWidget* EntryWidget = Cast<UUserWidget>(UserObjectListEntry.GetObject()))
+	{
+		if (const UListView* OwningListView = Cast<UListView>(UUserListEntryLibrary::GetOwningListView(EntryWidget)))
+		{
+			if (UObject* Item = UUserObjectListEntryLibrary::GetListItemObject(UserObjectListEntry))
+			{
+				const int32 Index = OwningListView->GetIndexForItem(Item);
+				return Index != INDEX_NONE && Index == OwningListView->GetNumItems() - 1;
+			}
+		}
+	}
+	return false;
+}
diff --git a/Engine/Source/Runtime/UMG/Public/Blueprint/IUserObjectListEntry.h b/Engine/Source/Runtime/UMG/Public/Blueprint/IUserObjectListEntry.h
--- a/Engine/Source/Runtime/UMG/Public/Blueprint/IUserObjectListEntry.h
+++ b/Engine/Source/Runtime/UMG/Public/Blueprint/IUserObjectListEntry.h
@@ -60,4 +60,12 @@ public:
 	/** Returns the index of the list item in the owning UListView */
 	UFUNCTION(BlueprintPure, Category = UserObjectListEntry, meta = (DefaultToSelf = UserObjectListEntry))
 	static UMG_API int32 GetListItemIndex(TScriptInterface<IUserObjectListEntry> UserObjectListEntry);
+
+	/** Returns true if the widget list entry is the first item in the list view */
+	UFUNCTION(BlueprintPure, Category = UserObjectListEntry, meta = (DefaultToSelf = UserObjectListEntry))
+	static UMG_API bool IsFirstWidget(TScriptInterface<IUserObjectListEntry> UserObjectListEntry);
+
+	/** Returns true if the widget list entry is the last item in the list view */
+	UFUNCTION(BlueprintPure, Category = UserObjectListEntry, meta = (DefaultToSelf = UserObjectListEntry))
+	static UMG_API bool IsLastWidget(TScriptInterface<IUserObjectListEntry> UserObjectListEntry);
 };
