Fix UMGDesigner Drag/Drop. Ensure designer widget is hittestable.
Fix WidgetAnimation State, Content widget that is possesed should retain a handle on the widget internals to properly bind and animate properties.

#rnx

diff --git a/Engine/Source/Runtime/UMG/Private/Components/Widget.cpp b/Engine/Source/Runtime/UMG/Private/Components/Widget.cpp
--- a/Engine/Source/Runtime/UMG/Private/Components/Widget.cpp
+++ b/Engine/Source/Runtime/UMG/Private/Components/Widget.cpp
@@ -1149,7 +1149,8 @@ TSharedRef<SWidget> UWidget::RebuildDesignWidget(TSharedRef<SWidget> Content)
 TSharedRef<SWidget> UWidget::CreateDesignerOutline(TSharedRef<SWidget> Content) const
 {
 	return SNew(SOverlay)
-	
+		.Visibility(EVisibility::Visible)
+		.AddMetaData<FInvisibleToWidgetReflectorMetaData>(FInvisibleToWidgetReflectorMetaData())
 		+ SOverlay::Slot()
 		.HAlign(HAlign_Fill)
 		.VAlign(VAlign_Fill)
@@ -1460,7 +1461,7 @@ void UWidget::SynchronizeProperties()
 	// We want to apply the bindings to the cached widget, which could be the SWidget, or the SObjectWidget, 
 	// in the case where it's a user widget.  We always want to prefer the SObjectWidget so that bindings to 
 	// visibility and enabled status are not stomping values setup in the root widget in the User Widget.
-	TSharedPtr<SWidget> SafeWidget = GetCachedWrappedWidget();
+	TSharedPtr<SWidget> SafeWidget = GetCachedWidget();
 	if ( !SafeWidget.IsValid() )
 	{
 		return;
