UIComponent: Support MVVM bindings to UIComponent Navigation Delegates, not FNames

#rnx

diff --git a/Engine/Source/Runtime/UMG/Private/Extensions/UIComponents/NavigationUIComponent.cpp b/Engine/Source/Runtime/UMG/Private/Extensions/UIComponents/NavigationUIComponent.cpp
--- a/Engine/Source/Runtime/UMG/Private/Extensions/UIComponents/NavigationUIComponent.cpp
+++ b/Engine/Source/Runtime/UMG/Private/Extensions/UIComponents/NavigationUIComponent.cpp
@@ -86,8 +86,13 @@ void UNavigationUIComponent::OnConstruct()
 
 	// Bind user provided functions to component delegates
 
-	OnNavigationEnteredDelegate.BindUFunction(OwningUserWidget, OnNavigationEntered);
-	OnNavigationExitedDelegate.BindUFunction(OwningUserWidget, OnNavigationExited);
+	FScriptDelegate EnteredDelegate;
+	EnteredDelegate.BindUFunction(OwningUserWidget, OnNavigationEntered);
+	OnNavigationEnteredDelegate.AddUnique(EnteredDelegate);
+
+	FScriptDelegate ExitedDelegate;
+	ExitedDelegate.BindUFunction(OwningUserWidget, OnNavigationExited);
+	OnNavigationExitedDelegate.AddUnique(ExitedDelegate);
 
 	TSharedPtr<SWidget> Content = OwningUserWidget->GetCachedWidget();
 	if (!Content.IsValid())
@@ -130,12 +135,12 @@ void UNavigationUIComponent::HandleNavigationTransition(const FNavigationTransit
 	{
 		case ENavigationTransitionDirection::Incoming:
 		{
-			OnNavigationEnteredDelegate.ExecuteIfBound(NavigationTransition.Type, OldUMGWidget, NewUMGWidget);
+			OnNavigationEnteredDelegate.Broadcast(NavigationTransition.Type, OldUMGWidget, NewUMGWidget);
 			break;
 		}
 		case ENavigationTransitionDirection::Outgoing:
 		{
-			OnNavigationExitedDelegate.ExecuteIfBound(NavigationTransition.Type, OldUMGWidget, NewUMGWidget);
+			OnNavigationExitedDelegate.Broadcast(NavigationTransition.Type, OldUMGWidget, NewUMGWidget);
 			break;
 		}
 	}
diff --git a/Engine/Source/Runtime/UMG/Public/Extensions/UIComponents/NavigationUIComponent.h b/Engine/Source/Runtime/UMG/Public/Extensions/UIComponents/NavigationUIComponent.h
--- a/Engine/Source/Runtime/UMG/Public/Extensions/UIComponents/NavigationUIComponent.h
+++ b/Engine/Source/Runtime/UMG/Public/Extensions/UIComponents/NavigationUIComponent.h
@@ -15,7 +15,7 @@ class UWidget;
 * If Entered, Other Widget is the Widget Navigated From 
 * If Exited , Other Widget is the Widget Navigated To
 */
-DECLARE_DYNAMIC_DELEGATE_ThreeParams(FOnNavigationTransition, EUINavigation, Type, UWidget*, OldFocusedWidget, UWidget*, NewFocusedWidget);
+DECLARE_DYNAMIC_MULTICAST_DELEGATE_ThreeParams(FOnNavigationTransition, EUINavigation, Type, UWidget*, OldFocusedWidget, UWidget*, NewFocusedWidget);
 
 UCLASS(BlueprintType, Experimental)
 class UNavigationUIComponent : public UUIComponent
@@ -28,22 +28,25 @@ private:
 	virtual void OnGraphRenamed(UEdGraph* Graph, const FName& OldName, const FName& NewName) override;
 
 	void HandleNavigationTransition(const FNavigationTransition& NavigationTransition);
-public:
-	/* FNames are references to user provided functions.  */
 
-	UPROPERTY(EditAnywhere, BlueprintReadWrite, Category = Navigation)
+public:
+	/** 
+	* FNames are references to user provided functions. 
+	* No direct access to properties, modifiable in DetailsView via FNavigationUIComponentCustomization::CustomizeNavigationTransitionFunction
+	*/
+	UPROPERTY()
 	FName OnNavigationEntered; 
 
-	UPROPERTY(EditAnywhere, BlueprintReadWrite, Category = Navigation)
+	UPROPERTY()
 	FName OnNavigationExited;
 
 private:
+
 	/* Delegates are the concrete functions from the provided user provided function names */
-	
-	UPROPERTY()
-	FOnNavigationTransition OnNavigationEnteredDelegate; 
+	UPROPERTY(BlueprintAssignable, Category = "Navigation", meta = (AllowPrivateAccess = true, DisplayName = "On Navigation Entered"))
+	FOnNavigationTransition OnNavigationEnteredDelegate;
 
-	UPROPERTY()
+	UPROPERTY(BlueprintAssignable, Category = "Navigation", meta = (AllowPrivateAccess = true, DisplayName = "On Navigation Exited"))
 	FOnNavigationTransition OnNavigationExitedDelegate;
 
 	TSharedPtr<FNavigationTransitionMetadata> NavigationTransitionMetadata;
