Fix to componennts not being cleaned propertly when a Widget is deleted and the Widget Blueprint is saved before compiling.

diff --git a/Engine/Source/Editor/UMGEditor/Private/UIComponentWidgetBlueprintExtension.cpp b/Engine/Source/Editor/UMGEditor/Private/UIComponentWidgetBlueprintExtension.cpp
--- a/Engine/Source/Editor/UMGEditor/Private/UIComponentWidgetBlueprintExtension.cpp
+++ b/Engine/Source/Editor/UMGEditor/Private/UIComponentWidgetBlueprintExtension.cpp
@@ -86,9 +86,11 @@ void UUIComponentWidgetBlueprintExtension::RenameWidget(const FName& OldVarName,
 
 bool UUIComponentWidgetBlueprintExtension::VerifyContainer(UUserWidget* UserWidget) const
 {
-	const UUIComponentUserWidgetExtension* UserWidgetExtension = UserWidget->GetExtension<UUIComponentUserWidgetExtension>();
+	UUIComponentUserWidgetExtension* UserWidgetExtension = UserWidget->GetExtension<UUIComponentUserWidgetExtension>();
 	if (UserWidgetExtension)
 	{
+		UserWidgetExtension->CleanupComponents();
+		
 		TArray<FUIComponentTarget> Components;
 		ComponentContainer->ForEachComponentTarget([&Components](const FUIComponentTarget& ComponentTarget){
 			Components.Push(ComponentTarget);
diff --git a/Engine/Source/Runtime/UMG/Private/Extensions/UIComponentUserWidgetExtension.cpp b/Engine/Source/Runtime/UMG/Private/Extensions/UIComponentUserWidgetExtension.cpp
--- a/Engine/Source/Runtime/UMG/Private/Extensions/UIComponentUserWidgetExtension.cpp
+++ b/Engine/Source/Runtime/UMG/Private/Extensions/UIComponentUserWidgetExtension.cpp
@@ -4,6 +4,7 @@
 
 #include "Blueprint/UserWidget.h"
 #include "CoreMinimal.h"
+#include "Blueprint/WidgetTree.h"
 #include "Components/Widget.h"
 #include "Extensions/UIComponent.h"
 #include "Extensions/UIComponentContainer.h"
@@ -87,6 +88,7 @@ void UUIComponentUserWidgetExtension::InitializeComponents()
 
 void UUIComponentUserWidgetExtension::InitializeInEditor()
 {
+	CleanupComponents();
 	// We simply call the Initialize function as we need the same code to execute in Editor.
 	Initialize();
 }
@@ -130,6 +132,14 @@ void UUIComponentUserWidgetExtension::RemoveComponent(const UClass* ComponentCla
 	}
 }
 
+void UUIComponentUserWidgetExtension::CleanupComponents()
+{
+	const UUserWidget* UserWidget = GetUserWidget();
+	if (UserWidget && UserWidget->WidgetTree && ComponentContainer)
+	{
+		ComponentContainer->CleanupUIComponents(UserWidget->WidgetTree);
+	}
+}
 
 void UUIComponentUserWidgetExtension::MoveComponent(FName OwnerName, const UClass* ComponentClass, const UClass* RelativeToComponentClass, bool bMoveAfter)
 {
diff --git a/Engine/Source/Runtime/UMG/Public/Extensions/UIComponentUserWidgetExtension.h b/Engine/Source/Runtime/UMG/Public/Extensions/UIComponentUserWidgetExtension.h
--- a/Engine/Source/Runtime/UMG/Public/Extensions/UIComponentUserWidgetExtension.h
+++ b/Engine/Source/Runtime/UMG/Public/Extensions/UIComponentUserWidgetExtension.h
@@ -44,6 +44,8 @@ public:
 
 	UE_API void RemoveComponent(const UClass* ComponentClass, FName OwnerName);
 	
+	UE_API void CleanupComponents();
+	
 	UE_API void MoveComponent(FName OwnerName, const UClass* ComponentClass, const UClass* RelativeToComponentClass, bool bMoveAfter);
 
 #endif // WITH_EDITOR
