Fix to Widget with a Component that might be not selectable in the UMG Designer.

diff --git a/Engine/Source/Runtime/UMG/Private/Components/Widget.cpp b/Engine/Source/Runtime/UMG/Private/Components/Widget.cpp
--- a/Engine/Source/Runtime/UMG/Private/Components/Widget.cpp
+++ b/Engine/Source/Runtime/UMG/Private/Components/Widget.cpp
@@ -1054,6 +1054,12 @@ TSharedRef<SWidget> UWidget::TakeWidget_Private(ConstructMethodType ConstructMet
 						// If the Widget is the same of null, we do nothing 
 						if (SafeComponentWidget != PublicWidget)
 						{
+#if WITH_EDITOR
+							if (IsDesignTime())
+							{
+								SafeComponentWidget->SetVisibility(EVisibility::Visible);
+							}
+#endif						
 							ComponentWrapperWidget = SafeComponentWidget;
 							PublicWidget = SafeComponentWidget;
 						}
