Corection to UIComponentUserWidgetExtention intialization to better support the Canvas usage in UMG Designer.

diff --git a/Engine/Source/Runtime/UMG/Private/Extensions/UIComponentUserWidgetExtension.cpp b/Engine/Source/Runtime/UMG/Private/Extensions/UIComponentUserWidgetExtension.cpp
--- a/Engine/Source/Runtime/UMG/Private/Extensions/UIComponentUserWidgetExtension.cpp
+++ b/Engine/Source/Runtime/UMG/Private/Extensions/UIComponentUserWidgetExtension.cpp
@@ -74,6 +74,11 @@ void UUIComponentUserWidgetExtension::InitializeContainer(UUIComponentContainer*
 	ComponentContainer = InComponentContainer;
 }
 
+bool UUIComponentUserWidgetExtension::IsContainerInitialized() const
+{
+	return ComponentContainer != nullptr;
+}
+
 void UUIComponentUserWidgetExtension::InitializeComponents()
 {
 	if (ComponentContainer)
diff --git a/Engine/Source/Runtime/UMG/Private/Extensions/UIComponentWidgetBlueprintGeneratedClassExtension.cpp b/Engine/Source/Runtime/UMG/Private/Extensions/UIComponentWidgetBlueprintGeneratedClassExtension.cpp
--- a/Engine/Source/Runtime/UMG/Private/Extensions/UIComponentWidgetBlueprintGeneratedClassExtension.cpp
+++ b/Engine/Source/Runtime/UMG/Private/Extensions/UIComponentWidgetBlueprintGeneratedClassExtension.cpp
@@ -97,6 +97,12 @@ void UUIComponentWidgetBlueprintGeneratedClassExtension::Initialize(UUserWidget*
 	if (!UserWidgetExtension)
 	{
 		UserWidgetExtension = UserWidget->AddExtension<UUIComponentUserWidgetExtension>();
+	}
+    // Since some Widget can be duplicated manually, we could end up not Adding the Extension but still need to Initialize it.
+    // When we are in UMG Designer preview widget, it will be initialized by the UUIComponentWidgetBlueprintExtension, in that case, no
+	// need to Initialize it.
+	if (!UserWidgetExtension->IsContainerInitialized())
+	{
 		UserWidgetExtension->InitializeContainer(DuplicateContainer(UserWidget));
 	}
 }
diff --git a/Engine/Source/Runtime/UMG/Public/Extensions/UIComponentUserWidgetExtension.h b/Engine/Source/Runtime/UMG/Public/Extensions/UIComponentUserWidgetExtension.h
--- a/Engine/Source/Runtime/UMG/Public/Extensions/UIComponentUserWidgetExtension.h
+++ b/Engine/Source/Runtime/UMG/Public/Extensions/UIComponentUserWidgetExtension.h
@@ -32,6 +32,7 @@ public:
 	UE_API UUIComponent* GetComponent(const UClass* ComponentClass, FName OwnerName) const;
 
 	UE_API void InitializeContainer(UUIComponentContainer* InComponentContainer);
+	UE_API bool IsContainerInitialized() const;
 	
 
 #if WITH_EDITOR
