UMG: fix how we consume the current animation time on play

A previous change addressed the problem of events fired on the very first frame: as we evaluate the animation with [0, 0], we need to "advance" the current time so that the next update is (0, NextTick) instead of [0, NextTick]. Otherwise, whatever happens on the first frame (such as an event) would get evaluated again.

The problem however is that this previous change "played" through the first tick. This creates a problem if the animation has zero-length (either naturally, or because its end time is passed as its start time when playing) In that case, the status is immediately set to "Stopped", since we advanced passed the end time, and we never get to handling this properly (such as firing the OnAnimationFinished callback)

So this change fixes that by advancing the time by exactly one tick bound.

diff --git a/Engine/Source/Runtime/MovieScene/Private/Evaluation/MovieScenePlaybackManager.cpp b/Engine/Source/Runtime/MovieScene/Private/Evaluation/MovieScenePlaybackManager.cpp
--- a/Engine/Source/Runtime/MovieScene/Private/Evaluation/MovieScenePlaybackManager.cpp
+++ b/Engine/Source/Runtime/MovieScene/Private/Evaluation/MovieScenePlaybackManager.cpp
@@ -356,6 +356,12 @@ FFrameNumber FMovieScenePlaybackManager::GetLastValidTick() const
 	return SequenceEndTick - EndOffsetTicks - 1;  // minus one tick for exclusive end frame.
 }
 
+FMovieSceneContext FMovieScenePlaybackManager::UpdateToNextTick()
+{
+	const FFrameTime CurrentTick = PlaybackPosition.GetCurrentPosition();
+	return FMovieSceneContext(PlaybackPosition.PlayTo(CurrentTick, PlayDirection), PlaybackStatus);
+}
+
 FMovieSceneContext FMovieScenePlaybackManager::UpdateAtCurrentTime() const
 {
 	return FMovieSceneContext(PlaybackPosition.GetCurrentPositionAsRange(), PlaybackStatus);
diff --git a/Engine/Source/Runtime/MovieScene/Public/Evaluation/MovieScenePlaybackManager.h b/Engine/Source/Runtime/MovieScene/Public/Evaluation/MovieScenePlaybackManager.h
--- a/Engine/Source/Runtime/MovieScene/Public/Evaluation/MovieScenePlaybackManager.h
+++ b/Engine/Source/Runtime/MovieScene/Public/Evaluation/MovieScenePlaybackManager.h
@@ -85,6 +85,14 @@ public:
 	 */
 	MOVIESCENE_API void UpdateTo(const FFrameTime NextTime, FContexts& OutContexts);
 
+	/**
+	 * Updates the playback state over only the current tick. This generally does not advance the
+	 * current time by one tick, but by one *tick bound*. That is, if the current time is tick 0,
+	 * this will evaluate the [0, 0] range, advancing to an exclusive lower bound, so that the
+	 * next update will be (0, XYZ).
+	 */
+	MOVIESCENE_API FMovieSceneContext UpdateToNextTick();
+
 	/**
 	 * Returns an evaluation context for the current time, i.e. using a zero-width evaluation
 	 * range set around the current time.
diff --git a/Engine/Source/Runtime/UMG/Private/Animation/WidgetAnimationState.cpp b/Engine/Source/Runtime/UMG/Private/Animation/WidgetAnimationState.cpp
--- a/Engine/Source/Runtime/UMG/Private/Animation/WidgetAnimationState.cpp
+++ b/Engine/Source/Runtime/UMG/Private/Animation/WidgetAnimationState.cpp
@@ -534,23 +534,15 @@ void FWidgetAnimationState::Play(const FWidgetAnimationStatePlayParams& PlayPara
 	TSharedPtr<FMovieSceneEntitySystemRunner> Runner = SharedPlaybackState->GetRunner();
 	if (Runner)
 	{
-		// Generate a full context for this start using Update rather than UpdateAtCurrentTime.
-		// This is so we update the 'previous' time to the start time so we don't evaluate it again during our first tick
-		FMovieScenePlaybackManager::FContexts StartContexts;
-		PlaybackManager.Update(0, StartContexts);
-
-		// We should just have a single context
-		ensure(StartContexts.Num() == 1);
-		if (StartContexts.Num() > 0)
-		{
-			FMovieSceneContext& FirstContext(StartContexts[0]);
+		// Use UpdateToNextTick here so that we don't re-evaluate the first tick a second time when
+		// we start calling Tick().
+		const FMovieSceneContext FirstContext = PlaybackManager.UpdateToNextTick();
 
-			Runner->QueueUpdate(
-				FirstContext,
-				SharedPlaybackState->GetRootInstanceHandle(),
-				FSimpleDelegate::CreateSP(this, &FWidgetAnimationState::OnBegunPlay),
-				ERunnerUpdateFlags::Flush);
-		}
+		Runner->QueueUpdate(
+			FirstContext,
+			SharedPlaybackState->GetRootInstanceHandle(),
+			FSimpleDelegate::CreateSP(this, &FWidgetAnimationState::OnBegunPlay),
+			ERunnerUpdateFlags::Flush);
 	}
 
 #if !NO_LOGGING
