UMG Viewport: Make sure we EndPlay on the preview world during teardown.

diff --git a/Engine/Source/Runtime/UMG/Private/Components/Viewport.cpp b/Engine/Source/Runtime/UMG/Private/Components/Viewport.cpp
--- a/Engine/Source/Runtime/UMG/Private/Components/Viewport.cpp
+++ b/Engine/Source/Runtime/UMG/Private/Components/Viewport.cpp
@@ -107,6 +107,10 @@ FUMGViewportClient::FUMGViewportClient(FPreviewScene* InPreviewScene)
 
 FUMGViewportClient::~FUMGViewportClient()
 {
+	if (PreviewScene->GetWorld())
+	{
+		PreviewScene->GetWorld()->EndPlay(EEndPlayReason::RemovedFromWorld);
+	}
 }
 
 void FUMGViewportClient::Tick(float InDeltaTime)
@@ -346,12 +350,12 @@ class SAutoRefreshViewport : public SViewport
 	}
 
 public:
+	/** preview scene */
+	FPreviewScene PreviewScene;
+
 	TSharedPtr<FUMGViewportClient> ViewportClient;
 	
 	TSharedPtr<FSceneViewport> Viewport;
-
-	/** preview scene */
-	FPreviewScene PreviewScene;
 };
 
 /////////////////////////////////////////////////////
