Ensure we have an FSlateUser before getting the current focused widget. In a VR editor session, the user (VR) could be gone by the time we reset the focus widget.

#rnx

diff --git a/Engine/Source/Runtime/UMG/Private/Components/WidgetComponent.cpp b/Engine/Source/Runtime/UMG/Private/Components/WidgetComponent.cpp
--- a/Engine/Source/Runtime/UMG/Private/Components/WidgetComponent.cpp
+++ b/Engine/Source/Runtime/UMG/Private/Components/WidgetComponent.cpp
@@ -1523,7 +1523,13 @@ void UWidgetComponent::ResetWorldUserFocusOnWidget()
 	// on our widget to release any focus paths/strongs refs they might have to our widget.
 	for(int32 UserIndex : UsersEnteredWorldWidget)
 	{
-		TSharedPtr<SWidget> CurrentWidgetInPath = FSlateApplication::Get().GetUser(UserIndex)->GetFocusedWidget();
+		TSharedPtr<FSlateUser> SlateUser = FSlateApplication::Get().GetUser(UserIndex);
+		if (!SlateUser)
+		{
+			continue;
+		}
+
+		TSharedPtr<SWidget> CurrentWidgetInPath = SlateUser->GetFocusedWidget();
 		while(CurrentWidgetInPath != nullptr)
 		{
 			if(CurrentWidgetInPath == SlateWindow)
