Exposing text selection/cursor movement through a new Slate event

#FailedValidations Preflight

diff --git a/Engine/Source/Runtime/Slate/Private/Widgets/Input/SEditableText.cpp b/Engine/Source/Runtime/Slate/Private/Widgets/Input/SEditableText.cpp
--- a/Engine/Source/Runtime/Slate/Private/Widgets/Input/SEditableText.cpp
+++ b/Engine/Source/Runtime/Slate/Private/Widgets/Input/SEditableText.cpp
@@ -38,6 +38,7 @@ void SEditableText::Construct( const FArguments& InArgs )
 	OnIsTypedCharValid = InArgs._OnIsTypedCharValid;
 	OnTextChangedCallback = InArgs._OnTextChanged;
 	OnTextCommittedCallback = InArgs._OnTextCommitted;
+	OnCursorMovedWithSelectionCallback = InArgs._OnCursorMovedWithSelection;
 	MinDesiredWidth = InArgs._MinDesiredWidth;
 	bSelectAllTextOnCommit = InArgs._SelectAllTextOnCommit;
 	bSelectWordOnMouseDoubleClick = InArgs._SelectWordOnMouseDoubleClick;
@@ -224,7 +225,7 @@ FReply SEditableText::OnKeyChar( const FGeometry& MyGeometry, const FCharacterEv
 	{
 		Reply = OnKeyCharHandler.Execute(MyGeometry, InCharacterEvent);
 	}
-
+	
 	if (!Reply.IsEventHandled())
 	{
 		Reply = EditableTextLayout->HandleKeyChar(InCharacterEvent);
@@ -676,6 +677,9 @@ void SEditableText::OnTextCommitted(const FText& InText, const ETextCommit::Type
 
 void SEditableText::OnCursorMoved(const FTextLocation& InLocation)
 {
+	FTextSelection TextSelection = GetSelection();
+	
+	OnCursorMovedWithSelectionCallback.ExecuteIfBound(InLocation, TextSelection);
 	Invalidate(EInvalidateWidgetReason::Layout);
 }
 
diff --git a/Engine/Source/Runtime/Slate/Private/Widgets/Input/SEditableTextBox.cpp b/Engine/Source/Runtime/Slate/Private/Widgets/Input/SEditableTextBox.cpp
--- a/Engine/Source/Runtime/Slate/Private/Widgets/Input/SEditableTextBox.cpp
+++ b/Engine/Source/Runtime/Slate/Private/Widgets/Input/SEditableTextBox.cpp
@@ -40,6 +40,7 @@ void SEditableTextBox::Construct( const FArguments& InArgs )
 	OnTextChanged = InArgs._OnTextChanged;
 	OnVerifyTextChanged = InArgs._OnVerifyTextChanged;
 	OnTextCommitted = InArgs._OnTextCommitted;
+	OnCursorMovedWithSelection = InArgs._OnCursorMovedWithSelection;
 
 	SBorder::Construct( SBorder::FArguments()
 		.BorderImage( this, &SEditableTextBox::GetBorderImage )
@@ -75,6 +76,7 @@ void SEditableTextBox::Construct( const FArguments& InArgs )
 					.ContextMenuExtender( InArgs._ContextMenuExtender )
 					.OnTextChanged(this, &SEditableTextBox::OnEditableTextChanged)
 					.OnTextCommitted(this, &SEditableTextBox::OnEditableTextCommitted)
+					.OnCursorMovedWithSelection(this, &SEditableTextBox::OnEditableTextCursorMovedWithSelection)
 					.MinDesiredWidth( InArgs._MinDesiredWidth )
 					.SelectAllTextOnCommit( InArgs._SelectAllTextOnCommit )
 					.SelectWordOnMouseDoubleClick(InArgs._SelectWordOnMouseDoubleClick)
@@ -543,3 +545,8 @@ void SEditableTextBox::OnEditableTextCommitted(const FText& InText, ETextCommit:
 
 	OnTextCommitted.ExecuteIfBound(InText, InCommitType);
 }
+
+void SEditableTextBox::OnEditableTextCursorMovedWithSelection(const FTextLocation& InLocation, const FTextSelection& TextSelection)
+{
+	OnCursorMovedWithSelection.ExecuteIfBound(InLocation, TextSelection);
+}
diff --git a/Engine/Source/Runtime/Slate/Public/Widgets/Input/SEditableText.h b/Engine/Source/Runtime/Slate/Public/Widgets/Input/SEditableText.h
--- a/Engine/Source/Runtime/Slate/Public/Widgets/Input/SEditableText.h
+++ b/Engine/Source/Runtime/Slate/Public/Widgets/Input/SEditableText.h
@@ -31,6 +31,8 @@ struct FTextLocation;
 enum class ETextFlowDirection : uint8;
 enum class ETextShapingMethod : uint8;
 
+DECLARE_DELEGATE_TwoParams(FOnCursorMovedWithSelection, const FTextLocation&, const FTextSelection&);
+
 /**
  * Editable text widget
  */
@@ -142,6 +144,9 @@ public:
 
 		/** Called whenever the text is committed.  This happens when the user presses enter or the text box loses focus. */
 		SLATE_EVENT( FOnTextCommitted, OnTextCommitted )
+	
+		/** Called whenever the user's selection or cursor position of the text in the text field has changed. */
+		SLATE_EVENT( FOnCursorMovedWithSelection, OnCursorMovedWithSelection )
 
 		/** Minimum width that a text block should be */
 		SLATE_ATTRIBUTE( float, MinDesiredWidth )
@@ -519,6 +524,9 @@ protected:
 
 	/** Called whenever the text is committed.  This happens when the user presses enter or the text box loses focus. */
 	FOnTextCommitted OnTextCommittedCallback;
+	
+	/** Called whenever the user's selection and/or cursor position of the text has changed. */
+	FOnCursorMovedWithSelection OnCursorMovedWithSelectionCallback;
 
 	/** Prevents the editable text from being smaller than desired in certain cases (e.g. when it is empty) */
 	TAttribute<float> MinDesiredWidth;
diff --git a/Engine/Source/Runtime/Slate/Public/Widgets/Input/SEditableTextBox.h b/Engine/Source/Runtime/Slate/Public/Widgets/Input/SEditableTextBox.h
--- a/Engine/Source/Runtime/Slate/Public/Widgets/Input/SEditableTextBox.h
+++ b/Engine/Source/Runtime/Slate/Public/Widgets/Input/SEditableTextBox.h
@@ -123,6 +123,9 @@ public:
 
 		/** Called whenever the text is committed.  This happens when the user presses enter or the text box loses focus. */
 		SLATE_EVENT( FOnTextCommitted, OnTextCommitted )
+	
+		/** Called whenever the user's selection or cursor position of the text in the text field has changed. */
+		SLATE_EVENT( FOnCursorMovedWithSelection, OnCursorMovedWithSelection )
 
 		/** Called whenever the text is changed programmatically or interactively by the user */
 		SLATE_EVENT( FOnVerifyTextChanged, OnVerifyTextChanged )
@@ -426,6 +429,9 @@ protected:
 
 	/** Callback when the editable text is committed. */
 	SLATE_API void OnEditableTextCommitted(const FText& InText, ETextCommit::Type InCommitType);
+	
+	/** Callback when the user's cursor position and/or selection has changed. */
+	SLATE_API void OnEditableTextCursorMovedWithSelection(const FTextLocation& InLocation, const FTextSelection& TextSelection);
 
 	const FEditableTextBoxStyle* Style;
 
@@ -470,6 +476,9 @@ protected:
 
 	/** Callback to verify text when changed. Will return an error message to denote problems. */
 	FOnVerifyTextChanged OnVerifyTextChanged;
+	
+	/** Callback when the user's selection and/or cursor position of the text has changed. */
+	FOnCursorMovedWithSelection OnCursorMovedWithSelection;
 
 private:
 
diff --git a/Engine/Source/Runtime/UMG/Private/Components/EditableTextBox.cpp b/Engine/Source/Runtime/UMG/Private/Components/EditableTextBox.cpp
--- a/Engine/Source/Runtime/UMG/Private/Components/EditableTextBox.cpp
+++ b/Engine/Source/Runtime/UMG/Private/Components/EditableTextBox.cpp
@@ -104,6 +104,7 @@ TSharedRef<SWidget> UEditableTextBox::RebuildWidget()
 		.AllowContextMenu(AllowContextMenu)
 		.OnTextChanged(BIND_UOBJECT_DELEGATE(FOnTextChanged, HandleOnTextChanged))
 		.OnTextCommitted(BIND_UOBJECT_DELEGATE(FOnTextCommitted, HandleOnTextCommitted))
+		.OnCursorMovedWithSelection(BIND_UOBJECT_DELEGATE(FOnCursorMovedWithSelection, HandleOnCursorMovedWithSelection))
 		.VirtualKeyboardType(EVirtualKeyboardType::AsKeyboardType(KeyboardType.GetValue()))
 		.VirtualKeyboardOptions(VirtualKeyboardOptions)
 		.VirtualKeyboardTrigger(VirtualKeyboardTrigger)
@@ -397,6 +398,11 @@ void UEditableTextBox::HandleOnTextChanged(const FText& InText)
 	}
 }
 
+void UEditableTextBox::HandleOnCursorMovedWithSelection(const FTextLocation& InLocation, const FTextSelection& Selection)
+{
+	OnCursorMovedWithSelectionEvent.Broadcast(InLocation, Selection);
+}
+
 void UEditableTextBox::HandleOnTextCommitted(const FText& InText, ETextCommit::Type CommitMethod)
 {
 	if (SetTextInternal(InText))
diff --git a/Engine/Source/Runtime/UMG/Public/Components/EditableTextBox.h b/Engine/Source/Runtime/UMG/Public/Components/EditableTextBox.h
--- a/Engine/Source/Runtime/UMG/Public/Components/EditableTextBox.h
+++ b/Engine/Source/Runtime/UMG/Public/Components/EditableTextBox.h
@@ -11,6 +11,7 @@
 #include "Components/Widget.h"
 #include "Components/TextWidgetTypes.h"
 #include "Widgets/Text/ISlateEditableTextWidget.h"
+#include "Framework/Text/TextLayout.h"
 #include "EditableTextBox.generated.h"
 
 class SEditableTextBox;
@@ -31,6 +32,8 @@ public:
 
 	DECLARE_DYNAMIC_MULTICAST_DELEGATE_OneParam(FOnEditableTextBoxChangedEvent, const FText&, Text);
 	DECLARE_DYNAMIC_MULTICAST_DELEGATE_TwoParams(FOnEditableTextBoxCommittedEvent, const FText&, Text, ETextCommit::Type, CommitMethod);
+	
+	DECLARE_MULTICAST_DELEGATE_TwoParams(FOnEditableTextBoxCursorMovedWithSelectionEvent, const FTextLocation&, const FTextSelection&);
 
 public:
 	UE_DEPRECATED(5.1, "Direct access to Text is deprecated. Please use the getter or setter.")
@@ -142,6 +145,9 @@ public:
 	UPROPERTY(BlueprintAssignable, Category="TextBox|Event")
 	FOnEditableTextBoxCommittedEvent OnTextCommitted;
 
+	/** Called whenever the user's selection and/or cursor position of the text has changed. */
+	FOnEditableTextBoxCursorMovedWithSelectionEvent OnCursorMovedWithSelectionEvent;
+	
 	/** Provide a alternative mechanism for error reporting. */
 	//SLATE_ARGUMENT(TSharedPtr<class IErrorReportingWidget>, ErrorReporting)
 
@@ -276,6 +282,7 @@ protected:
 
 	UMG_API virtual void HandleOnTextChanged(const FText& Text);
 	UMG_API virtual void HandleOnTextCommitted(const FText& Text, ETextCommit::Type CommitMethod);
+	UMG_API virtual void HandleOnCursorMovedWithSelection(const FTextLocation& InLocation, const FTextSelection& Selection);
 
 #if WITH_ACCESSIBILITY
 	UMG_API virtual TSharedPtr<SWidget> GetAccessibleWidget() const override;
