UIComponents: Navigation UI Component use BlueprintAssignable

#rnx

diff --git a/Engine/Source/Runtime/UMG/Public/Extensions/UIComponents/NavigationUIComponent.h b/Engine/Source/Runtime/UMG/Public/Extensions/UIComponents/NavigationUIComponent.h
--- a/Engine/Source/Runtime/UMG/Public/Extensions/UIComponents/NavigationUIComponent.h
+++ b/Engine/Source/Runtime/UMG/Public/Extensions/UIComponents/NavigationUIComponent.h
@@ -40,11 +40,13 @@ public:
 	UPROPERTY()
 	FName OnNavigationExited;
 
+private:
+
 	/* Delegates are the concrete functions from the provided user provided function names */
-	UPROPERTY(BlueprintReadOnly, Category = Navigation)
-	FOnNavigationTransition OnNavigationEnteredDelegate; 
+	UPROPERTY(BlueprintAssignable, Category = "Navigation", meta = (AllowPrivateAccess = true, DisplayName = "On Navigation Entered"))
+	FOnNavigationTransition OnNavigationEnteredDelegate;
 
-	UPROPERTY(BlueprintReadOnly, Category = Navigation)
+	UPROPERTY(BlueprintAssignable, Category = "Navigation", meta = (AllowPrivateAccess = true, DisplayName = "On Navigation Exited"))
 	FOnNavigationTransition OnNavigationExitedDelegate;
 
 	TSharedPtr<FNavigationTransitionMetadata> NavigationTransitionMetadata;
