LIstView: Trigger OnItemDragDetected BP delegate, regardless if native counterpart is bound. Remove unecessary IsBound check on multicast events.

#rnx

diff --git a/Engine/Source/Runtime/UMG/Private/Components/ListView.cpp b/Engine/Source/Runtime/UMG/Private/Components/ListView.cpp
--- a/Engine/Source/Runtime/UMG/Private/Components/ListView.cpp
+++ b/Engine/Source/Runtime/UMG/Private/Components/ListView.cpp
@@ -308,14 +308,10 @@ UDragDropOperation* UListView::HandleListEntryDragDetected(const FGeometry& MyGe
 {
 	if (bAllowDragging)
 	{
-		
 		if (const TObjectPtrWrapTypeOf<UObject*>* ListItem = ItemFromEntryWidget(EntryWidget))
 		{
-			if (OnItemDragDetected().IsBound())
-			{
-				OnItemDragDetected().Broadcast(*ListItem);
-				BP_OnItemDragDetected.Broadcast(*ListItem, MyGeometry);
-			}
+			OnItemDragDetected().Broadcast(*ListItem);
+			BP_OnItemDragDetected.Broadcast(*ListItem, MyGeometry);
 
 			bIsDragging = true;
 			BP_OnListViewDraggingStateChanged.Broadcast(bIsDragging);
@@ -355,10 +351,7 @@ FReply UListView::HandleListEntryAcceptDrop(const FDragDropEvent& DropEvent, EIt
 		if (ListItem)
 		{
 			BP_OnItemAcceptDrop.Broadcast(*ListItem, UDragDropOperation::ConvertSlateDropZoneToUMG(DropZone), NativeOp.Get()->GetOperation());
-			if (OnItemAcceptDrop().IsBound())
-			{
-				OnItemAcceptDrop().Broadcast(*ListItem);
-			}
+			OnItemAcceptDrop().Broadcast(*ListItem);
 		}
 		return FReply::Handled().EndDragDrop();
 	}
@@ -373,10 +366,7 @@ void UListView::HandleListEntryDragEnter(const FDragDropEvent& DragDropEvent, UU
 		if (const TObjectPtrWrapTypeOf<UObject*>* ListItem = ItemFromEntryWidget(EntryWidget))
 		{
 			BP_OnItemDragEnter.Broadcast(*ListItem, NativeOp->GetOperation());
-			if (OnItemDragEnter().IsBound())
-			{
-				OnItemDragEnter().Broadcast(*ListItem);
-			}
+			OnItemDragEnter().Broadcast(*ListItem);
 		}
 	}
 }
@@ -389,10 +379,7 @@ void UListView::HandleListEntryDragLeave(const FDragDropEvent& DragDropEvent, UU
 		if (const TObjectPtrWrapTypeOf<UObject*>* ListItem = ItemFromEntryWidget(EntryWidget))
 		{
 			BP_OnItemDragLeave.Broadcast(*ListItem, NativeOp->GetOperation());
-			if (OnItemDragLeave().IsBound())
-			{
-					OnItemDragLeave().Broadcast(*ListItem);
-			}
+			OnItemDragLeave().Broadcast(*ListItem);
 		}
 	}
 }
@@ -405,11 +392,7 @@ void UListView::HandleListEntryDragCancelled(const FDragDropEvent& DragDropEvent
 	TSharedPtr<FUMGDragDropOp> NativeOp = DragDropEvent.GetOperationAs<FUMGDragDropOp>();
 	if (NativeOp.IsValid())
 	{
-		if (OnItemDragCancelled().IsBound())
-		{
-			OnItemDragCancelled().Broadcast(DragDropEvent);
-		}
-
+		OnItemDragCancelled().Broadcast(DragDropEvent);
 		BP_OnItemDragCancelled.Broadcast(NativeOp->GetOperation());
 	}
 }
