Clamp the downscaled rects in FPostBufferUpdater::Draw_RenderThread to a minimum size of 1x1 to avoid an assert.

diff --git a/Engine/Source/Runtime/UMG/Private/Slate/SPostBufferUpdate.cpp b/Engine/Source/Runtime/UMG/Private/Slate/SPostBufferUpdate.cpp
--- a/Engine/Source/Runtime/UMG/Private/Slate/SPostBufferUpdate.cpp
+++ b/Engine/Source/Runtime/UMG/Private/Slate/SPostBufferUpdate.cpp
@@ -121,6 +121,7 @@ void FPostBufferUpdater::Draw_RenderThread(FRDGBuilder& GraphBuilder, const FDra
 		GraphBuilder.UseInternalAccessMode(Texture);
 	}
 
+	const FIntPoint ViewRectMinimumSize = FIntPoint(1, 1);
 	for (const FActivePostBuffer& ActivePostBuffer : ActivePostBuffers)
 	{
 		if (ActivePostBuffer.Texture)
@@ -144,10 +145,14 @@ void FPostBufferUpdater::Draw_RenderThread(FRDGBuilder& GraphBuilder, const FDra
 					const FVector2f InputOffset = InvDownscale * FVector2f(InputTexture.ViewRect.Min.X, InputTexture.ViewRect.Min.Y);
 					const FVector2f OutputOffset = InvDownscale * FVector2f(OutputTexture.ViewRect.Min.X, OutputTexture.ViewRect.Min.Y);
 
-					InputTexture.ViewRect = FIntRect(DownscalePt * WorldTopLeft.IntPoint(), DownscalePt * WorldBottomRight.IntPoint());
+					const FIntPoint InputTextureTopLeft = DownscalePt * WorldTopLeft.IntPoint();
+					const FIntPoint InputTextureBottomRight = DownscalePt * WorldBottomRight.IntPoint();
+					InputTexture.ViewRect = FIntRect(InputTextureTopLeft, InputTextureBottomRight.ComponentMax(InputTextureTopLeft + ViewRectMinimumSize));
 
 					// Subtract the input offset because the geometry world transform will have it implicitly added.
-					OutputTexture.ViewRect = FIntRect((WorldTopLeft - InputOffset + OutputOffset).IntPoint(), (WorldBottomRight - InputOffset + OutputOffset).IntPoint());
+					const FIntPoint OutputTextureTopLeft = (WorldTopLeft - InputOffset + OutputOffset).IntPoint();
+					const FIntPoint OutputTextureBottomRight = (WorldBottomRight - InputOffset + OutputOffset).IntPoint();
+					OutputTexture.ViewRect = FIntRect(OutputTextureTopLeft, OutputTextureBottomRight.ComponentMax(OutputTextureTopLeft + ViewRectMinimumSize));
 				}
 				ActivePostBuffer.Proxy->PostProcess_Renderthread(GraphBuilder, InputTexture, OutputTexture);
 			}
