Add setter for touch scrolling on the base list view

#rnx

diff --git a/Engine/Source/Runtime/UMG/Private/Components/ListViewBase.cpp b/Engine/Source/Runtime/UMG/Private/Components/ListViewBase.cpp
--- a/Engine/Source/Runtime/UMG/Private/Components/ListViewBase.cpp
+++ b/Engine/Source/Runtime/UMG/Private/Components/ListViewBase.cpp
@@ -119,6 +119,15 @@ UMG_API void UListViewBase::SetIsPointerScrollingEnabled(bool bInIsPointerScroll
 	}
 }
 
+UMG_API void UListViewBase::SetIsTouchScrollingEnabled(bool bInIsTouchScrollingEnabled)
+{
+	bEnableTouchScrolling = bInIsTouchScrollingEnabled;
+	if (MyTableViewBase)
+	{
+		MyTableViewBase->SetIsTouchScrollingEnabled(bInIsTouchScrollingEnabled);
+	}
+}
+
 UMG_API void UListViewBase::SetIsGamepadScrollingEnabled(bool bInIsGamepadScrollingEnabled)
 {
 	bIsGamepadScrollingEnabled = bInIsGamepadScrollingEnabled;
diff --git a/Engine/Source/Runtime/UMG/Public/Components/ListViewBase.h b/Engine/Source/Runtime/UMG/Public/Components/ListViewBase.h
--- a/Engine/Source/Runtime/UMG/Public/Components/ListViewBase.h
+++ b/Engine/Source/Runtime/UMG/Public/Components/ListViewBase.h
@@ -607,6 +607,10 @@ public:
 	UFUNCTION(BlueprintCallable, Category = ListViewBase)
 	UMG_API void SetIsPointerScrollingEnabled(bool bInIsPointerScrollingEnabled);
 
+	/** Enables/Disable touch scrolling **/
+	UFUNCTION(BlueprintCallable, Category = ListViewBase)
+	UMG_API void SetIsTouchScrollingEnabled(bool bInIsTouchScrollingEnabled);
+
 	/** Enable/Disable the ability of the list to scroll via gamepad. */
 	UFUNCTION(BlueprintCallable, Category = ListViewBase)
 	UMG_API void SetIsGamepadScrollingEnabled(bool bInIsGamepadScrollingEnabled);
