Exposing TouchStarted, TouchMoved and TouchEnd events from STableViewBase to be able to use them in Blueprints. Also, exposing the current Overscroll value.

#rnx

diff --git a/Engine/Source/Runtime/Slate/Private/Widgets/Views/STableViewBase.cpp b/Engine/Source/Runtime/Slate/Private/Widgets/Views/STableViewBase.cpp
--- a/Engine/Source/Runtime/Slate/Private/Widgets/Views/STableViewBase.cpp
+++ b/Engine/Source/Runtime/Slate/Private/Widgets/Views/STableViewBase.cpp
@@ -746,6 +746,8 @@ FReply STableViewBase::OnTouchStarted( const FGeometry& MyGeometry, const FPoint
 	// See OnPreviewMouseButtonDown()
 	//     if (MouseEvent.IsTouchEvent())
 
+	NotifyTouchStart();
+
 	return FReply::Unhandled();
 }
 
@@ -777,6 +779,8 @@ FReply STableViewBase::OnTouchMoved( const FGeometry& MyGeometry, const FPointer
 			{
 				ScrollBar->BeginScrolling();
 
+				NotifyTouchMove();
+
 				// The user has moved the list some amount; they are probably
 				// trying to scroll. From now on, the list assumes the user is scrolling
 				// until they lift their finger.
@@ -799,6 +803,8 @@ FReply STableViewBase::OnTouchEnded( const FGeometry& MyGeometry, const FPointer
 
 	ScrollBar->EndScrolling();
 
+	NotifyTouchEnd();
+
 	if (HasMouseCapture())
 	{
 		return FReply::Handled().ReleaseMouseCapture();
diff --git a/Engine/Source/Runtime/Slate/Public/Framework/SlateDelegates.h b/Engine/Source/Runtime/Slate/Public/Framework/SlateDelegates.h
--- a/Engine/Source/Runtime/Slate/Public/Framework/SlateDelegates.h
+++ b/Engine/Source/Runtime/Slate/Public/Framework/SlateDelegates.h
@@ -127,6 +127,15 @@ DECLARE_DELEGATE_OneParam(FOnVector4ValueChanged, FVector4&)
 /** Invoked when a list view has reached its target offset.  */
 DECLARE_DELEGATE(FOnFinishedScrolling);
 
+/** Invoked when the touch has started.  */
+DECLARE_DELEGATE(FOnTouchStart);
+
+/** Invoked when the touch has moved.  */
+DECLARE_DELEGATE(FOnTouchMove);
+
+/** Invoked when the touch has ended.  */
+DECLARE_DELEGATE(FOnTouchEnd);
+
 DECLARE_DELEGATE(FOnTableViewBadState);
 
 template< typename ArgumentType >
diff --git a/Engine/Source/Runtime/Slate/Public/Widgets/Views/SListView.h b/Engine/Source/Runtime/Slate/Public/Widgets/Views/SListView.h
--- a/Engine/Source/Runtime/Slate/Public/Widgets/Views/SListView.h
+++ b/Engine/Source/Runtime/Slate/Public/Widgets/Views/SListView.h
@@ -316,6 +316,12 @@ public:
 		
 		SLATE_EVENT( FOnFinishedScrolling, OnFinishedScrolling )
 
+		SLATE_EVENT( FOnTouchStart, OnTouchStart )
+
+		SLATE_EVENT( FOnTouchMove, OnTouchMove )
+
+		SLATE_EVENT( FOnTouchEnd, OnTouchEnd )
+
 		SLATE_ITEMS_SOURCE_ARGUMENT( ItemType, ListItemsSource )
 
 		SLATE_ATTRIBUTE_DEPRECATED( float, ItemHeight, 5.5, "The ItemHeight is only used for Tile. See ShouldArrangeAsTiles")
@@ -406,6 +412,9 @@ public:
 		this->OnRowReleased = InArgs._OnRowReleased;
 		this->OnItemScrolledIntoView = InArgs._OnItemScrolledIntoView;
 		this->OnFinishedScrolling = InArgs._OnFinishedScrolling;
+		this->OnTouchStart = InArgs._OnTouchStart;
+		this->OnTouchMove = InArgs._OnTouchMove;
+		this->OnTouchEnd = InArgs._OnTouchEnd;
 
 		this->SetItemsSource(InArgs.MakeListItemsSource(this->SharedThis(this)));
 		this->OnContextMenuOpening = InArgs._OnContextMenuOpening;
@@ -2350,6 +2359,21 @@ protected:
 		OnFinishedScrolling.ExecuteIfBound();
 	}
 
+	virtual void NotifyTouchStart() override
+	{
+		OnTouchStart.ExecuteIfBound();
+	}
+
+	virtual void NotifyTouchMove() override
+	{
+		OnTouchMove.ExecuteIfBound();
+	}
+
+	virtual void NotifyTouchEnd() override
+	{
+		OnTouchEnd.ExecuteIfBound();
+	}
+
 	virtual float ScrollBy( const FGeometry& MyGeometry, float ScrollByAmountInSlateUnits, EAllowOverscroll InAllowOverscroll ) override
 	{
 		auto DoubleFractional = [](double Value) -> double
@@ -2695,6 +2719,15 @@ protected:
 	/** Delegate to be invoked when TargetScrollOffset is reached at the end of a ::Tick. */
 	FOnFinishedScrolling OnFinishedScrolling;
 
+	/** Delegate to be invoked when the touch has start. */
+	FOnTouchStart OnTouchStart;
+
+	/** Delegate to be invoked when the user has scrolled the list some amount. */
+	FOnTouchMove OnTouchMove;
+
+	/** Delegate to be invoked when the user has lifted its finger after scrolling. */
+	FOnTouchEnd OnTouchEnd;
+
 	/** A set of selected data items */
 	TItemSet SelectedItems;
 
diff --git a/Engine/Source/Runtime/Slate/Public/Widgets/Views/STableViewBase.h b/Engine/Source/Runtime/Slate/Public/Widgets/Views/STableViewBase.h
--- a/Engine/Source/Runtime/Slate/Public/Widgets/Views/STableViewBase.h
+++ b/Engine/Source/Runtime/Slate/Public/Widgets/Views/STableViewBase.h
@@ -445,6 +445,15 @@ protected:
 	 */
 	virtual void NotifyFinishedScrolling() = 0;
 
+	/** Called when the user has started to touch. */
+	virtual void NotifyTouchStart() = 0;
+
+	/** Called when the user has scrolled the list some amount. */
+	virtual void NotifyTouchMove() = 0;
+
+	/** Called when the user has lifted its finger after scrolling. */
+	virtual void NotifyTouchEnd() = 0;
+
 	/** Util Function so templates classes don't need to include SlateApplication */
 	SLATE_API void NavigateToWidget(const uint32 UserIndex, const TSharedPtr<SWidget>& NavigationDestination, ENavigationSource NavigationSource = ENavigationSource::FocusedWidget, EUINavigation NavigationType = EUINavigation::Invalid) const;
 
@@ -641,6 +650,10 @@ protected:
 	float GetSmoothDeltaTime(float InDeltaTime);
 	void UpdateSmoothDeltaTime(float InDeltaTime);
 
+public:
+
+	float GetOverscroll() const { return Overscroll.GetOverscroll(GetTickSpaceGeometry()); }
+
 private:
 	/** When true, a refresh should occur the next tick */
 	bool bItemsNeedRefresh = false;
diff --git a/Engine/Source/Runtime/Slate/Public/Widgets/Views/STileView.h b/Engine/Source/Runtime/Slate/Public/Widgets/Views/STileView.h
--- a/Engine/Source/Runtime/Slate/Public/Widgets/Views/STileView.h
+++ b/Engine/Source/Runtime/Slate/Public/Widgets/Views/STileView.h
@@ -94,6 +94,12 @@ public:
 		
 		SLATE_EVENT( FOnFinishedScrolling, OnFinishedScrolling )
 
+		SLATE_EVENT( FOnTouchStart, OnTouchStart )
+
+		SLATE_EVENT( FOnTouchMove, OnTouchMove )
+
+		SLATE_EVENT( FOnTouchEnd, OnTouchEnd )
+
 		SLATE_EVENT( FOnItemScrolledIntoView, OnItemScrolledIntoView )
 
 		SLATE_ITEMS_SOURCE_ARGUMENT( ItemType, ListItemsSource )
diff --git a/Engine/Source/Runtime/Slate/Public/Widgets/Views/STreeView.h b/Engine/Source/Runtime/Slate/Public/Widgets/Views/STreeView.h
--- a/Engine/Source/Runtime/Slate/Public/Widgets/Views/STreeView.h
+++ b/Engine/Source/Runtime/Slate/Public/Widgets/Views/STreeView.h
@@ -179,6 +179,12 @@ public:
 
 		SLATE_EVENT( FOnFinishedScrolling, OnFinishedScrolling )
 
+		SLATE_EVENT( FOnTouchStart, OnTouchStart )
+
+		SLATE_EVENT( FOnTouchMove, OnTouchMove )
+
+		SLATE_EVENT( FOnTouchEnd, OnTouchEnd )
+
 		SLATE_EVENT( FOnItemScrolledIntoView, OnItemScrolledIntoView )
 
 		SLATE_EVENT( FOnGetChildren, OnGetChildren )
diff --git a/Engine/Source/Runtime/UMG/Private/Components/ListView.cpp b/Engine/Source/Runtime/UMG/Private/Components/ListView.cpp
--- a/Engine/Source/Runtime/UMG/Private/Components/ListView.cpp
+++ b/Engine/Source/Runtime/UMG/Private/Components/ListView.cpp
@@ -463,6 +463,24 @@ void UListView::OnListViewFinishedScrollingInternal()
 	BP_OnListViewFinishedScrolling.Broadcast();
 }
 
+void UListView::OnListViewTouchStartInternal()
+{
+	ITypedUMGListView::OnListViewTouchStartInternal();
+	BP_OnListViewTouchStart.Broadcast();
+}
+
+void UListView::OnListViewTouchMoveInternal()
+{
+	ITypedUMGListView::OnListViewTouchMoveInternal();
+	BP_OnListViewTouchMove.Broadcast();
+}
+
+void UListView::OnListViewTouchEndInternal()
+{
+	ITypedUMGListView::OnListViewTouchEndInternal();
+	BP_OnListViewTouchEnd.Broadcast();
+}
+
 bool UListView::OnIsSelectableOrNavigableInternal(UObject* FirstSelectedItem)
 {
 	return BP_OnIsItemSelectableOrNavigable.IsBound() ? BP_OnIsItemSelectableOrNavigable.Execute(FirstSelectedItem) : ITypedUMGListView::OnIsSelectableOrNavigableInternal(FirstSelectedItem);
diff --git a/Engine/Source/Runtime/UMG/Public/Components/ListView.h b/Engine/Source/Runtime/UMG/Public/Components/ListView.h
--- a/Engine/Source/Runtime/UMG/Public/Components/ListView.h
+++ b/Engine/Source/Runtime/UMG/Public/Components/ListView.h
@@ -17,6 +17,9 @@ DECLARE_DYNAMIC_MULTICAST_DELEGATE_TwoParams(FOnItemIsHoveredChangedDynamic, UOb
 DECLARE_DYNAMIC_MULTICAST_DELEGATE_TwoParams(FOnListItemScrolledIntoViewDynamic, UObject*, Item, UUserWidget*, Widget);
 DECLARE_DYNAMIC_MULTICAST_DELEGATE_TwoParams(FOnListViewScrolledDynamic, float, ItemOffset, float, DistanceRemaining);
 DECLARE_DYNAMIC_MULTICAST_DELEGATE(FOnListViewFinishedScrollingDynamic);
+DECLARE_DYNAMIC_MULTICAST_DELEGATE(FOnListViewTouchStartDynamic);
+DECLARE_DYNAMIC_MULTICAST_DELEGATE(FOnListViewTouchMoveDynamic);
+DECLARE_DYNAMIC_MULTICAST_DELEGATE(FOnListViewTouchEndDynamic);
 DECLARE_DYNAMIC_DELEGATE_RetVal_OneParam(bool, FOnIsItemSelectableOrNavigableDynamic, UObject*, Item);
 
 DECLARE_DYNAMIC_MULTICAST_DELEGATE_TwoParams(FOnItemDragDropMulticastDynamic, UObject*, Item, UDragDropOperation*, Operation);
@@ -178,6 +181,9 @@ protected:
 	UMG_API virtual void OnItemScrolledIntoViewInternal(UObject* Item, UUserWidget& EntryWidget) override;
 	UMG_API virtual void OnListViewScrolledInternal(float ItemOffset, float DistanceRemaining) override;
 	UMG_API virtual void OnListViewFinishedScrollingInternal() override;
+	UMG_API virtual void OnListViewTouchStartInternal() override;
+	UMG_API virtual void OnListViewTouchMoveInternal() override;
+	UMG_API virtual void OnListViewTouchEndInternal() override;
 	UMG_API virtual bool OnIsSelectableOrNavigableInternal(UObject* FirstSelectedItem) override;
 
 	UMG_API void HandleOnEntryInitializedInternal(UObject* Item, const TSharedRef<ITableRow>& TableRow);
@@ -397,7 +403,16 @@ private:
 	
 	UPROPERTY(BlueprintAssignable, Category = Events, meta = (DisplayName = "On List View Finished Scrolling"))
 	FOnListViewFinishedScrollingDynamic BP_OnListViewFinishedScrolling;	
+	
+	UPROPERTY(BlueprintAssignable, Category = Events, meta = (DisplayName = "On List View Touch Start"))
+	FOnListViewTouchStartDynamic BP_OnListViewTouchStart;
+
+	UPROPERTY(BlueprintAssignable, Category = Events, meta = (DisplayName = "On List View Touch Move"))
+	FOnListViewTouchMoveDynamic BP_OnListViewTouchMove;
 
+	UPROPERTY(BlueprintAssignable, Category = Events, meta = (DisplayName = "On List View Touch End"))
+	FOnListViewTouchEndDynamic BP_OnListViewTouchEnd;
+	
 	UPROPERTY(BlueprintReadWrite, Category = Events, meta = (AllowPrivateAccess, DisplayName = "On Is Item Selectable Or Navigable"))
 	FOnIsItemSelectableOrNavigableDynamic BP_OnIsItemSelectableOrNavigable;
 };
diff --git a/Engine/Source/Runtime/UMG/Public/Components/ListViewBase.h b/Engine/Source/Runtime/UMG/Public/Components/ListViewBase.h
--- a/Engine/Source/Runtime/UMG/Public/Components/ListViewBase.h
+++ b/Engine/Source/Runtime/UMG/Public/Components/ListViewBase.h
@@ -73,6 +73,15 @@ public:
 	DECLARE_MULTICAST_DELEGATE(FOnFinishedScrolling);
 	virtual FOnFinishedScrolling& OnFinishedScrolling() const = 0;
 
+	DECLARE_MULTICAST_DELEGATE(FOnTouchStart);
+	virtual FOnTouchStart& OnTouchStart() const = 0;
+
+	DECLARE_MULTICAST_DELEGATE(FOnTouchMove);
+	virtual FOnTouchMove& OnTouchMove() const = 0;
+
+	DECLARE_MULTICAST_DELEGATE(FOnTouchEnd);
+	virtual FOnTouchEnd& OnTouchEnd() const = 0;
+
 	DECLARE_MULTICAST_DELEGATE_TwoParams(FOnItemScrolledIntoView, ItemType, UUserWidget&);
 	virtual FOnItemScrolledIntoView& OnItemScrolledIntoView() const = 0;
 
@@ -294,6 +303,9 @@ protected:
 			.OnItemScrolledIntoView_UObject(Implementer, &UListViewBaseT::HandleItemScrolledIntoView)
 			.OnListViewScrolled_UObject(Implementer, &UListViewBaseT::HandleListViewScrolled)
 			.OnFinishedScrolling_UObject(Implementer, &UListViewBaseT::HandleFinishedScrolling)
+			.OnTouchStart_UObject(Implementer, &UListViewBaseT::HandleTouchStart)
+			.OnTouchMove_UObject(Implementer, &UListViewBaseT::HandleTouchMove)
+			.OnTouchEnd_UObject(Implementer, &UListViewBaseT::HandleTouchEnd)
 			.OnMouseButtonClick_UObject(Implementer, &UListViewBaseT::HandleItemClicked)
 			.OnMouseButtonDoubleClick_UObject(Implementer, &UListViewBaseT::HandleItemDoubleClicked);
 		ListView->BindToRefreshRow(TSlateDelegates< ItemType >::FOnRefreshRow::CreateUObject(Implementer, &UListViewBaseT::HandleRefreshRow));
@@ -338,6 +350,9 @@ protected:
 			.OnItemScrolledIntoView_UObject(Implementer, &UListViewBaseT::HandleItemScrolledIntoView)
 			.OnTileViewScrolled_UObject(Implementer, &UListViewBaseT::HandleListViewScrolled)
 			.OnFinishedScrolling_UObject(Implementer, &UListViewBaseT::HandleFinishedScrolling)
+			.OnTouchStart_UObject(Implementer, &UListViewBaseT::HandleTouchStart)
+			.OnTouchMove_UObject(Implementer, &UListViewBaseT::HandleTouchMove)
+			.OnTouchEnd_UObject(Implementer, &UListViewBaseT::HandleTouchEnd)
 			.OnMouseButtonClick_UObject(Implementer, &UListViewBaseT::HandleItemClicked)
 			.OnMouseButtonDoubleClick_UObject(Implementer, &UListViewBaseT::HandleItemDoubleClicked);
 		TileView->BindToRefreshRow(TSlateDelegates< ItemType >::FOnRefreshRow::CreateUObject(Implementer, &UListViewBaseT::HandleRefreshRow));
@@ -377,6 +392,9 @@ protected:
 			.OnItemScrolledIntoView_UObject(Implementer, &UListViewBaseT::HandleItemScrolledIntoView)
 			.OnTreeViewScrolled_UObject(Implementer, &UListViewBaseT::HandleListViewScrolled)
 			.OnFinishedScrolling_UObject(Implementer, &UListViewBaseT::HandleFinishedScrolling)
+			.OnTouchStart_UObject(Implementer, &UListViewBaseT::HandleTouchStart)
+			.OnTouchMove_UObject(Implementer, &UListViewBaseT::HandleTouchMove)
+			.OnTouchEnd_UObject(Implementer, &UListViewBaseT::HandleTouchEnd)
 			.OnMouseButtonClick_UObject(Implementer, &UListViewBaseT::HandleItemClicked)
 			.OnMouseButtonDoubleClick_UObject(Implementer, &UListViewBaseT::HandleItemDoubleClicked)
 			.OnGetChildren_UObject(Implementer, &UListViewBaseT::HandleGetChildren)
@@ -421,6 +439,9 @@ protected:
 	virtual void OnItemScrolledIntoViewInternal(ItemType Item, UUserWidget& EntryWidget) {}
 	virtual void OnListViewScrolledInternal(float ItemOffset, float DistanceRemaining) {}
 	virtual void OnListViewFinishedScrollingInternal() {}
+	virtual void OnListViewTouchStartInternal() {}
+	virtual void OnListViewTouchMoveInternal() {}
+	virtual void OnListViewTouchEndInternal() {}
 	virtual void OnItemExpansionChangedInternal(ItemType Item, bool bIsExpanded) {}
 
 private:
@@ -490,6 +511,24 @@ private:
 		OnFinishedScrolling().Broadcast();
 	}
 
+	void HandleTouchStart()
+	{
+		OnListViewTouchStartInternal();
+		OnTouchStart().Broadcast();
+	}
+
+	void HandleTouchMove()
+	{
+		OnListViewTouchMoveInternal();
+		OnTouchMove().Broadcast();
+	}
+
+	void HandleTouchEnd()
+	{
+		OnListViewTouchEndInternal();
+		OnTouchEnd().Broadcast();
+	}
+
 	void HandleItemScrolledIntoView(ItemType Item, const TSharedPtr<ITableRow>& InWidget)
 	{
 		UUserWidget* RowWidget = GetEntryWidgetFromItem(Item);
@@ -603,6 +642,9 @@ public:
 	UFUNCTION(BlueprintCallable, Category = ListViewBase)
 	UMG_API void SetAllowOverScroll(bool NewAllowOverscroll);
 
+	UFUNCTION(BlueprintCallable, Category = ListViewBase)
+	float GetOverscroll() const { return MyTableViewBase->GetOverscroll(); }
+
 	/** Enable/Disable the ability of the list to scroll. This should be use as a temporary disable. */
 	UFUNCTION(BlueprintCallable, Category = ListViewBase)
 	UMG_API void SetIsPointerScrollingEnabled(bool bInIsPointerScrollingEnabled);
@@ -913,6 +955,9 @@ private:	\
 	mutable FOnItemScrolledIntoView OnItemScrolledIntoViewEvent;	\
 	mutable FOnListViewScrolled OnListViewScrolledEvent;	\
 	mutable FOnFinishedScrolling OnFinishedScrollingEvent;	\
+	mutable FOnTouchStart OnTouchStartEvent;	\
+	mutable FOnTouchMove OnTouchMoveEvent;	\
+	mutable FOnTouchEnd OnTouchEndEvent;	\
 	mutable FOnItemExpansionChanged OnItemExpansionChangedEvent;	\
 	mutable FOnGetEntryClassForItem OnGetEntryClassForItemDelegate;	\
 	mutable FOnIsItemSelectableOrNavigable OnIsItemSelectableOrNavigableDelegate;	\
@@ -931,6 +976,9 @@ public:	\
 	virtual FOnItemScrolledIntoView& OnItemScrolledIntoView() const override { return OnItemScrolledIntoViewEvent; }	\
 	virtual FOnListViewScrolled& OnListViewScrolled() const override { return OnListViewScrolledEvent; }	\
 	virtual FOnFinishedScrolling& OnFinishedScrolling() const override { return OnFinishedScrollingEvent; }	\
+	virtual FOnTouchStart& OnTouchStart() const override { return OnTouchStartEvent; }	\
+	virtual FOnTouchMove& OnTouchMove() const override { return OnTouchMoveEvent; }	\
+	virtual FOnTouchEnd& OnTouchEnd() const override { return OnTouchEndEvent; }	\
 	virtual FOnItemExpansionChanged& OnItemExpansionChanged() const override { return OnItemExpansionChangedEvent; }	\
 	virtual FOnGetEntryClassForItem& OnGetEntryClassForItem() const override { return OnGetEntryClassForItemDelegate; }	\
 	virtual FOnIsItemSelectableOrNavigable& OnIsItemSelectableOrNavigable() const override { return OnIsItemSelectableOrNavigableDelegate; }
