UMG: Add GetScrollbarVisibility to ListView

#rnx

#Synced-CL 48251233

diff --git a/Engine/Source/Runtime/UMG/Private/Components/ListViewBase.cpp b/Engine/Source/Runtime/UMG/Private/Components/ListViewBase.cpp
--- a/Engine/Source/Runtime/UMG/Private/Components/ListViewBase.cpp
+++ b/Engine/Source/Runtime/UMG/Private/Components/ListViewBase.cpp
@@ -101,6 +101,16 @@ void UListViewBase::SetScrollbarVisibility(ESlateVisibility InVisibility)
 	}
 }
 
+ESlateVisibility UListViewBase::GetScrollbarVisibility() const
+{
+	if (MyTableViewBase)
+	{
+		return UWidget::ConvertRuntimeToSerializedVisibility(MyTableViewBase->GetScrollbarVisibility());
+	}
+	// Returning collapsed if TableView is not available, as the scrollbar does not occupy space in the ui layout  
+	return ESlateVisibility::Collapsed; 
+}
+
 void  UListViewBase::SetAllowOverScroll(bool NewAllowOverscroll)
 {
 	if (MyTableViewBase)
diff --git a/Engine/Source/Runtime/UMG/Public/Components/ListViewBase.h b/Engine/Source/Runtime/UMG/Public/Components/ListViewBase.h
--- a/Engine/Source/Runtime/UMG/Public/Components/ListViewBase.h
+++ b/Engine/Source/Runtime/UMG/Public/Components/ListViewBase.h
@@ -639,6 +639,9 @@ public:
 	UFUNCTION(BlueprintCallable, Category = ListViewBase)
 	UMG_API void SetScrollbarVisibility(ESlateVisibility InVisibility);
 
+	UFUNCTION(BlueprintCallable, Category = ListViewBase)
+	UMG_API ESlateVisibility GetScrollbarVisibility() const;
+
 	UFUNCTION(BlueprintCallable, Category = ListViewBase)
 	UMG_API void SetAllowOverScroll(bool NewAllowOverscroll);
 
