ListView: Expose ReturnFocusToSelection flag

#Synced-CL 48318573

diff --git a/Engine/Source/Runtime/Slate/Public/Widgets/Views/SListView.h b/Engine/Source/Runtime/Slate/Public/Widgets/Views/SListView.h
--- a/Engine/Source/Runtime/Slate/Public/Widgets/Views/SListView.h
+++ b/Engine/Source/Runtime/Slate/Public/Widgets/Views/SListView.h
@@ -2090,6 +2090,11 @@ public:
 		ScrollIntoViewAlignment = NewScrollIntoViewAlignment;
 	}
 
+	void SetReturnFocusToSelection(bool bEnabled)
+	{
+		bReturnFocusToSelection = bEnabled;
+	}
+
 	/**
 	 * Find a widget for this item if it has already been constructed.
 	 *
diff --git a/Engine/Source/Runtime/UMG/Private/Components/ListView.cpp b/Engine/Source/Runtime/UMG/Private/Components/ListView.cpp
--- a/Engine/Source/Runtime/UMG/Private/Components/ListView.cpp
+++ b/Engine/Source/Runtime/UMG/Private/Components/ListView.cpp
@@ -486,6 +486,15 @@ bool UListView::OnIsSelectableOrNavigableInternal(UObject* FirstSelectedItem)
 	return BP_OnIsItemSelectableOrNavigable.IsBound() ? BP_OnIsItemSelectableOrNavigable.Execute(FirstSelectedItem) : ITypedUMGListView::OnIsSelectableOrNavigableInternal(FirstSelectedItem);
 }
 
+void UListView::SetReturnFocusToSelection(bool bEnabled)
+{
+	bReturnFocusToSelection = bEnabled;
+	if (MyListView.IsValid())
+	{
+		MyListView->SetReturnFocusToSelection(bEnabled);
+	}
+}
+
 void UListView::SetVerticalEntrySpacing(float InVerticalEntrySpacing)
 {
 	if (VerticalEntrySpacing != InVerticalEntrySpacing)
diff --git a/Engine/Source/Runtime/UMG/Public/Components/ListView.h b/Engine/Source/Runtime/UMG/Public/Components/ListView.h
--- a/Engine/Source/Runtime/UMG/Public/Components/ListView.h
+++ b/Engine/Source/Runtime/UMG/Public/Components/ListView.h
@@ -294,6 +294,10 @@ public:
 	UFUNCTION(Category = ListView)
 	FMargin GetScrollBarPadding() const { return ScrollBarPadding; }
 
+	/** Set the return Focus to selection flag. */
+	UFUNCTION(BlueprintCallable, Category = ListView)
+	void SetReturnFocusToSelection(bool bEnabled);
+
 	/** Set the vertical spacing between entries. */
 	UMG_API void SetVerticalEntrySpacing(float InVerticalEntrySpacing);
 
@@ -306,6 +310,7 @@ public:
 	/** Gets the Shadow Brush Style when scrolling */
 	UMG_API const FScrollBoxStyle& GetShadowBrushStyle() const;
 
+
 protected:
 	UE_DEPRECATED(5.6, "InitHorizontalEntrySpacing is deprecated, please use SetHorizontalEntrySpacing instead.")
 	UMG_API void InitHorizontalEntrySpacing(float InHorizontalEntrySpacing);
