Add WithIdenticalViaEquality traits to FEventReply to fix initialization issues when creating FEventReply objects in Python

#Synced-CL 48761279

diff --git a/Engine/Source/Runtime/UMG/Public/Components/SlateWrapperTypes.h b/Engine/Source/Runtime/UMG/Public/Components/SlateWrapperTypes.h
--- a/Engine/Source/Runtime/UMG/Public/Components/SlateWrapperTypes.h
+++ b/Engine/Source/Runtime/UMG/Public/Components/SlateWrapperTypes.h
@@ -136,9 +136,20 @@ public:
 	{
 	}
 
+	bool operator==(const FEventReply& Other) const { return NativeReply.IsEventHandled() == Other.NativeReply.IsEventHandled(); }
+
 	FReply NativeReply;
 };
 
+template<>
+struct TStructOpsTypeTraits<FEventReply> : public TStructOpsTypeTraitsBase2<FEventReply>
+{
+	enum
+	{
+		WithIdenticalViaEquality = true,
+	};
+};
+
 /** A struct exposing size param related properties to UMG. */
 USTRUCT(BlueprintType)
 struct FSlateChildSize
