Deprecate UGameViewportSubsystem::Get in favor of the one which does not take a UWorld. This subsystem started out as a game instance subsystem but was later converted to an engine subsystem, but this was never removed.

#rnx
#Synced-CL 48730004

diff --git a/Engine/Plugins/Experimental/UIFramework/Source/Private/UIFPresenter.cpp b/Engine/Plugins/Experimental/UIFramework/Source/Private/UIFPresenter.cpp
--- a/Engine/Plugins/Experimental/UIFramework/Source/Private/UIFPresenter.cpp
+++ b/Engine/Plugins/Experimental/UIFramework/Source/Private/UIFPresenter.cpp
@@ -16,7 +16,7 @@ UUIFrameworkGameViewportPresenter::FWidgetPair::FWidgetPair(UWidget* InWidget, F
 
 void UUIFrameworkGameViewportPresenter::AddToViewport(UWidget* UMGWidget, const FUIFrameworkGameLayerSlot& Slot)
 {
-	if (UGameViewportSubsystem* Subsystem = UGameViewportSubsystem::Get(GetOuterUUIFrameworkPlayerComponent()->GetWorld()))
+	if (UGameViewportSubsystem* Subsystem = UGameViewportSubsystem::Get())
 	{
 		FGameViewportWidgetSlot GameViewportWidgetSlot;
 		GameViewportWidgetSlot.ZOrder = Slot.ZOrder;
@@ -39,7 +39,7 @@ void UUIFrameworkGameViewportPresenter::RemoveFromViewport(FUIFrameworkWidgetId
 	int32 IndexOf = Widgets.IndexOfByPredicate([WidgetId](const FWidgetPair& Other) { return Other.WidgetId == WidgetId; });
 	if (IndexOf != INDEX_NONE)
 	{
-		if (UGameViewportSubsystem* Subsystem = UGameViewportSubsystem::Get(GetOuterUUIFrameworkPlayerComponent()->GetWorld()))
+		if (UGameViewportSubsystem* Subsystem = UGameViewportSubsystem::Get())
 		{
 			Subsystem->RemoveWidget(Widgets[IndexOf].UMGWidget.Get());
 		}
@@ -60,7 +60,7 @@ void UUIFrameworkGameViewportPresenter::BeginDestroy()
 {
 	if (GetOuter() && !IsTemplate())
 	{
-		if (UGameViewportSubsystem* Subsystem = UGameViewportSubsystem::Get(GetOuterUUIFrameworkPlayerComponent()->GetWorld()))
+		if (UGameViewportSubsystem* Subsystem = UGameViewportSubsystem::Get())
 		{
 			for (FWidgetPair& Pair : Widgets)
 			{
diff --git a/Engine/Source/Runtime/UMG/Private/Blueprint/GameViewportSubsystem.cpp b/Engine/Source/Runtime/UMG/Private/Blueprint/GameViewportSubsystem.cpp
--- a/Engine/Source/Runtime/UMG/Private/Blueprint/GameViewportSubsystem.cpp
+++ b/Engine/Source/Runtime/UMG/Private/Blueprint/GameViewportSubsystem.cpp
@@ -74,10 +74,12 @@ void UGameViewportSubsystem::Deinitialize()
 	Super::Deinitialize();
 }
 
+PRAGMA_DISABLE_DEPRECATION_WARNINGS
 UGameViewportSubsystem* UGameViewportSubsystem::Get(UWorld* World)
 {
 	return Get();
 }
+PRAGMA_ENABLE_DEPRECATION_WARNINGS
 
 UGameViewportSubsystem* UGameViewportSubsystem::Get()
 {
diff --git a/Engine/Source/Runtime/UMG/Private/Components/Widget.cpp b/Engine/Source/Runtime/UMG/Private/Components/Widget.cpp
--- a/Engine/Source/Runtime/UMG/Private/Components/Widget.cpp
+++ b/Engine/Source/Runtime/UMG/Private/Components/Widget.cpp
@@ -345,7 +345,7 @@ bool UWidget::IsInViewport() const
 {
 	if (bIsManagedByGameViewportSubsystem)
 	{
-		if (UGameViewportSubsystem* Subsystem = UGameViewportSubsystem::Get(GetWorld()))
+		if (UGameViewportSubsystem* Subsystem = UGameViewportSubsystem::Get())
 		{
 			return Subsystem->IsWidgetAdded(this);
 		}
@@ -912,7 +912,7 @@ void UWidget::RemoveFromParent()
 	{
 		if (bIsManagedByGameViewportSubsystem)
 		{
-			if (UGameViewportSubsystem* Subsystem = UGameViewportSubsystem::Get(GetWorld()))
+			if (UGameViewportSubsystem* Subsystem = UGameViewportSubsystem::Get())
 			{
 				Subsystem->RemoveWidget(this);
 			}
@@ -1671,7 +1671,7 @@ void UWidget::BeginDestroy()
 {
 	if (bIsManagedByGameViewportSubsystem)
 	{
-		if (UGameViewportSubsystem* Subsystem = UGameViewportSubsystem::Get(GetWorld()))
+		if (UGameViewportSubsystem* Subsystem = UGameViewportSubsystem::Get())
 		{
 			Subsystem->RemoveWidget(this);
 		}
diff --git a/Engine/Source/Runtime/UMG/Private/UserWidget.cpp b/Engine/Source/Runtime/UMG/Private/UserWidget.cpp
--- a/Engine/Source/Runtime/UMG/Private/UserWidget.cpp
+++ b/Engine/Source/Runtime/UMG/Private/UserWidget.cpp
@@ -1354,7 +1354,7 @@ UWidget* UUserWidget::GetRootWidget() const
 
 void UUserWidget::AddToViewport(int32 ZOrder)
 {
-	if (UGameViewportSubsystem* Subsystem = UGameViewportSubsystem::Get(GetWorld()))
+	if (UGameViewportSubsystem* Subsystem = UGameViewportSubsystem::Get())
 	{
 		FGameViewportWidgetSlot ViewportSlot;
 		if (bIsManagedByGameViewportSubsystem)
@@ -1368,7 +1368,7 @@ void UUserWidget::AddToViewport(int32 ZOrder)
 
 bool UUserWidget::AddToPlayerScreen(int32 ZOrder)
 {
-	if (UGameViewportSubsystem* Subsystem = UGameViewportSubsystem::Get(GetWorld()))
+	if (UGameViewportSubsystem* Subsystem = UGameViewportSubsystem::Get())
 	{
 		if (ULocalPlayer* LocalPlayer = GetOwningLocalPlayer())
 		{
@@ -1515,7 +1515,7 @@ APlayerCameraManager* UUserWidget::GetOwningPlayerCameraManager() const
 
 void UUserWidget::SetPositionInViewport(FVector2D Position, bool bRemoveDPIScale)
 {
-	if (UGameViewportSubsystem* Subsystem = UGameViewportSubsystem::Get(GetWorld()))
+	if (UGameViewportSubsystem* Subsystem = UGameViewportSubsystem::Get())
 	{
 		if (bIsManagedByGameViewportSubsystem)
 		{
@@ -1533,7 +1533,7 @@ void UUserWidget::SetPositionInViewport(FVector2D Position, bool bRemoveDPIScale
 
 void UUserWidget::SetDesiredSizeInViewport(FVector2D DesiredSize)
 {
-	if (UGameViewportSubsystem* Subsystem = UGameViewportSubsystem::Get(GetWorld()))
+	if (UGameViewportSubsystem* Subsystem = UGameViewportSubsystem::Get())
 	{
 		if (bIsManagedByGameViewportSubsystem)
 		{
@@ -1551,7 +1551,7 @@ void UUserWidget::SetDesiredSizeInViewport(FVector2D DesiredSize)
 
 void UUserWidget::SetAnchorsInViewport(FAnchors Anchors)
 {
-	if (UGameViewportSubsystem* Subsystem = UGameViewportSubsystem::Get(GetWorld()))
+	if (UGameViewportSubsystem* Subsystem = UGameViewportSubsystem::Get())
 	{
 		if (bIsManagedByGameViewportSubsystem)
 		{
@@ -1573,7 +1573,7 @@ void UUserWidget::SetAnchorsInViewport(FAnchors Anchors)
 
 void UUserWidget::SetAlignmentInViewport(FVector2D Alignment)
 {
-	if (UGameViewportSubsystem* Subsystem = UGameViewportSubsystem::Get(GetWorld()))
+	if (UGameViewportSubsystem* Subsystem = UGameViewportSubsystem::Get())
 	{
 		if (bIsManagedByGameViewportSubsystem)
 		{
@@ -1597,7 +1597,7 @@ FMargin UUserWidget::GetFullScreenOffset() const
 {
 	if (bIsManagedByGameViewportSubsystem)
 	{
-		if (UGameViewportSubsystem* Subsystem = UGameViewportSubsystem::Get(GetWorld()))
+		if (UGameViewportSubsystem* Subsystem = UGameViewportSubsystem::Get())
 		{
 			return Subsystem->GetWidgetSlot(this).Offsets;
 		}
@@ -1609,7 +1609,7 @@ FAnchors UUserWidget::GetAnchorsInViewport() const
 {
 	if (bIsManagedByGameViewportSubsystem)
 	{
-		if (UGameViewportSubsystem* Subsystem = UGameViewportSubsystem::Get(GetWorld()))
+		if (UGameViewportSubsystem* Subsystem = UGameViewportSubsystem::Get())
 		{
 			return Subsystem->GetWidgetSlot(this).Anchors;
 		}
@@ -1621,7 +1621,7 @@ FVector2D UUserWidget::GetAlignmentInViewport() const
 {
 	if (bIsManagedByGameViewportSubsystem)
 	{
-		if (UGameViewportSubsystem* Subsystem = UGameViewportSubsystem::Get(GetWorld()))
+		if (UGameViewportSubsystem* Subsystem = UGameViewportSubsystem::Get())
 		{
 			return Subsystem->GetWidgetSlot(this).Alignment;
 		}
diff --git a/Engine/Source/Runtime/UMG/Public/Blueprint/GameViewportSubsystem.h b/Engine/Source/Runtime/UMG/Public/Blueprint/GameViewportSubsystem.h
--- a/Engine/Source/Runtime/UMG/Public/Blueprint/GameViewportSubsystem.h
+++ b/Engine/Source/Runtime/UMG/Public/Blueprint/GameViewportSubsystem.h
@@ -60,6 +60,8 @@ public:
 	//~ End Subsystem
 
 	static UMG_API UGameViewportSubsystem* Get();
+
+	UE_DEPRECATED(5.8, "Use the Get function which takes no arguments instead.")
 	static UMG_API UGameViewportSubsystem* Get(UWorld* World);
 	
 public:
