UMGViewport -Added bool to expose if PreviewScene is of Game or Editor type
- Moved Preview scene initialisation later to enable setting world type.

Mocap Manager - Set PCapViewport to always create an EditorPreview world.

#virtualized
#Synced-CL 49325839

diff --git a/Engine/Source/Runtime/UMG/Private/Components/Viewport.cpp b/Engine/Source/Runtime/UMG/Private/Components/Viewport.cpp
--- a/Engine/Source/Runtime/UMG/Private/Components/Viewport.cpp
+++ b/Engine/Source/Runtime/UMG/Private/Components/Viewport.cpp
@@ -316,14 +316,15 @@ FSceneView* FUMGViewportClient::CalcSceneView(FSceneViewFamily* ViewFamily)
 
 class SAutoRefreshViewport : public SViewport
 {
-	SLATE_BEGIN_ARGS(SAutoRefreshViewport)
+	SLATE_BEGIN_ARGS( SAutoRefreshViewport )
+		: _bIsEditorPreview(false)
 	{
 	}
+	SLATE_ARGUMENT( bool, bIsEditorPreview )
 	SLATE_END_ARGS()
-
-	SAutoRefreshViewport()
-		: PreviewScene(FPreviewScene::ConstructionValues().SetEditor(false))
-	{}
+	
+	SAutoRefreshViewport() = default;
+	
 
 	void Construct(const FArguments& InArgs)
 	{
@@ -332,8 +333,13 @@ class SAutoRefreshViewport : public SViewport
 		ParentArgs.EnableBlending(false);
 		//ParentArgs.RenderDirectlyToWindow(true);
 		SViewport::Construct(ParentArgs);
+		
+		bIsEditorPreview = InArgs._bIsEditorPreview;
+		
+		PreviewScene = MakeUnique<FPreviewScene>(FPreviewScene::ConstructionValues().SetEditor(bIsEditorPreview));
+		
 
-		ViewportClient = MakeShareable(new FUMGViewportClient(&PreviewScene));
+		ViewportClient = MakeShareable(new FUMGViewportClient(PreviewScene.Get()));
 		Viewport = MakeShareable(new FSceneViewport(ViewportClient.Get(), SharedThis(this)));
 
 		// The viewport widget needs an interface so it knows what should render
@@ -351,11 +357,15 @@ class SAutoRefreshViewport : public SViewport
 
 public:
 	/** preview scene */
-	FPreviewScene PreviewScene;
+	
+	TUniquePtr<FPreviewScene> PreviewScene;
+	
+	bool bIsEditorPreview = false;
 
 	TSharedPtr<FUMGViewportClient> ViewportClient;
 	
 	TSharedPtr<FSceneViewport> Viewport;
+
 };
 
 /////////////////////////////////////////////////////
@@ -396,7 +406,8 @@ TSharedRef<SWidget> UViewport::RebuildWidget()
 	}
 	else
 	{
-		ViewportWidget = SNew(SAutoRefreshViewport);
+		ViewportWidget = SNew(SAutoRefreshViewport)
+			.bIsEditorPreview(bIsEditorPreview);
 
 		if ( GetChildrenCount() > 0 )
 		{
@@ -443,7 +454,7 @@ UWorld* UViewport::GetViewportWorld() const
 {
 	if ( ViewportWidget.IsValid() )
 	{
-		return ViewportWidget->PreviewScene.GetWorld();
+		return ViewportWidget->PreviewScene.Get()->GetWorld();
 	}
 
 	return NULL;
@@ -544,12 +555,12 @@ void UViewport::SetEnableAdvancedFeatures(bool InEnableAdvancedFeatures)
 
 void UViewport::SetLightIntensity(float InLightIntensity)
 {
-	ViewportWidget->PreviewScene.SetLightBrightness(InLightIntensity);
+	ViewportWidget->PreviewScene.Get()->SetLightBrightness(InLightIntensity);
 }
 
 void UViewport::SetSkyIntensity(float InLightIntensity)
 {
-	ViewportWidget->PreviewScene.SetSkyBrightness(InLightIntensity);
+	ViewportWidget->PreviewScene.Get()->SetSkyBrightness(InLightIntensity);
 }
 
 PRAGMA_DISABLE_DEPRECATION_WARNINGS
diff --git a/Engine/Source/Runtime/UMG/Public/Components/Viewport.h b/Engine/Source/Runtime/UMG/Public/Components/Viewport.h
--- a/Engine/Source/Runtime/UMG/Public/Components/Viewport.h
+++ b/Engine/Source/Runtime/UMG/Public/Components/Viewport.h
@@ -239,6 +239,13 @@ class UViewport : public UContentWidget
 	UPROPERTY(EditAnywhere, BlueprintReadWrite, Getter, Setter, Category = Appearance)
 	FLinearColor BackgroundColor;
 
+	/**
+	 * True = WorldType in the new viewport will be set to "EditorPreview". This is advised when spawning a UMG viewport for Editor Tooling
+	 * as the default World is WorldType=GamePreview and can cause unneeded subsystems to respond to the creation of the new UWorld.
+	 */
+	UPROPERTY(EditAnywhere, Category = "World")
+	bool bIsEditorPreview = false;
+
 	UFUNCTION(BlueprintCallable, Category="Viewport")
 	UMG_API UWorld* GetViewportWorld() const;
 
