PR #14245: [UMG] Expose AnyTextSelected to UMG for editable text blocks

#Synced-CL 49745357

diff --git a/Engine/Source/Runtime/UMG/Private/Components/EditableTextBox.cpp b/Engine/Source/Runtime/UMG/Private/Components/EditableTextBox.cpp
--- a/Engine/Source/Runtime/UMG/Private/Components/EditableTextBox.cpp
+++ b/Engine/Source/Runtime/UMG/Private/Components/EditableTextBox.cpp
@@ -287,6 +287,16 @@ bool UEditableTextBox::GetSelectAllTextOnCommit() const
 	return SelectAllTextOnCommit;
 }
 
+bool UEditableTextBox::IsAnyTextSelected() const
+{
+	if ( MyEditableTextBlock.IsValid() )
+	{
+		return MyEditableTextBlock->AnyTextSelected();
+	}
+	
+	return false;
+}
+
 void UEditableTextBox::SetForegroundColor(FLinearColor color)
 {
 	WidgetStyle.ForegroundColor = color;
diff --git a/Engine/Source/Runtime/UMG/Public/Components/EditableTextBox.h b/Engine/Source/Runtime/UMG/Public/Components/EditableTextBox.h
--- a/Engine/Source/Runtime/UMG/Public/Components/EditableTextBox.h
+++ b/Engine/Source/Runtime/UMG/Public/Components/EditableTextBox.h
@@ -218,6 +218,9 @@ public:
 	/** Whether to select all text when pressing enter to commit changes */
 	UMG_API bool GetSelectAllTextOnCommit() const;
 
+	/** Whether any text is selected in the text block */
+	UMG_API bool IsAnyTextSelected() const;
+
 	UFUNCTION(BlueprintCallable, Category="Widget",  meta=(DisplayName="SetError (Text Box)"))
 	UMG_API void SetError(FText InError);
 
