EditorUsability : Tooltips
- Editor Utility Tooltips now follow editor theme as well, other UMG widgets are still using legacy style (white) in PIE or packaged build
- Removed margin for custom tooltips

#Synced-CL 49549735

diff --git a/Engine/Source/Runtime/UMG/Private/Components/Widget.cpp b/Engine/Source/Runtime/UMG/Private/Components/Widget.cpp
--- a/Engine/Source/Runtime/UMG/Private/Components/Widget.cpp
+++ b/Engine/Source/Runtime/UMG/Private/Components/Widget.cpp
@@ -535,14 +535,25 @@ void UWidget::SetToolTipText(const FText& InToolTipText)
 	{
 		if (!ToolTipText.IsEmpty())
 		{
+			// Only show legacy tooltip in PIE/Packaged build, editor tooltips should follow unreal theme
+			const bool bIsInGameView = !!GetGameInstance();
+
 			// Match legacy tooltip style with white background
+			const FSlateBrush* BorderImage = bIsInGameView 
+				? FAppStyle::Get().GetBrush("ToolTip.BrightBackground")
+				: FAppStyle::Get().GetBrush("ToolTip.Background");
+		
+			const FSlateColor TextColor = bIsInGameView 
+				? FLinearColor::Black
+				: FAppStyle::Get().GetSlateColor("ToolTip.TextColor");
+
 			const TSharedRef<SToolTip> ToolTip = SNew(SToolTip)
-				.BorderImage(FAppStyle::Get().GetBrush("ToolTip.BrightBackground"))
+				.BorderImage(BorderImage)
 				[
 					SNew(STextBlock)
 					.Text(ToolTipText)
 					.Font(FAppStyle::Get().GetFontStyle("ToolTip.Font"))
-					.ColorAndOpacity(FLinearColor::Black)
+					.ColorAndOpacity(TextColor)
 					.WrapTextAt_Static(&SToolTip::GetToolTipWrapWidth)
 				];
 
@@ -569,9 +580,17 @@ void UWidget::SetToolTip(UWidget* InToolTipWidget)
 	{
 		if ( ToolTipWidget )
 		{
+			// Only show legacy tooltip in PIE/Packaged build, editor tooltips should follow unreal theme
+			const bool bIsInGameView = !!GetGameInstance();
+
 			// Match legacy tooltip style with white background
+			const FSlateBrush* BorderImage = bIsInGameView
+				? FAppStyle::Get().GetBrush("ToolTip.BrightBackground")
+				: FAppStyle::Get().GetBrush("ToolTip.Background");
+
 			TSharedRef<SToolTip> ToolTip = SNew(SToolTip)
-				.BorderImage(FAppStyle::Get().GetBrush("ToolTip.BrightBackground"))
+				.BorderImage(BorderImage)
+				.TextMargin(0.f)
 				[
 					ToolTipWidget->TakeWidget()
 				];
@@ -1545,9 +1564,17 @@ PRAGMA_DISABLE_DEPRECATION_WARNINGS
 	}
 	else if ( ToolTipWidget != nullptr )
 	{
+		// Only show legacy tooltip in PIE/Packaged build, editor tooltips should follow unreal theme
+		const bool bIsInGameView = !!GetGameInstance();
+
+		// Match legacy tooltip style with white background
+		const FSlateBrush* BorderImage = bIsInGameView
+			? FAppStyle::Get().GetBrush("ToolTip.BrightBackground")
+			: FAppStyle::Get().GetBrush("ToolTip.Background");
+
 		TSharedRef<SToolTip> ToolTip = SNew(SToolTip)
-			// Match legacy tooltip style with white background
-			.BorderImage(FAppStyle::Get().GetBrush("ToolTip.BrightBackground"))
+			.BorderImage(BorderImage)
+			.TextMargin(0.f)
 			[
 				ToolTipWidget->TakeWidget()
 			];
@@ -1556,14 +1583,25 @@ PRAGMA_DISABLE_DEPRECATION_WARNINGS
 	}
 	else if ( !ToolTipText.IsEmpty() || ToolTipTextDelegate.IsBound() )
 	{
+		// Only show legacy tooltip in PIE/Packaged build, editor tooltips should follow unreal theme
+		const bool bIsInGameView = !!GetGameInstance();
+
 		// Match legacy tooltip style with white background
+		const FSlateBrush* BorderImage = bIsInGameView 
+			? FAppStyle::Get().GetBrush("ToolTip.BrightBackground")
+			: FAppStyle::Get().GetBrush("ToolTip.Background");
+		
+		const FSlateColor TextColor = bIsInGameView 
+			? FLinearColor::Black
+			: FAppStyle::Get().GetSlateColor("ToolTip.TextColor");
+
 		const TSharedRef<SToolTip> ToolTip = SNew(SToolTip)
-			.BorderImage(FAppStyle::Get().GetBrush("ToolTip.BrightBackground"))
+			.BorderImage(BorderImage)
 			[
 				SNew(STextBlock)
 				.Text(PROPERTY_BINDING(FText, ToolTipText))
 				.Font(FAppStyle::Get().GetFontStyle("ToolTip.Font"))
-				.ColorAndOpacity(FLinearColor::Black)
+				.ColorAndOpacity(TextColor)
 				.WrapTextAt_Static(&SToolTip::GetToolTipWrapWidth)
 			];
 
