UserWidget Focus Forwarding not Responding as Expected - DesiredFocusWidget will sometimes frustratingly reset to "self" after compile
- It is possible for Outer->WidgetTree to be nullptr, so the setter can't resolve the child and we end up with self/null
- Handle this case by creating a widget child and saving it with the target name, then later resolve against a valid widget tree at runtime

#changelist validated

#Synced-CL 50897367

diff --git a/Engine/Source/Runtime/UMG/Private/UserWidget.cpp b/Engine/Source/Runtime/UMG/Private/UserWidget.cpp
--- a/Engine/Source/Runtime/UMG/Private/UserWidget.cpp
+++ b/Engine/Source/Runtime/UMG/Private/UserWidget.cpp
@@ -1146,18 +1146,24 @@ bool UUserWidget::SetDesiredFocusWidget(FName WidgetName)
 
 bool UUserWidget::SetDesiredFocusWidget(UWidget* Widget)
 {
-	if (Widget && WidgetTree)
+	// If the widget is invalid or not contained within our known widget tree, return false
+	if (!Widget)
+	{
+		return false;
+	}
+
+	if (WidgetTree)
 	{
 		TArray<UWidget*> AllWidgets;
 		WidgetTree->GetAllWidgets(AllWidgets);
 
-		if (AllWidgets.Contains(Widget))
+		if (!AllWidgets.Contains(Widget))
 		{
-			DesiredFocusWidget = FWidgetChild(this, Widget->GetFName());
-			return DesiredFocusWidget.GetWidget() != nullptr;
+			return false;
 		}
 	}
-	return false;
+	// Otherwise, we should still remember the name to allow to resolve later against a valid tree.
+	return SetDesiredFocusWidget(Widget->GetFName());
 }
 
 
@@ -1168,7 +1174,12 @@ FName UUserWidget::GetDesiredFocusWidgetName() const
 
 UWidget* UUserWidget::GetDesiredFocusWidget() const
 {
-	return DesiredFocusWidget.GetWidget();
+	if (WidgetTree)
+	{
+		// Ensure that we are always resolving against the widget tree
+		return const_cast<FWidgetChild&>(DesiredFocusWidget).Resolve(WidgetTree);
+	}
+	return nullptr;
 }
 
 UWidget* UUserWidget::GetWidgetHandle(TSharedRef<SWidget> InWidget)
@@ -2420,10 +2431,21 @@ FReply UUserWidget::NativeOnFocusReceived( const FGeometry& InGeometry, const FF
 	FReply Reply = OnFocusReceived( InGeometry, InFocusEvent ).NativeReply;
 
 	// Forward focus if Desired Focus is set
-	if (UWidget * WidgetToFocus = DesiredFocusWidget.Resolve(WidgetTree))
+	if (UWidget* WidgetToFocus = GetDesiredFocusWidget())
 	{
-		return FReply::Handled().SetUserFocus(WidgetToFocus->GetCachedWidget().ToSharedRef(),InFocusEvent.GetCause());
+		if (TSharedPtr<SWidget> SlateWidget = WidgetToFocus->GetCachedWidget())
+		{
+			FWidgetPath Path;
+			// Ensure that the widget is actually hosted/ the Widget tree is not null
+			if (FSlateApplication::Get().GeneratePathToWidgetUnchecked(SlateWidget.ToSharedRef(), Path) && Path.IsValid() && Path.ContainsWidget(GetCachedWidget().Get()))
+			{
+				// Execute full slate focus tunneling using the path
+				FSlateApplication::Get().SetUserFocus(InFocusEvent.GetUser(), Path, InFocusEvent.GetCause());
+				return FReply::Handled();
+			}
+		}
 	}
+
 	return Reply;
 }
 
