New listview support:
- Add an item to a given index
- Add a list of items to the end of the list
- Add a list of items at a given index
- Remove a list of items

#Synced-CL 51008735

diff --git a/Engine/Source/Runtime/UMG/Private/Components/ListView.cpp b/Engine/Source/Runtime/UMG/Private/Components/ListView.cpp
--- a/Engine/Source/Runtime/UMG/Private/Components/ListView.cpp
+++ b/Engine/Source/Runtime/UMG/Private/Components/ListView.cpp
@@ -64,20 +64,98 @@ void UListView::AddItem(UObject* Item)
 
 	ListItems.Add(Item);
 
-	const TArray<UObject*> Added = { Item };
-	const TArray<UObject*> Removed;
-	OnItemsChanged(Added, Removed);
+	OnItemsChanged({ Item }, {});
+
+	RequestRefresh();
+}
+
+void UListView::AddItemAt(UObject* Item, int32 Index)
+{
+	if (Item == nullptr)
+	{
+		FFrame::KismetExecutionMessage(TEXT("Cannot add null item into ListView."), ELogVerbosity::Warning, "NullListViewItem");
+		return;
+	}
+
+	if (ListItems.Contains(Item))
+	{
+		FFrame::KismetExecutionMessage(TEXT("Cannot add duplicate item into ListView."), ELogVerbosity::Warning, "DuplicateListViewItem");
+		return;
+	}
+
+	// Safety clamp the provided index to a valid range
+	Index = FMath::Clamp(Index, 0, ListItems.Num());
+	ListItems.Insert(Item, Index);
+
+	OnItemsChanged({ Item }, {});
+
+	RequestRefresh();
+}
+
+void UListView::AddItems(const TArray<UObject*>& Items)
+{
+	ListItems.Reserve(ListItems.Num() + Items.Num());
+
+	for (UObject* Item : Items)
+	{
+		if (Item != nullptr)
+		{
+			ListItems.Add(Item);
+		}
+	}
+
+	OnItemsChanged(Items, {});
+
+	RequestRefresh();
+}
+
+void UListView::AddItemsAt(const TArray<UObject*>& Items, int32 Index)
+{
+	ListItems.Reserve(ListItems.Num() + Items.Num());
+
+	// Safety clamp the provided index to a valid range
+	Index = FMath::Clamp(Index, 0, ListItems.Num());
+
+	int32 Offset = 0;
+	for (UObject* Item : Items)
+	{
+		if (Item != nullptr)
+		{
+			ListItems.Insert(Item, Index + Offset);
+			++Offset;
+		}
+	}
+
+	OnItemsChanged(Items, {});
 
 	RequestRefresh();
 }
 
 void UListView::RemoveItem(UObject* Item)
 {
+	if (Item == nullptr)
+	{
+		return;
+	}
+
 	ListItems.Remove(Item);
 
-	const TArray<UObject*> Added;
-	const TArray<UObject*> Removed = { Item };
-	OnItemsChanged(Added, Removed);
+	OnItemsChanged({}, { Item });
+
+	RequestRefresh();
+}
+
+void UListView::RemoveItems(const TArray<UObject*>& Items)
+{
+	for (UObject* Item : Items)
+	{
+		if (Item)
+		{
+			ListItems.Remove(Item);
+		}
+	}
+
+	OnItemsChanged({}, Items);
 
 	RequestRefresh();
 }
@@ -248,7 +326,7 @@ void UListView::OnItemsChanged(const TArray<UObject*>& AddedItems, const TArray<
 		{
 			AddedActor->OnEndPlay.AddDynamic(this, &UListView::OnListItemEndPlayed);
 		}
-		else if (AActor* AddedItemOuterActor = AddedItem->GetTypedOuter<AActor>())
+		else if (AActor* AddedItemOuterActor = AddedItem ? AddedItem->GetTypedOuter<AActor>() : nullptr)
 		{
 			// Unique so that we don't spam events for shared actor outers but this also means we can't
 			// unsubscribe when processing RemovedItems
diff --git a/Engine/Source/Runtime/UMG/Public/Components/ListView.h b/Engine/Source/Runtime/UMG/Public/Components/ListView.h
--- a/Engine/Source/Runtime/UMG/Public/Components/ListView.h
+++ b/Engine/Source/Runtime/UMG/Public/Components/ListView.h
@@ -93,14 +93,30 @@ public:
 	UFUNCTION(BlueprintCallable, Category = ListView)
 	const TArray<UObject*>& GetListItems() const { return ListItems; }
 	
-	/** Adds an the item to the list */
+	/** Adds an the item to the end of the list */
 	UFUNCTION(BlueprintCallable, Category = ListView)
 	UMG_API void AddItem(UObject* Item);
 
+	/** Adds an the item to the list at the given index */
+	UFUNCTION(BlueprintCallable, Category = ListView)
+	UMG_API void AddItemAt(UObject* Item, int32 Index);
+
+	/** Adds multiple items to the end of the list */
+	UFUNCTION(BlueprintCallable, Category = ListView)
+	UMG_API void AddItems(const TArray<UObject*>& Items);
+
+	/** Adds multiple items to the list starting at the given index */
+	UFUNCTION(BlueprintCallable, Category = ListView)
+	UMG_API void AddItemsAt(const TArray<UObject*>& Items, int32 Index);
+
 	/** Removes an the item from the list */
 	UFUNCTION(BlueprintCallable, Category = ListView)
 	UMG_API void RemoveItem(UObject* Item);
 
+	/** Removes multiple items from the list */
+	UFUNCTION(BlueprintCallable, Category = ListView)
+	UMG_API void RemoveItems(const TArray<UObject*>& Items);
+
 	/** Returns the item at the given index */
 	UFUNCTION(BlueprintCallable, Category = ListView)
 	UMG_API UObject* GetItemAt(int32 Index) const;
