Fix crash when garbage-collecting a UWrapBoxSlot/UPanelSlot

During incremental garbage collection, UPanelSlot::ReleaseSlateResources could crash with an access violation when the Content widget had already been destroyed in the same GC pass. The raw pointer null-check (if (Content &&)) does not account for objects that are pending kill or have been freed, making it unsafe to dereference during GC teardown. Replacing it with IsValid(Content) ensures the Content widget is still alive before calling through to its ReleaseSlateResources.

#rnx

#Tests Tested in PIE. Added a breakpoint in the affected code and made sure the GC would still proceed. The original crash was unreproducable so that's all I could test.
#Synced-CL 51590040

#Synced-CL 51590040

diff --git a/Engine/Source/Runtime/UMG/Private/Components/PanelSlot.cpp b/Engine/Source/Runtime/UMG/Private/Components/PanelSlot.cpp
--- a/Engine/Source/Runtime/UMG/Private/Components/PanelSlot.cpp
+++ b/Engine/Source/Runtime/UMG/Private/Components/PanelSlot.cpp
@@ -30,7 +30,7 @@ void UPanelSlot::ReleaseSlateResources(bool bReleaseChildren)
 	Super::ReleaseSlateResources(bReleaseChildren);
 
 	// ReleaseSlateResources for Content unless the content is a UUserWidget as they are responsible for releasing their own content.
-	if (Content && !Content->IsA<UUserWidget>())
+	if (IsValid(Content) && !Content->IsA<UUserWidget>())
 	{
 		Content->ReleaseSlateResources(bReleaseChildren);
 	}
