Add render target format override to UWidgetComponent
#rnx
#Synced-CL 51994201

diff --git a/Engine/Source/Runtime/UMG/Private/Components/WidgetComponent.cpp b/Engine/Source/Runtime/UMG/Private/Components/WidgetComponent.cpp
--- a/Engine/Source/Runtime/UMG/Private/Components/WidgetComponent.cpp
+++ b/Engine/Source/Runtime/UMG/Private/Components/WidgetComponent.cpp
@@ -633,6 +633,8 @@ UWidgetComponent::UWidgetComponent( const FObjectInitializer& PCIP )
 	, TintColorAndOpacity( FLinearColor::White )
 	, OpacityFromTexture( 1.0f )
 	, BlendMode( EWidgetBlendMode::Masked )
+	, bOverrideRenderTargetFormat( false )
+	, RenderTargetFormatOverride( RTF_RGBA8 )
 	, bIsTwoSided( false )
 	, TickWhenOffscreen( false )
 	, SharedLayerName(TEXT("WidgetComponentScreenLayer"))
@@ -1953,7 +1955,9 @@ void UWidgetComponent::UpdateRenderTarget(FIntPoint DesiredRenderTargetSize)
 
 	if ( DesiredRenderTargetSize.X != 0 && DesiredRenderTargetSize.Y != 0 )
 	{
-		const EPixelFormat requestedFormat = FSlateApplication::Get().GetRenderer()->GetSlateRecommendedColorFormat();
+		const EPixelFormat requestedFormat = bOverrideRenderTargetFormat
+			? GetPixelFormatFromRenderTargetFormat(RenderTargetFormatOverride)
+			: FSlateApplication::Get().GetRenderer()->GetSlateRecommendedColorFormat();
 
 		if ( RenderTarget == nullptr )
 		{
@@ -1977,7 +1981,7 @@ void UWidgetComponent::UpdateRenderTarget(FIntPoint DesiredRenderTargetSize)
 			if ( bClearColorChanged || RenderTarget->SizeX != DesiredRenderTargetSize.X || RenderTarget->SizeY != DesiredRenderTargetSize.Y )
 			{
 				RenderTarget->ClearColor = ActualBackgroundColor;
-				RenderTarget->InitCustomFormat(DesiredRenderTargetSize.X, DesiredRenderTargetSize.Y, PF_B8G8R8A8, false);
+				RenderTarget->InitCustomFormat(DesiredRenderTargetSize.X, DesiredRenderTargetSize.Y, requestedFormat, false);
 				RenderTarget->UpdateResourceImmediate();
 				bWidgetRenderStateDirty = true;
 			}
diff --git a/Engine/Source/Runtime/UMG/Public/Components/WidgetComponent.h b/Engine/Source/Runtime/UMG/Public/Components/WidgetComponent.h
--- a/Engine/Source/Runtime/UMG/Public/Components/WidgetComponent.h
+++ b/Engine/Source/Runtime/UMG/Public/Components/WidgetComponent.h
@@ -19,6 +19,7 @@ class SWindow;
 class UBodySetup;
 class UMaterialInstanceDynamic;
 class UTextureRenderTarget2D;
+enum ETextureRenderTargetFormat : int;
 
 UENUM(BlueprintType)
 enum class EWidgetSpace : uint8
@@ -559,6 +560,14 @@ protected:
 	UPROPERTY(EditAnywhere, Category=Rendering)
 	EWidgetBlendMode BlendMode;
 
+	/** When enabled, overrides the render target pixel format instead of using the default Slate format (RGBA8). */
+	UPROPERTY(EditAnywhere, Category=Rendering, AdvancedDisplay)
+	bool bOverrideRenderTargetFormat;
+
+	/** The render target format to use when bOverrideRenderTargetFormat is true. */
+	UPROPERTY(EditAnywhere, Category=Rendering, AdvancedDisplay, meta=(EditCondition="bOverrideRenderTargetFormat"))
+	TEnumAsByte<ETextureRenderTargetFormat> RenderTargetFormatOverride;
+
 	/** Is the component visible from behind? */
 	UPROPERTY(EditAnywhere, Category=Rendering)
 	bool bIsTwoSided;
