Minor optimization: Stop going througth all widgets once we found the one we are looking for.

#Synced-CL 52059832

diff --git a/Engine/Source/Runtime/UMG/Private/WidgetTree.cpp b/Engine/Source/Runtime/UMG/Private/WidgetTree.cpp
--- a/Engine/Source/Runtime/UMG/Private/WidgetTree.cpp
+++ b/Engine/Source/Runtime/UMG/Private/WidgetTree.cpp
@@ -33,11 +33,13 @@ UWidget* UWidgetTree::FindWidget(const FName& Name) const
 {
 	UWidget* FoundWidget = nullptr;
 
-	ForEachWidget([&] (UWidget* Widget) {
+	ForEachWidgetUntil([&FoundWidget, &Name] (UWidget* Widget) {
 		if ( Widget->GetFName() == Name )
 		{
 			FoundWidget = Widget;
+			return false;
 		}
+		return true;
 	});
 
 	return FoundWidget;
@@ -47,11 +49,13 @@ UWidget* UWidgetTree::FindWidget(TSharedRef<SWidget> InWidget) const
 {
 	UWidget* FoundWidget = nullptr;
 
-	ForEachWidget([&] (UWidget* Widget) {
+	ForEachWidgetUntil([&FoundWidget, &InWidget] (UWidget* Widget) {
 		if ( Widget->GetCachedWrappedWidget() == InWidget )
 		{
 			FoundWidget = Widget;
+			return false;
 		}
+		return true;
 	});
 
 	return FoundWidget;
