[UMG] Add design-time preview info to UWidgetLayoutLibrary and fix designer preview reconstruction

- Add FWidgetDesignTimeInfo struct and GetDesignTimeInfo() to UWidgetLayoutLibrary
so any widget can query the designer's preview size/orientation in PreConstruct.
- Fix SDesignerView to call RefreshPreview() when the portrait/landscape toggle
or device resolution picker changes, so PreConstruct fires with updated dimensions.
- In editor builds, log an error if GetDesignTimeInfo() is called outside of
a design-time preview. In non-editor builds, returns zeroed values silently.

#rnx
#Synced-CL 52101912

diff --git a/Engine/Source/Editor/UMGEditor/Private/Designer/SDesignerView.cpp b/Engine/Source/Editor/UMGEditor/Private/Designer/SDesignerView.cpp
--- a/Engine/Source/Editor/UMGEditor/Private/Designer/SDesignerView.cpp
+++ b/Engine/Source/Editor/UMGEditor/Private/Designer/SDesignerView.cpp
@@ -3563,6 +3563,14 @@ void SDesignerView::HandleOnCommonResolutionSelected(const FPlayScreenResolution
 		MarkDesignModifed(/*bRequiresRecompile*/ false);
 	}
 
+	// Kick off a preview rebuild so PreConstruct fires with the new preview dimensions.
+	// Note: RefreshPreview() destroys the current preview synchronously; the new preview
+	// is not available until UpdatePreviewWidget() runs on the next Tick.
+	if (TSharedPtr<FWidgetBlueprintEditor> Editor = BlueprintEditor.Pin())
+	{
+		Editor->RefreshPreview();
+	}
+
 	BroadcastDesignerChanged();
 
 	ResolutionTextFade.Play(this->AsShared());
@@ -3854,6 +3862,14 @@ FReply SDesignerView::HandleSwapAspectRatioClicked()
 		MarkDesignModifed(/*bRequiresRecompile*/ false);
 	}
 
+	// Kick off a preview rebuild so PreConstruct fires with the new preview dimensions.
+	// Note: RefreshPreview() destroys the current preview synchronously; the new preview
+	// is not available until UpdatePreviewWidget() runs on the next Tick.
+	if (TSharedPtr<FWidgetBlueprintEditor> Editor = BlueprintEditor.Pin())
+	{
+		Editor->RefreshPreview();
+	}
+
 	BroadcastDesignerChanged();
 
 	ResolutionTextFade.Play(this->AsShared());
diff --git a/Engine/Source/Runtime/UMG/Private/WidgetLayoutLibrary.cpp b/Engine/Source/Runtime/UMG/Private/WidgetLayoutLibrary.cpp
--- a/Engine/Source/Runtime/UMG/Private/WidgetLayoutLibrary.cpp
+++ b/Engine/Source/Runtime/UMG/Private/WidgetLayoutLibrary.cpp
@@ -1,7 +1,9 @@
 // Copyright Epic Games, Inc. All Rights Reserved.
 
 #include "Blueprint/WidgetLayoutLibrary.h"
+#include "UMGPrivate.h"
 #include "EngineGlobals.h"
+#include "Misc/ConfigCacheIni.h"
 #include "GameFramework/PlayerController.h"
 #include "Engine/GameViewportClient.h"
 #include "Engine/Engine.h"
@@ -381,5 +383,32 @@ void UWidgetLayoutLibrary::RemoveAllWidgets(const UObject* WorldContextObject)
 	}
 }
 
+FWidgetDesignTimeInfo UWidgetLayoutLibrary::GetDesignTimeInfo(const UWidget* Widget)
+{
+	FWidgetDesignTimeInfo Info;
+
+#if WITH_EDITOR
+	// Validate that the caller is in a design-time preview context.
+	if (!Widget || !Widget->IsDesignTime())
+	{
+		UE_LOG(LogUMG, Error, TEXT("GetDesignTimeInfo() called outside of a design-time preview. This function should only be called from PreConstruct in the UMG Designer."));
+		return Info;
+	}
+
+	int32 PreviewWidth = 0;
+	int32 PreviewHeight = 0;
+
+	const FString ConfigSection = TEXT("UMGEditor.Designer");
+
+	GConfig->GetInt(*ConfigSection, TEXT("PreviewWidth"), PreviewWidth, GEditorPerProjectIni);
+	GConfig->GetInt(*ConfigSection, TEXT("PreviewHeight"), PreviewHeight, GEditorPerProjectIni);
+
+	Info.PreviewSize = FVector2D(static_cast<double>(PreviewWidth), static_cast<double>(PreviewHeight));
+	Info.bIsPortrait = PreviewHeight > PreviewWidth;
+#endif
+
+	return Info;
+}
+
 #undef LOCTEXT_NAMESPACE
 
diff --git a/Engine/Source/Runtime/UMG/Public/Blueprint/WidgetLayoutLibrary.h b/Engine/Source/Runtime/UMG/Public/Blueprint/WidgetLayoutLibrary.h
--- a/Engine/Source/Runtime/UMG/Public/Blueprint/WidgetLayoutLibrary.h
+++ b/Engine/Source/Runtime/UMG/Public/Blueprint/WidgetLayoutLibrary.h
@@ -8,6 +8,25 @@
 #include "Layout/Geometry.h"
 #include "WidgetLayoutLibrary.generated.h"
 
+/**
+ * Design-time preview information from the UMG Designer.
+ * Contains the current designer preview dimensions and derived orientation.
+ * Only populated at design time — all values are zeroed/false at runtime.
+ */
+USTRUCT(BlueprintType)
+struct UMG_API FWidgetDesignTimeInfo
+{
+	GENERATED_BODY()
+
+	/** The current preview dimensions in the UMG Designer (width x height). */
+	UPROPERTY(BlueprintReadOnly, Category = "Design Time")
+	FVector2D PreviewSize = FVector2D::ZeroVector;
+
+	/** True when the designer preview is in portrait orientation (height > width). */
+	UPROPERTY(BlueprintReadOnly, Category = "Design Time")
+	bool bIsPortrait = false;
+};
+
 class APlayerController;
 class UGameViewportClient;
 class UCanvasPanelSlot;
@@ -212,4 +231,19 @@ public:
 	 */
 	UFUNCTION(BlueprintCallable, BlueprintCosmetic, Category="Viewport", meta=( WorldContext="WorldContextObject" ))
 	static UMG_API void RemoveAllWidgets(const UObject* WorldContextObject);
+
+	/**
+	 * Returns design-time preview information from the UMG Designer.
+	 * Use this in PreConstruct to adapt widget layout to the designer's current preview
+	 * dimensions (e.g., switching between portrait and landscape layouts).
+	 *
+	 * In editor builds, logs an error and returns zeroed values if called outside of a
+	 * design-time preview (e.g., at runtime or from a non-preview context). In non-editor
+	 * builds, returns zeroed values silently. Use viewport/display system
+	 * queries instead for runtime orientation.
+	 *
+	 * @param Widget	The calling widget (auto-filled via DefaultToSelf).
+	 */
+	UFUNCTION(BlueprintPure, Category = "Design Time", meta = (DevelopmentOnly, DefaultToSelf = "Widget"))
+	static UMG_API FWidgetDesignTimeInfo GetDesignTimeInfo(const UWidget* Widget);
 };
