[UMG] Remove DevelopmentOnly flag from GetDesignTimeInfo

The DevelopmentOnly UFUNCTION metadata caused cook errors in non-development builds
when the Blueprint node was connected to a Break node. The Break node's input struct
type was stripped at cook time, causing 'No input structure to break' compile failures.

The function still guards itself at runtime:
- IsDesignTime() check logs an error if called outside PreConstruct
- #if WITH_EDITOR guard ensures preview size data is only read in editor
- Non-editor builds return a default empty FWidgetDesignTimeInfo (harmless)

#rnx

#doc https://docs.google.com/document/d/13zhBZfSw14CtjwS3Vavobez54p_2wywTjtoG8PAb9hE/edit?tab=t.jwcohzo46vz2

#Synced-CL 52421609

diff --git a/Engine/Source/Runtime/UMG/Public/Blueprint/WidgetLayoutLibrary.h b/Engine/Source/Runtime/UMG/Public/Blueprint/WidgetLayoutLibrary.h
--- a/Engine/Source/Runtime/UMG/Public/Blueprint/WidgetLayoutLibrary.h
+++ b/Engine/Source/Runtime/UMG/Public/Blueprint/WidgetLayoutLibrary.h
@@ -243,6 +243,6 @@ public:
 	 *
 	 * @param Widget	The calling widget (auto-filled via DefaultToSelf).
 	 */
-	UFUNCTION(BlueprintPure, Category = "Design Time", meta = (DevelopmentOnly, DefaultToSelf = "Widget"))
+	UFUNCTION(BlueprintPure, Category = "Design Time", meta = (DefaultToSelf = "Widget"))
 	static UMG_API FWidgetDesignTimeInfo GetDesignTimeInfo(const UWidget* Widget);
 };
