PR #14655: Fix IsAnimationPlayingForward default return value to false

#Synced-CL 52523877
#Synced-CL 52523877

diff --git a/Engine/Source/Runtime/UMG/Private/UserWidget.cpp b/Engine/Source/Runtime/UMG/Private/UserWidget.cpp
--- a/Engine/Source/Runtime/UMG/Private/UserWidget.cpp
+++ b/Engine/Source/Runtime/UMG/Private/UserWidget.cpp
@@ -1061,7 +1061,7 @@ bool UUserWidget::IsAnimationPlayingForward(const UWidgetAnimation* InAnimation)
 		}
 	}
 
-	return true;
+	return false;
 }
 
 void UUserWidget::BroadcastAnimationFinishedPlaying(FWidgetAnimationState& State)
